Merging CL #9612832 from //UE4/Dev-Enterprise

Fixed build breakage when CoreTech is not included in the build

#rnx

diff --git a/Engine/Plugins/Enterprise/DatasmithCADImporter/Source/DatasmithCADTranslator/Private/CoreTechParser.cpp b/Engine/Plugins/Enterprise/DatasmithCADImporter/Source/DatasmithCADTranslator/Private/CoreTechParser.cpp
--- a/Engine/Plugins/Enterprise/DatasmithCADImporter/Source/DatasmithCADTranslator/Private/CoreTechParser.cpp
+++ b/Engine/Plugins/Enterprise/DatasmithCADImporter/Source/DatasmithCADTranslator/Private/CoreTechParser.cpp
@@ -1,7 +1,7 @@
 // Copyright 1998-2019 Epic Games, Inc. All Rights Reserved.
 #include "CoreTechParser.h"
 
-#ifndef USE_CORETECH_MT_PARSER
+#if !defined(USE_CORETECH_MT_PARSER) && defined(CAD_LIBRARY)
 
 
 #include "CoreTechHelper.h"
diff --git a/Engine/Plugins/Enterprise/DatasmithCADImporter/Source/DatasmithCADTranslator/Private/CoreTechParser.h b/Engine/Plugins/Enterprise/DatasmithCADImporter/Source/DatasmithCADTranslator/Private/CoreTechParser.h
--- a/Engine/Plugins/Enterprise/DatasmithCADImporter/Source/DatasmithCADTranslator/Private/CoreTechParser.h
+++ b/Engine/Plugins/Enterprise/DatasmithCADImporter/Source/DatasmithCADTranslator/Private/CoreTechParser.h
@@ -1,7 +1,7 @@
 // Copyright 1998-2019 Epic Games, Inc. All Rights Reserved.
 #pragma once
 
-#ifndef USE_CORETECH_MT_PARSER
+#if !defined(USE_CORETECH_MT_PARSER) && defined(CAD_LIBRARY)
 
 #include "CoreMinimal.h"
 
diff --git a/Engine/Plugins/Enterprise/DatasmithCADImporter/Source/DatasmithCADTranslator/Private/CoreTechParserMT.cpp b/Engine/Plugins/Enterprise/DatasmithCADImporter/Source/DatasmithCADTranslator/Private/CoreTechParserMT.cpp
--- a/Engine/Plugins/Enterprise/DatasmithCADImporter/Source/DatasmithCADTranslator/Private/CoreTechParserMT.cpp
+++ b/Engine/Plugins/Enterprise/DatasmithCADImporter/Source/DatasmithCADTranslator/Private/CoreTechParserMT.cpp
@@ -1,7 +1,7 @@
 // Copyright 1998-2019 Epic Games, Inc. All Rights Reserved.
 #include "CoreTechParserMT.h"
 
-#ifdef USE_CORETECH_MT_PARSER
+#if defined(USE_CORETECH_MT_PARSER) && defined(CAD_LIBRARY)
 
 #include "Containers/Array.h"
 #include "Containers/Map.h"
diff --git a/Engine/Plugins/Enterprise/DatasmithCADImporter/Source/DatasmithCADTranslator/Private/CoreTechParserMT.h b/Engine/Plugins/Enterprise/DatasmithCADImporter/Source/DatasmithCADTranslator/Private/CoreTechParserMT.h
--- a/Engine/Plugins/Enterprise/DatasmithCADImporter/Source/DatasmithCADTranslator/Private/CoreTechParserMT.h
+++ b/Engine/Plugins/Enterprise/DatasmithCADImporter/Source/DatasmithCADTranslator/Private/CoreTechParserMT.h
@@ -1,7 +1,7 @@
 // Copyright 1998-2019 Epic Games, Inc. All Rights Reserved.
 #pragma once
 
-#ifdef USE_CORETECH_MT_PARSER
+#if defined(USE_CORETECH_MT_PARSER) && defined(CAD_LIBRARY)
 
 #include "CoreMinimal.h"
 
diff --git a/Engine/Plugins/Enterprise/DatasmithCADImporter/Source/DatasmithCADTranslator/Private/DatasmithCADTranslator.cpp b/Engine/Plugins/Enterprise/DatasmithCADImporter/Source/DatasmithCADTranslator/Private/DatasmithCADTranslator.cpp
--- a/Engine/Plugins/Enterprise/DatasmithCADImporter/Source/DatasmithCADTranslator/Private/DatasmithCADTranslator.cpp
+++ b/Engine/Plugins/Enterprise/DatasmithCADImporter/Source/DatasmithCADTranslator/Private/DatasmithCADTranslator.cpp
@@ -2,6 +2,7 @@
 
 #include "DatasmithCADTranslator.h"
 
+#ifdef CAD_LIBRARY
 #include "CADLibraryOptions.h"
 #include "CoreTechParametricSurfaceExtension.h"
 #include "DatasmithCADTranslatorModule.h"
@@ -145,6 +146,7 @@ void FDatasmithCADTranslator::SetSceneImportOptions(TArray<TStrongObjectPtr<UObj
 		Translator->SetTessellationOptions(GetCommonTessellationOptions());
 	}
 }
+#endif
 
 
 
diff --git a/Engine/Plugins/Enterprise/DatasmithCADImporter/Source/DatasmithCADTranslator/Private/DatasmithCADTranslator.h b/Engine/Plugins/Enterprise/DatasmithCADImporter/Source/DatasmithCADTranslator/Private/DatasmithCADTranslator.h
--- a/Engine/Plugins/Enterprise/DatasmithCADImporter/Source/DatasmithCADTranslator/Private/DatasmithCADTranslator.h
+++ b/Engine/Plugins/Enterprise/DatasmithCADImporter/Source/DatasmithCADTranslator/Private/DatasmithCADTranslator.h
@@ -4,15 +4,22 @@
 
 #include "CoreMinimal.h"
 
+#include "DatasmithCoreTechTranslator.h"
+#ifdef CAD_LIBRARY
 #include "CADLibraryOptions.h"
 #include "DatasmithCADTranslatorImpl.h"
-#include "DatasmithCoreTechTranslator.h"
 #include "UObject/ObjectMacros.h"
+#endif
 
 
 class FDatasmithCADTranslator : public FDatasmithCoreTechTranslator
 {
 public:
+#ifndef CAD_LIBRARY
+	// Begin IDatasmithTranslator overrides
+	virtual FName GetFName() const override { return "DatasmithCADTranslator"; };
+	// End IDatasmithTranslator overrides
+#else
 	FDatasmithCADTranslator();
 
 	// Begin IDatasmithTranslator overrides
@@ -34,5 +41,6 @@ public:
 
 private:
 	TSharedPtr<FDatasmithCADTranslatorImpl> Translator;
+#endif
 };
 
diff --git a/Engine/Plugins/Enterprise/DatasmithCADImporter/Source/DatasmithCADTranslator/Private/DatasmithCADTranslatorImpl.cpp b/Engine/Plugins/Enterprise/DatasmithCADImporter/Source/DatasmithCADTranslator/Private/DatasmithCADTranslatorImpl.cpp
--- a/Engine/Plugins/Enterprise/DatasmithCADImporter/Source/DatasmithCADTranslator/Private/DatasmithCADTranslatorImpl.cpp
+++ b/Engine/Plugins/Enterprise/DatasmithCADImporter/Source/DatasmithCADTranslator/Private/DatasmithCADTranslatorImpl.cpp
@@ -2,6 +2,7 @@
 
 #include "DatasmithCADTranslatorImpl.h"
 
+#ifdef CAD_LIBRARY
 #include "CoreTechTypes.h"
 
 #ifdef USE_CORETECH_MT_PARSER
@@ -84,3 +85,4 @@ void FDatasmithCADTranslatorImpl::UnloadScene()
 {
 	CTParser->UnloadScene();
 }
+#endif
\ No newline at end of file
diff --git a/Engine/Plugins/Enterprise/DatasmithCADImporter/Source/DatasmithCADTranslator/Private/DatasmithCADTranslatorImpl.h b/Engine/Plugins/Enterprise/DatasmithCADImporter/Source/DatasmithCADTranslator/Private/DatasmithCADTranslatorImpl.h
--- a/Engine/Plugins/Enterprise/DatasmithCADImporter/Source/DatasmithCADTranslator/Private/DatasmithCADTranslatorImpl.h
+++ b/Engine/Plugins/Enterprise/DatasmithCADImporter/Source/DatasmithCADTranslator/Private/DatasmithCADTranslatorImpl.h
@@ -4,6 +4,7 @@
 
 #include "CoreMinimal.h"
 
+#ifdef CAD_LIBRARY
 #include "CTSession.h"
 #include "DatasmithImportOptions.h"
 #include "IDatasmithSceneElements.h"
@@ -75,3 +76,4 @@ private:
 	TMap< IDatasmithMeshElement*, CT_OBJECT_ID* > MeshElementToCTBodyIdMap;
 #endif
 };
+#endif
diff --git a/Engine/Plugins/Enterprise/DatasmithCADImporter/Source/DatasmithOpenNurbsTranslator/Private/DatasmithOpenNurbsTranslator.cpp b/Engine/Plugins/Enterprise/DatasmithCADImporter/Source/DatasmithOpenNurbsTranslator/Private/DatasmithOpenNurbsTranslator.cpp
--- a/Engine/Plugins/Enterprise/DatasmithCADImporter/Source/DatasmithOpenNurbsTranslator/Private/DatasmithOpenNurbsTranslator.cpp
+++ b/Engine/Plugins/Enterprise/DatasmithCADImporter/Source/DatasmithOpenNurbsTranslator/Private/DatasmithOpenNurbsTranslator.cpp
@@ -1978,7 +1978,9 @@ bool FOpenNurbsTranslatorImpl::Read(ON_BinaryFile& Archive, TSharedRef<IDatasmit
 
 	// 	ScalingFactor is defined according to input Rhino file unit. CT don't modify CAD input according to the set file unit.
 	ScalingFactor = 100 / Settings.m_ModelUnitsAndTolerances.Scale(ON::LengthUnitSystem::Meters);
+#ifdef CAD_LIBRARY
 	LocalSession->SetScaleFactor(ScalingFactor);
+#endif
 
 	// Step 4: REQUIRED - Read bitmap table (it can be empty)
 	int Count = 0;
diff --git a/Engine/Plugins/Enterprise/DatasmithCADImporter/Source/DatasmithWireTranslator/Private/DatasmithWireTranslator.h b/Engine/Plugins/Enterprise/DatasmithCADImporter/Source/DatasmithWireTranslator/Private/DatasmithWireTranslator.h
--- a/Engine/Plugins/Enterprise/DatasmithCADImporter/Source/DatasmithWireTranslator/Private/DatasmithWireTranslator.h
+++ b/Engine/Plugins/Enterprise/DatasmithCADImporter/Source/DatasmithWireTranslator/Private/DatasmithWireTranslator.h
@@ -4,6 +4,11 @@
 
 #include "CoreMinimal.h"
 
+// Undef USE_OPENMODEL if CAD_LIBRARY is not defined
+#ifndef CAD_LIBRARY
+#undef USE_OPENMODEL
+#endif
+
 #include "CADLibraryOptions.h"
 #include "DatasmithCoreTechTranslator.h"
 #include "UObject/ObjectMacros.h"
