Merging CL #9551319, #9551621, #9579580 from //UE4/Dev-Enterprise

9551319:
Fixed too aggresive logic in IESLoader

9551621:
Fixed application of tessellation option for CAD import

9579580:
Dataprep - Make the Operation/(Filter/Fetcher) box tooltip better

diff --git a/Engine/Plugins/Enterprise/DatasmithCADImporter/Source/DatasmithCADTranslator/Private/CoreTechParser.cpp b/Engine/Plugins/Enterprise/DatasmithCADImporter/Source/DatasmithCADTranslator/Private/CoreTechParser.cpp
--- a/Engine/Plugins/Enterprise/DatasmithCADImporter/Source/DatasmithCADTranslator/Private/CoreTechParser.cpp
+++ b/Engine/Plugins/Enterprise/DatasmithCADImporter/Source/DatasmithCADTranslator/Private/CoreTechParser.cpp
@@ -1041,6 +1041,12 @@ CheckedCTError FCoreTechParser::Read()
 		return CT_IO_ERROR::IO_ERROR;;
 	}
 
+	// Force CoreTech to re-tessellate the model with the translator's tessellation parameters
+	CT_TESS_DATA_TYPE VertexType = CT_TESS_DOUBLE;
+	CT_TESS_DATA_TYPE NormalType = CT_TESS_FLOAT;
+	CT_TESS_DATA_TYPE UVType = CT_TESS_DOUBLE;
+	CTKIO_ChangeTesselationParameters(ImportParams.ChordTolerance, ImportParams.MaxEdgeLength, ImportParams.MaxNormalAngle, CT_FALSE, VertexType, NormalType, UVType);
+
 	TopoFixes();
 
 	TSharedRef<FImportDestination> RootNode = MakeShared<FImportDestination>();
@@ -1095,6 +1101,14 @@ TSharedPtr< IDatasmithMeshElement > FCoreTechParser::FindOrAddMeshElement(TShare
 		MeshElement->SetMaterial(MaterialIDElementName, MaterialId2Hash.Value);
 	}
 
+
+	FMD5 MD5;
+	MD5.Update(reinterpret_cast<const uint8*>(&TessellationOptionsHash), sizeof TessellationOptionsHash);
+
+	FMD5Hash Hash;
+	Hash.Set(MD5);
+	MeshElement->SetFileHash(Hash);
+
 	DatasmithScene->AddMesh(MeshElement);
 
 	BodyUUIDToMeshElementMap.Add(ShellUUID, MeshElement);
diff --git a/Engine/Plugins/Enterprise/DatasmithCADImporter/Source/DatasmithCADTranslator/Private/DatasmithCADTranslatorImpl.cpp b/Engine/Plugins/Enterprise/DatasmithCADImporter/Source/DatasmithCADTranslator/Private/DatasmithCADTranslatorImpl.cpp
--- a/Engine/Plugins/Enterprise/DatasmithCADImporter/Source/DatasmithCADTranslator/Private/DatasmithCADTranslatorImpl.cpp
+++ b/Engine/Plugins/Enterprise/DatasmithCADImporter/Source/DatasmithCADTranslator/Private/DatasmithCADTranslatorImpl.cpp
@@ -31,6 +31,10 @@ void FDatasmithCADTranslatorImpl::SetTessellationOptions(const FDatasmithTessell
 
 bool FDatasmithCADTranslatorImpl::Read()
 {
+	// Push translator's tessellation parameters to CTParser
+	// This call has no effect on the load of the model
+	CTParser->SetTessellationOptions(TessellationOptions);
+
 #ifdef USE_CORETECH_MT_PARSER
 	CTParser->SetOutputPath(OutputPath);
 	CTParser->Read();
@@ -50,10 +54,6 @@ bool FDatasmithCADTranslatorImpl::Read()
 	}
 #endif
 
-	// Force CoreTech to re-tessellate the model with the translator's tessellation parameters
-	// This call has no effect on the load of the model
-	CTParser->SetTessellationOptions(TessellationOptions);
-
 	return true;
 }
 
diff --git a/Engine/Plugins/Experimental/DataPrepEditor/Source/DataPrepEditor/Private/Widgets/Action/SDataprepFilter.cpp b/Engine/Plugins/Experimental/DataPrepEditor/Source/DataPrepEditor/Private/Widgets/Action/SDataprepFilter.cpp
--- a/Engine/Plugins/Experimental/DataPrepEditor/Source/DataPrepEditor/Private/Widgets/Action/SDataprepFilter.cpp
+++ b/Engine/Plugins/Experimental/DataPrepEditor/Source/DataPrepEditor/Private/Widgets/Action/SDataprepFilter.cpp
@@ -24,6 +24,10 @@
 void SDataprepFilter::Construct(const FArguments& InArgs, UDataprepFilter& InFilter, const TSharedRef<FDataprepSchemaActionContext>& InDataprepActionContext)
 {
 	Filter = &InFilter;
+
+	TAttribute<FText> TooltipTextAttribute = MakeAttributeSP( this, &SDataprepFilter::GetTooltipText );
+	SetToolTipText( TooltipTextAttribute );
+
 	SDataprepActionBlock::Construct( SDataprepActionBlock::FArguments(), InDataprepActionContext );
 }
 
@@ -166,6 +170,20 @@ void SDataprepFilter::InverseFilter()
 	}
 }
 
+FText SDataprepFilter::GetTooltipText() const
+{
+	FText TooltipText;
+	if ( Filter )
+	{
+		UDataprepFetcher* Fetcher = Filter->GetFetcher();
+		if ( Fetcher )
+		{
+			TooltipText = Fetcher->GetTooltipText();
+		}
+	}
+	return TooltipText;
+}
+
 void SDataprepFilter::AddReferencedObjects(FReferenceCollector& Collector)
 {
 	Collector.AddReferencedObject( Filter );
diff --git a/Engine/Plugins/Experimental/DataPrepEditor/Source/DataPrepEditor/Private/Widgets/Action/SDataprepFilter.h b/Engine/Plugins/Experimental/DataPrepEditor/Source/DataPrepEditor/Private/Widgets/Action/SDataprepFilter.h
--- a/Engine/Plugins/Experimental/DataPrepEditor/Source/DataPrepEditor/Private/Widgets/Action/SDataprepFilter.h
+++ b/Engine/Plugins/Experimental/DataPrepEditor/Source/DataPrepEditor/Private/Widgets/Action/SDataprepFilter.h
@@ -36,6 +36,8 @@ private:
 
 	void InverseFilter();
 
+	FText GetTooltipText() const;
+
 	//~ FGCObject interface
 	virtual void AddReferencedObjects(FReferenceCollector& Collector) override;
 
diff --git a/Engine/Plugins/Experimental/DataPrepEditor/Source/DataPrepEditor/Private/Widgets/Action/SDataprepOperation.cpp b/Engine/Plugins/Experimental/DataPrepEditor/Source/DataPrepEditor/Private/Widgets/Action/SDataprepOperation.cpp
--- a/Engine/Plugins/Experimental/DataPrepEditor/Source/DataPrepEditor/Private/Widgets/Action/SDataprepOperation.cpp
+++ b/Engine/Plugins/Experimental/DataPrepEditor/Source/DataPrepEditor/Private/Widgets/Action/SDataprepOperation.cpp
@@ -13,6 +13,10 @@
 void SDataprepOperation::Construct(const FArguments& InArgs, UDataprepOperation* InOperation, const TSharedRef<FDataprepSchemaActionContext>& InDataprepActionContext)
 {
 	Operation = InOperation;
+
+	TAttribute<FText> TooltipTextAttribute = MakeAttributeSP( this, &SDataprepOperation::GetTooltipText );
+	SetToolTipText( TooltipTextAttribute );
+
 	SDataprepActionBlock::Construct( SDataprepActionBlock::FArguments(), InDataprepActionContext );
 }
 
@@ -30,3 +34,13 @@ void SDataprepOperation::AddReferencedObjects(FReferenceCollector& Collector)
 {
 	Collector.AddReferencedObject( Operation );
 }
+
+FText SDataprepOperation::GetTooltipText() const
+{
+	FText TooltipText;
+	if ( Operation )
+	{
+		TooltipText = Operation->GetTooltip();
+	}
+	return TooltipText;
+}
diff --git a/Engine/Plugins/Experimental/DataPrepEditor/Source/DataPrepEditor/Private/Widgets/Action/SDataprepOperation.h b/Engine/Plugins/Experimental/DataPrepEditor/Source/DataPrepEditor/Private/Widgets/Action/SDataprepOperation.h
--- a/Engine/Plugins/Experimental/DataPrepEditor/Source/DataPrepEditor/Private/Widgets/Action/SDataprepOperation.h
+++ b/Engine/Plugins/Experimental/DataPrepEditor/Source/DataPrepEditor/Private/Widgets/Action/SDataprepOperation.h
@@ -26,6 +26,8 @@ protected:
 
 private:
 
+	FText GetTooltipText() const;
+
 	//~ FGCObject interface
 	virtual void AddReferencedObjects(FReferenceCollector& Collector) override;
 	 
diff --git a/Engine/Source/Editor/UnrealEd/Private/Factories/IESLoader.cpp b/Engine/Source/Editor/UnrealEd/Private/Factories/IESLoader.cpp
--- a/Engine/Source/Editor/UnrealEd/Private/Factories/IESLoader.cpp
+++ b/Engine/Source/Editor/UnrealEd/Private/Factories/IESLoader.cpp
@@ -32,14 +32,9 @@ static void JumpOverWhiteSpace(const uint8*& BufferPos)
 			BufferPos += 2;
 			continue;
 		}
-		else if(*BufferPos == 10)
-		{
-			// no valid MSDOS return file
-			check(0);
-		}
 		else if(*BufferPos <= ' ')
 		{
-			// tab, space, invisible characters
+			// Skip tab, space and invisible characters
 			++BufferPos;
 			continue;
 		}
@@ -69,7 +64,7 @@ static void GetLineContent(const uint8*& BufferPos, char Line[256], bool bStopOn
 		}
 		else if(*BufferPos == '\n')
 		{
-			BufferPos += 1;
+			++BufferPos;
 			break;
 		}
 		else if(bStopOnWhitespace && (*BufferPos <= ' '))
