Copied cl 9527580, 9614944

IFC: Import unreferenced objects
Wire: Fix wrong metric and some material issues

diff --git a/Engine/Plugins/Enterprise/DatasmithCADImporter/Source/DatasmithWireTranslator/Private/DatasmithWireTranslator.cpp b/Engine/Plugins/Enterprise/DatasmithCADImporter/Source/DatasmithWireTranslator/Private/DatasmithWireTranslator.cpp
--- a/Engine/Plugins/Enterprise/DatasmithCADImporter/Source/DatasmithWireTranslator/Private/DatasmithWireTranslator.cpp
+++ b/Engine/Plugins/Enterprise/DatasmithCADImporter/Source/DatasmithWireTranslator/Private/DatasmithWireTranslator.cpp
@@ -14,6 +14,9 @@
 #include "OpenModelUtils.h"
 #include "Translators\DatasmithTranslator.h"
 
+#include "StaticMeshDescription.h"
+#include "StaticMeshOperations.h"
+
 #ifdef CAD_LIBRARY
 #include "AliasCoretechWrapper.h" // requires CoreTech as public dependency
 #include "CoreTechParametricSurfaceExtension.h"
@@ -102,9 +105,9 @@ public:
 	void SetOutputPath(const FString& Path) { OutputPath = Path; }
 
 	//double GetMetricUnit() const { return LocalSession->GetImportParameters().MetricUnit; }
-	FImportParameters& GetImportParameters() 
+	FImportParameters& GetImportParameters()
 	{
-		return LocalSession->GetImportParameters(); 
+		return LocalSession->GetImportParameters();
 	}
 
 
@@ -230,40 +233,60 @@ bool FWireTranslatorImpl::Read()
 	return true;
 }
 
-void FWireTranslatorImpl::CreateAlCommonMaterial(AlShader *Shader, TSharedRef<IDatasmithUEPbrMaterialElement> MaterialElement)
+void FWireTranslatorImpl::AddAlBlinnParameters(AlShader *Shader, TSharedRef<IDatasmithUEPbrMaterialElement> MaterialElement)
 {
-	float Color[3], Transparency[3], Incandescence[3];
-	float Glow = 0;
-	float TransparencyDepth = 0;
-	float TransparencyShade = 0;
-
-	Color[0] = 0.f;
-	Color[1] = 0.f;
-	Color[2] = 0.f;
-	Transparency[0] = 0.f;
-	Transparency[1] = 0.f;
-	Transparency[2] = 0.f;
-	Incandescence[0] = 0.f;
-	Incandescence[1] = 0.f;
-	Incandescence[2] = 0.f;
-
-	bool TransparencyDefined = false;
-	bool IncandescenceDefined = false;
-
-	AlList *list = Shader->fields();
-	AlShadingFieldItem *item = static_cast<AlShadingFieldItem *>(list->first());
-	while (item)
-	{
-		double Value = 0.0;
-		statusCode errorCode = Shader->parameter(item->field(), Value);
-		if (errorCode != 0)
+	// Default values for a Blinn material
+	double Color[] = { 0.57, 0.58, 0.60 };
+	double TransparencyColor[] = { 0.0, 0.0, 0.0 };
+	double IncandescenceColor[] = { 0.0, 0.0, 0.0 };
+	double SpecularColor[] = { 0.15, 0.15, 0.15 } ;
+	double Diffuse = 1.0;
+	double GlowIntensity = 0.0;
+	double Gloss = 0.8;
+	double Eccentricity = 0.35;
+	double Specularity = 1.0;
+	double Reflectivity = 0.5;
+	double SpecularRolloff = 0.5;
+
+	AlList* List = Shader->fields();
+	for (AlShadingFieldItem* Item = static_cast<AlShadingFieldItem *>(List->first()); Item; Item = Item->nextField())
+	{
+		double Value = 0.0f;
+		statusCode ErrorCode = Shader->parameter(Item->field(), Value);
+		if (ErrorCode != 0)
 		{
-			item = item->nextField();
 			continue;
 		}
 
-		switch (item->field())
+		switch (Item->field())
 		{
+		case AlShadingFields::kFLD_SHADING_BLINN_DIFFUSE:
+			Diffuse = Value;
+			break;
+		case AlShadingFields::kFLD_SHADING_BLINN_GLOSS_:
+			Gloss = Value;
+			break;
+		case AlShadingFields::kFLD_SHADING_BLINN_SPECULAR_R:
+			SpecularColor[0] = Value;
+			break;
+		case AlShadingFields::kFLD_SHADING_BLINN_SPECULAR_G:
+			SpecularColor[1] = Value;
+			break;
+		case AlShadingFields::kFLD_SHADING_BLINN_SPECULAR_B:
+			SpecularColor[2] = Value;
+			break;
+		case AlShadingFields::kFLD_SHADING_BLINN_SPECULARITY_:
+			Specularity = Value;
+			break;
+		case AlShadingFields::kFLD_SHADING_BLINN_SPECULAR_ROLLOFF:
+			SpecularRolloff = Value;
+			break;
+		case AlShadingFields::kFLD_SHADING_BLINN_ECCENTRICITY:
+			Eccentricity = Value;
+			break;
+		case AlShadingFields::kFLD_SHADING_BLINN_REFLECTIVITY:
+			Reflectivity = Value;
+			break;
 		case AlShadingFields::kFLD_SHADING_COMMON_COLOR_R:
 			Color[0] = Value;
 			break;
@@ -273,284 +296,872 @@ void FWireTranslatorImpl::CreateAlCommonMaterial(AlShader *Shader, TSharedRef<ID
 		case  AlShadingFields::kFLD_SHADING_COMMON_COLOR_B:
 			Color[2] = Value;
 			break;
-
+		case AlShadingFields::kFLD_SHADING_COMMON_INCANDESCENCE_R:
+			IncandescenceColor[0] = Value;
+			break;
+		case AlShadingFields::kFLD_SHADING_COMMON_INCANDESCENCE_G:
+			IncandescenceColor[1] = Value;
+			break;
+		case AlShadingFields::kFLD_SHADING_COMMON_INCANDESCENCE_B:
+			IncandescenceColor[2] = Value;
+			break;
 		case  AlShadingFields::kFLD_SHADING_COMMON_TRANSPARENCY_R:
-			Transparency[0] = Value;
-			TransparencyDefined = true;
+			TransparencyColor[0] = Value;
 			break;
 		case  AlShadingFields::kFLD_SHADING_COMMON_TRANSPARENCY_G:
-			Transparency[1] = Value;
-			TransparencyDefined = true;
+			TransparencyColor[1] = Value;
 			break;
 		case  AlShadingFields::kFLD_SHADING_COMMON_TRANSPARENCY_B:
-			Transparency[2] = Value;
-			TransparencyDefined = true;
+			TransparencyColor[2] = Value;
 			break;
-
-		case  AlShadingFields::kFLD_SHADING_COMMON_TRANSPARENCY_DEPTH:
-			TransparencyDepth = Value;
-			TransparencyDefined = true;
+		case AlShadingFields::kFLD_SHADING_COMMON_GLOW_INTENSITY:
+			GlowIntensity = Value;
 			break;
+		default:
+			continue;
+		}
+	}
 
-		case  AlShadingFields::kFLD_SHADING_COMMON_TRANSPARENCY_SHADE:
-			TransparencyShade = Value;
-			TransparencyDefined = true;
-			break;
+	float Opacity = 1.0f - (TransparencyColor[0] + TransparencyColor[1] + TransparencyColor[2]) / 3.0f;
+	bool bIsTransparent = !FMath::IsNearlyEqual(Opacity, 1.0f);
 
-		case AlShadingFields::kFLD_SHADING_COMMON_INCANDESCENCE_R:
-			Incandescence[0] = Value;
-			IncandescenceDefined = true;
-			break;
-		case AlShadingFields::kFLD_SHADING_COMMON_INCANDESCENCE_G:
-			Incandescence[1] = Value;
-			IncandescenceDefined = true;
-			break;
-		case AlShadingFields::kFLD_SHADING_COMMON_INCANDESCENCE_B:
-			Incandescence[2] = Value;
-			IncandescenceDefined = true;
+	// Construct parameter expressions
+	IDatasmithMaterialExpressionScalar* DiffuseExpression = MaterialElement->AddMaterialExpression<IDatasmithMaterialExpressionScalar>();
+	DiffuseExpression->GetScalar() = Diffuse;
+	DiffuseExpression->SetName(TEXT("Diffuse"));
 
-		case AlShadingFields::kFLD_SHADING_COMMON_GLOW_INTENSITY:
-			Glow = Value;
-			break;
+	IDatasmithMaterialExpressionScalar* GlossExpression = MaterialElement->AddMaterialExpression<IDatasmithMaterialExpressionScalar>();
+	GlossExpression->GetScalar() = Gloss;
+	GlossExpression->SetName(TEXT("Gloss"));
 
-		}
-		item = item->nextField();
-	}
+	IDatasmithMaterialExpressionColor* SpecularColorExpression = MaterialElement->AddMaterialExpression<IDatasmithMaterialExpressionColor>();
+	SpecularColorExpression->SetName(TEXT("SpecularColor"));
+	SpecularColorExpression->GetColor() = FLinearColor(pow(SpecularColor[0], 2.2), pow(SpecularColor[1], 2.2), pow(SpecularColor[2], 2.2), 1.0f);
 
-	FLinearColor BaseColor = FLinearColor::FromPow22Color(FColor(Color[0], Color[1], Color[2], 255));
-	IDatasmithMaterialExpressionColor* BaseColorExpression = MaterialElement->AddMaterialExpression<IDatasmithMaterialExpressionColor>();
-	BaseColorExpression->SetName(TEXT("Color"));
-	BaseColorExpression->GetColor() = BaseColor;
-	MaterialElement->GetBaseColor().SetExpression(BaseColorExpression);
+	IDatasmithMaterialExpressionScalar* SpecularityExpression = MaterialElement->AddMaterialExpression<IDatasmithMaterialExpressionScalar>();
+	SpecularityExpression->GetScalar() = Specularity * 0.3;
+	SpecularityExpression->SetName(TEXT("Specularity"));
 
-	if (TransparencyDefined)
-	{
-		FLinearColor TransparencyColor = FLinearColor::FromPow22Color(FColor(Transparency[0], Transparency[1], Transparency[2], 255));
-		IDatasmithMaterialExpressionColor* TransparencyExpression = MaterialElement->AddMaterialExpression<IDatasmithMaterialExpressionColor>();
-		TransparencyExpression->SetName(TEXT("Transparency"));
-		TransparencyExpression->GetColor() = TransparencyColor;
+	IDatasmithMaterialExpressionScalar* SpecularRolloffExpression = MaterialElement->AddMaterialExpression<IDatasmithMaterialExpressionScalar>();
+	SpecularRolloffExpression->GetScalar() = SpecularRolloff;
+	SpecularRolloffExpression->SetName(TEXT("SpecularRolloff"));
 
-		IDatasmithMaterialExpressionGeneric* OneMinus = MaterialElement->AddMaterialExpression< IDatasmithMaterialExpressionGeneric >();
-		OneMinus->SetExpressionName(TEXT("OneMinus"));
-		TransparencyExpression->ConnectExpression(*OneMinus->GetInput(0));
-		MaterialElement->GetOpacity().SetExpression(OneMinus);
-	}
+	IDatasmithMaterialExpressionScalar* EccentricityExpression = MaterialElement->AddMaterialExpression<IDatasmithMaterialExpressionScalar>();
+	EccentricityExpression->GetScalar() = Eccentricity;
+	EccentricityExpression->SetName(TEXT("Eccentricity"));
 
-	IDatasmithMaterialExpressionGeneric* GlowAdd = MaterialElement->AddMaterialExpression< IDatasmithMaterialExpressionGeneric >();
-	GlowAdd->SetExpressionName(TEXT("Add"));
+	IDatasmithMaterialExpressionScalar* ReflectivityExpression = MaterialElement->AddMaterialExpression<IDatasmithMaterialExpressionScalar>();
+	ReflectivityExpression->GetScalar() = Reflectivity;
+	ReflectivityExpression->SetName(TEXT("Reflectivity"));
 
-	if (IncandescenceDefined)
+	IDatasmithMaterialExpressionColor* ColorExpression = MaterialElement->AddMaterialExpression<IDatasmithMaterialExpressionColor>();
+	ColorExpression->SetName(TEXT("Color"));
+	ColorExpression->GetColor() = FLinearColor(pow(Color[0] / 255.0, 2.2), pow(Color[1] / 255.0, 2.2), pow(Color[2] / 255.0, 2.2), 255);
+
+	IDatasmithMaterialExpressionColor* IncandescenceColorExpression = MaterialElement->AddMaterialExpression<IDatasmithMaterialExpressionColor>();
+	IncandescenceColorExpression->SetName(TEXT("IncandescenceColor"));
+	IncandescenceColorExpression->GetColor() = FLinearColor(pow(IncandescenceColor[0] / 255.0, 2.2), pow(IncandescenceColor[1] / 255.0, 2.2), pow(IncandescenceColor[2] / 255.0, 2.2), 255);
+
+	IDatasmithMaterialExpressionColor* TransparencyColorExpression = MaterialElement->AddMaterialExpression<IDatasmithMaterialExpressionColor>();
+	TransparencyColorExpression->SetName(TEXT("TransparencyColor"));
+	TransparencyColorExpression->GetColor() = FLinearColor(pow(TransparencyColor[0] / 255.0, 2.2), pow(TransparencyColor[1] / 255.0, 2.2), pow(TransparencyColor[2] / 255.0, 2.2), 255);
+
+	IDatasmithMaterialExpressionScalar* GlowIntensityExpression = MaterialElement->AddMaterialExpression<IDatasmithMaterialExpressionScalar>();
+	GlowIntensityExpression->GetScalar() = GlowIntensity;
+	GlowIntensityExpression->SetName(TEXT("GlowIntensity"));
+
+	// Create aux expressions
+	IDatasmithMaterialExpressionGeneric* ColorSpecLerp = MaterialElement->AddMaterialExpression<IDatasmithMaterialExpressionGeneric>();
+	ColorSpecLerp->SetExpressionName(TEXT("LinearInterpolate"));
+
+	IDatasmithMaterialExpressionScalar* ColorSpecLerpValue = MaterialElement->AddMaterialExpression<IDatasmithMaterialExpressionScalar>();
+	ColorSpecLerpValue->GetScalar() = 0.96f;
+
+	IDatasmithMaterialExpressionGeneric* ColorMetallicLerp = MaterialElement->AddMaterialExpression<IDatasmithMaterialExpressionGeneric>();
+	ColorMetallicLerp->SetExpressionName(TEXT("LinearInterpolate"));
+
+	IDatasmithMaterialExpressionGeneric* DiffuseLerp = MaterialElement->AddMaterialExpression<IDatasmithMaterialExpressionGeneric>();
+	DiffuseLerp->SetExpressionName(TEXT("LinearInterpolate"));
+
+	IDatasmithMaterialExpressionScalar* DiffuseLerpA = MaterialElement->AddMaterialExpression<IDatasmithMaterialExpressionScalar>();
+	DiffuseLerpA->GetScalar() = 0.04f;
+
+	IDatasmithMaterialExpressionScalar* DiffuseLerpB = MaterialElement->AddMaterialExpression<IDatasmithMaterialExpressionScalar>();
+	DiffuseLerpB->GetScalar() = 1.0f;
+
+	IDatasmithMaterialExpressionGeneric* BaseColorMultiply = MaterialElement->AddMaterialExpression<IDatasmithMaterialExpressionGeneric>();
+	BaseColorMultiply->SetExpressionName(TEXT("Multiply"));
+
+	IDatasmithMaterialExpressionGeneric* BaseColorAdd = MaterialElement->AddMaterialExpression<IDatasmithMaterialExpressionGeneric>();
+	BaseColorAdd->SetExpressionName(TEXT("Add"));
+
+	IDatasmithMaterialExpressionGeneric* BaseColorTransparencyMultiply = MaterialElement->AddMaterialExpression<IDatasmithMaterialExpressionGeneric>();
+	BaseColorTransparencyMultiply->SetExpressionName(TEXT("Multiply"));
+
+	IDatasmithMaterialExpressionGeneric* IncandescenceMultiply = MaterialElement->AddMaterialExpression<IDatasmithMaterialExpressionGeneric>();
+	IncandescenceMultiply->SetExpressionName(TEXT("Multiply"));
+
+	IDatasmithMaterialExpressionGeneric* IncandescenceScaleMultiply = MaterialElement->AddMaterialExpression<IDatasmithMaterialExpressionGeneric>();
+	IncandescenceScaleMultiply->SetExpressionName(TEXT("Multiply"));
+
+	IDatasmithMaterialExpressionScalar* IncandescenceScale = MaterialElement->AddMaterialExpression<IDatasmithMaterialExpressionScalar>();
+	IncandescenceScale->GetScalar() = 100.0f;
+
+	IDatasmithMaterialExpressionGeneric* EccentricityMultiply = MaterialElement->AddMaterialExpression<IDatasmithMaterialExpressionGeneric>();
+	EccentricityMultiply->SetExpressionName(TEXT("Multiply"));
+
+	IDatasmithMaterialExpressionGeneric* EccentricityOneMinus = MaterialElement->AddMaterialExpression<IDatasmithMaterialExpressionGeneric>();
+	EccentricityOneMinus->SetExpressionName(TEXT("OneMinus"));
+
+	IDatasmithMaterialExpressionGeneric* RoughnessOneMinus = MaterialElement->AddMaterialExpression<IDatasmithMaterialExpressionGeneric>();
+	RoughnessOneMinus->SetExpressionName(TEXT("OneMinus"));
+
+	IDatasmithMaterialExpressionScalar* FresnelExponent = MaterialElement->AddMaterialExpression<IDatasmithMaterialExpressionScalar>();
+	FresnelExponent->GetScalar() = 4.0f;
+
+	IDatasmithMaterialExpressionFunctionCall* FresnelFunc = MaterialElement->AddMaterialExpression<IDatasmithMaterialExpressionFunctionCall>();
+	FresnelFunc->SetFunctionPathName(TEXT("/Engine/Functions/Engine_MaterialFunctions02/Fresnel_Function.Fresnel_Function"));
+
+	IDatasmithMaterialExpressionGeneric* FresnelLerp = MaterialElement->AddMaterialExpression<IDatasmithMaterialExpressionGeneric>();
+	FresnelLerp->SetExpressionName(TEXT("LinearInterpolate"));
+
+	IDatasmithMaterialExpressionScalar* FresnelLerpA = MaterialElement->AddMaterialExpression<IDatasmithMaterialExpressionScalar>();
+	FresnelLerpA->GetScalar() = 1.0f;
+
+	IDatasmithMaterialExpressionScalar* SpecularPowerExp = MaterialElement->AddMaterialExpression<IDatasmithMaterialExpressionScalar>();
+	SpecularPowerExp->GetScalar() = 0.5f;
+
+	IDatasmithMaterialExpressionGeneric* Power = MaterialElement->AddMaterialExpression<IDatasmithMaterialExpressionGeneric>();
+	Power->SetExpressionName(TEXT("Power"));
+
+	IDatasmithMaterialExpressionGeneric* FresnelMultiply = MaterialElement->AddMaterialExpression<IDatasmithMaterialExpressionGeneric>();
+	FresnelMultiply->SetExpressionName(TEXT("Multiply"));
+
+	IDatasmithMaterialExpressionGeneric* TransparencyOneMinus = MaterialElement->AddMaterialExpression<IDatasmithMaterialExpressionGeneric>();
+	TransparencyOneMinus->SetExpressionName(TEXT("OneMinus"));
+
+	IDatasmithMaterialExpressionFunctionCall* BreakFloat3 = nullptr;
+	IDatasmithMaterialExpressionGeneric* AddRG = nullptr;
+	IDatasmithMaterialExpressionGeneric* AddRGB = nullptr;
+	IDatasmithMaterialExpressionGeneric* Divide = nullptr;
+	IDatasmithMaterialExpressionScalar* DivideConstant = nullptr;
+	if (bIsTransparent)
 	{
-		FLinearColor IncandescenceColor = FLinearColor::FromPow22Color(FColor(Incandescence[0], Incandescence[1], Incandescence[2], 255));
-		IDatasmithMaterialExpressionColor* IncandescenceExpression = MaterialElement->AddMaterialExpression<IDatasmithMaterialExpressionColor>();
-		IncandescenceExpression->SetName(TEXT("Incandescence"));
-		IncandescenceExpression->GetColor() = IncandescenceColor;
-		IncandescenceExpression->ConnectExpression(*GlowAdd->GetInput(1));
+		BreakFloat3 = MaterialElement->AddMaterialExpression<IDatasmithMaterialExpressionFunctionCall>();
+		BreakFloat3->SetFunctionPathName(TEXT("/Engine/Functions/Engine_MaterialFunctions02/Utility/BreakFloat3Components.BreakFloat3Components"));
+
+		AddRG = MaterialElement->AddMaterialExpression<IDatasmithMaterialExpressionGeneric>();
+		AddRG->SetExpressionName(TEXT("Add"));
+
+		AddRGB = MaterialElement->AddMaterialExpression<IDatasmithMaterialExpressionGeneric>();
+		AddRGB->SetExpressionName(TEXT("Add"));
+
+		Divide = MaterialElement->AddMaterialExpression<IDatasmithMaterialExpressionGeneric>();
+		Divide->SetExpressionName(TEXT("Divide"));
+
+		DivideConstant = MaterialElement->AddMaterialExpression<IDatasmithMaterialExpressionScalar>();
+		DivideConstant->GetScalar() = 3.0f;
 	}
 
-	IDatasmithMaterialExpressionGeneric* GlowMultiply = MaterialElement->AddMaterialExpression< IDatasmithMaterialExpressionGeneric >();
-	GlowMultiply->SetExpressionName(TEXT("Multiply"));
+	// Connect expressions
+	SpecularColorExpression->ConnectExpression(*ColorSpecLerp->GetInput(0));
+	ColorExpression->ConnectExpression(*ColorSpecLerp->GetInput(1));
+	ColorSpecLerpValue->ConnectExpression(*ColorSpecLerp->GetInput(2));
 
-	BaseColorExpression->ConnectExpression(*GlowMultiply->GetInput(0));
+	ColorExpression->ConnectExpression(*ColorMetallicLerp->GetInput(0));
+	ColorSpecLerp->ConnectExpression(*ColorMetallicLerp->GetInput(1));
+	GlossExpression->ConnectExpression(*ColorMetallicLerp->GetInput(2));
 
-	IDatasmithMaterialExpressionScalar* GlowExpression = MaterialElement->AddMaterialExpression<IDatasmithMaterialExpressionScalar>();
-	GlowExpression->GetScalar() = Glow;
-	GlowExpression->SetName(TEXT("Glow Intensity"));
+	DiffuseLerpA->ConnectExpression(*DiffuseLerp->GetInput(0));
+	DiffuseLerpB->ConnectExpression(*DiffuseLerp->GetInput(1));
+	DiffuseExpression->ConnectExpression(*DiffuseLerp->GetInput(2));
 
-	GlowExpression->ConnectExpression(*GlowMultiply->GetInput(1));
+	ColorMetallicLerp->ConnectExpression(*BaseColorMultiply->GetInput(0));
+	DiffuseLerp->ConnectExpression(*BaseColorMultiply->GetInput(1));
 
-	GlowMultiply->ConnectExpression(*GlowAdd->GetInput(0));
+	BaseColorMultiply->ConnectExpression(*BaseColorAdd->GetInput(0));
+	IncandescenceColorExpression->ConnectExpression(*BaseColorAdd->GetInput(1));
 
-	MaterialElement->GetEmissiveColor().SetExpression(GlowAdd);
-}
+	BaseColorAdd->ConnectExpression(*BaseColorTransparencyMultiply->GetInput(0));
+	TransparencyOneMinus->ConnectExpression(*BaseColorTransparencyMultiply->GetInput(1));
 
-void FWireTranslatorImpl::AddAlBlinnParameters(AlShader *Shader, TSharedRef<IDatasmithUEPbrMaterialElement> MaterialElement)
-{
-	float Specular[3];
-	Specular[0] = Specular[1] = Specular[2] = 0;
-	bool SpecularDefined = false;
-	float Reflectivity = 0;
+	GlowIntensityExpression->ConnectExpression(*IncandescenceScaleMultiply->GetInput(0));
+	IncandescenceScale->ConnectExpression(*IncandescenceScaleMultiply->GetInput(1));
+
+	BaseColorTransparencyMultiply->ConnectExpression(*IncandescenceMultiply->GetInput(0));
+	IncandescenceScaleMultiply->ConnectExpression(*IncandescenceMultiply->GetInput(1));
+
+	EccentricityExpression->ConnectExpression(*EccentricityOneMinus->GetInput(0));
+
+	EccentricityOneMinus->ConnectExpression(*EccentricityMultiply->GetInput(0));
+	SpecularityExpression->ConnectExpression(*EccentricityMultiply->GetInput(1));
+
+	EccentricityMultiply->ConnectExpression(*RoughnessOneMinus->GetInput(0));
+
+	FresnelExponent->ConnectExpression(*FresnelFunc->GetInput(3));
 
+	SpecularRolloffExpression->ConnectExpression(*Power->GetInput(0));
+	SpecularPowerExp->ConnectExpression(*Power->GetInput(1));
 
-	AlList *list = Shader->fields();
-	AlShadingFieldItem *item = static_cast<AlShadingFieldItem *>(list->first());
-	while (item)
+	FresnelLerpA->ConnectExpression(*FresnelLerp->GetInput(0));
+	FresnelFunc->ConnectExpression(*FresnelLerp->GetInput(1));
+	Power->ConnectExpression(*FresnelLerp->GetInput(2));
+
+	FresnelLerp->ConnectExpression(*FresnelMultiply->GetInput(0));
+	ReflectivityExpression->ConnectExpression(*FresnelMultiply->GetInput(1));
+
+	TransparencyColorExpression->ConnectExpression(*TransparencyOneMinus->GetInput(0));
+
+	if (bIsTransparent)
+	{
+		TransparencyOneMinus->ConnectExpression(*BreakFloat3->GetInput(0));
+
+		BreakFloat3->ConnectExpression(*AddRG->GetInput(0), 0);
+		BreakFloat3->ConnectExpression(*AddRG->GetInput(1), 1);
+
+		AddRG->ConnectExpression(*AddRGB->GetInput(0));
+		BreakFloat3->ConnectExpression(*AddRGB->GetInput(1), 2);
+
+		AddRGB->ConnectExpression(*Divide->GetInput(0));
+		DivideConstant->ConnectExpression(*Divide->GetInput(1));
+	}
+
+	// Connect material outputs
+	MaterialElement->GetBaseColor().SetExpression(BaseColorTransparencyMultiply);
+	MaterialElement->GetMetallic().SetExpression(GlossExpression);
+	MaterialElement->GetSpecular().SetExpression(FresnelMultiply);
+	MaterialElement->GetRoughness().SetExpression(RoughnessOneMinus);
+	MaterialElement->GetEmissiveColor().SetExpression(IncandescenceMultiply);
+	if (bIsTransparent)
+	{
+		MaterialElement->GetOpacity().SetExpression(Divide);
+	}
+
+	MaterialElement->SetParentLabel(TEXT("BLINN"));
+}
+
+void FWireTranslatorImpl::AddAlLambertParameters(AlShader *Shader, TSharedRef<IDatasmithUEPbrMaterialElement> MaterialElement)
+{
+	// Default values for a Lambert material
+	double Color[] = { 0.57, 0.58, 0.60 };
+	double TransparencyColor[] = { 0.0, 0.0, 0.0 };
+	double IncandescenceColor[] = { 0.0, 0.0, 0.0 };
+	double Diffuse = 1.0;
+	double GlowIntensity = 0.0;
+
+	AlList* List = Shader->fields();
+	for (AlShadingFieldItem* Item = static_cast<AlShadingFieldItem *>(List->first()); Item; Item = Item->nextField())
 	{
-		double Value = 0.0;
-		statusCode errorCode = Shader->parameter(item->field(), Value);
-		if (errorCode != 0)
+		double Value = 0.0f;
+		statusCode ErrorCode = Shader->parameter(Item->field(), Value);
+		if (ErrorCode != 0)
 		{
-			item = item->nextField();
 			continue;
 		}
 
-		switch (item->field())
+		switch (Item->field())
 		{
-		case  AlShadingFields::kFLD_SHADING_BLINN_DIFFUSE:
-		case  AlShadingFields::kFLD_SHADING_BLINN_ECCENTRICITY:
-		case  AlShadingFields::kFLD_SHADING_BLINN_SPECULAR_ROLLOFF:
+		case AlShadingFields::kFLD_SHADING_LAMBERT_DIFFUSE:
+			Diffuse = Value;
 			break;
-
-		case  AlShadingFields::kFLD_SHADING_BLINN_REFLECTIVITY:
-			Reflectivity = Value;
+		case AlShadingFields::kFLD_SHADING_COMMON_COLOR_R:
+			Color[0] = Value;
 			break;
-
-		case  AlShadingFields::kFLD_SHADING_BLINN_SPECULAR_R:
-			Specular[0] = Value;
-			SpecularDefined = true;
+		case AlShadingFields::kFLD_SHADING_COMMON_COLOR_G:
+			Color[1] = Value;
 			break;
-		case  AlShadingFields::kFLD_SHADING_BLINN_SPECULAR_G:
-			Specular[1] = Value;
-			SpecularDefined = true;
+		case  AlShadingFields::kFLD_SHADING_COMMON_COLOR_B:
+			Color[2] = Value;
 			break;
-		case  AlShadingFields::kFLD_SHADING_BLINN_SPECULAR_B:
-			Specular[2] = Value;
-			SpecularDefined = true;
+		case AlShadingFields::kFLD_SHADING_COMMON_INCANDESCENCE_R:
+			IncandescenceColor[0] = Value;
 			break;
-
+		case AlShadingFields::kFLD_SHADING_COMMON_INCANDESCENCE_G:
+			IncandescenceColor[1] = Value;
+			break;
+		case AlShadingFields::kFLD_SHADING_COMMON_INCANDESCENCE_B:
+			IncandescenceColor[2] = Value;
+			break;
+		case  AlShadingFields::kFLD_SHADING_COMMON_TRANSPARENCY_R:
+			TransparencyColor[0] = Value;
+			break;
+		case  AlShadingFields::kFLD_SHADING_COMMON_TRANSPARENCY_G:
+			TransparencyColor[1] = Value;
+			break;
+		case  AlShadingFields::kFLD_SHADING_COMMON_TRANSPARENCY_B:
+			TransparencyColor[2] = Value;
+			break;
+		case AlShadingFields::kFLD_SHADING_COMMON_GLOW_INTENSITY:
+			GlowIntensity = Value;
+			break;
+		default:
+			continue;
 		}
-		item = item->nextField();
 	}
 
-	if (SpecularDefined)
+	float Opacity = 1.0f - (TransparencyColor[0] + TransparencyColor[1] + TransparencyColor[2]) / 3.0f;
+	bool bIsTransparent = !FMath::IsNearlyEqual(Opacity, 1.0f);
+
+	// Construct parameter expressions
+	IDatasmithMaterialExpressionScalar* DiffuseExpression = MaterialElement->AddMaterialExpression<IDatasmithMaterialExpressionScalar>();
+	DiffuseExpression->GetScalar() = Diffuse;
+	DiffuseExpression->SetName(TEXT("Diffuse"));
+
+	IDatasmithMaterialExpressionColor* ColorExpression = MaterialElement->AddMaterialExpression<IDatasmithMaterialExpressionColor>();
+	ColorExpression->SetName(TEXT("Color"));
+	ColorExpression->GetColor() = FLinearColor(pow(Color[0] / 255.0, 2.2), pow(Color[1] / 255.0, 2.2), pow(Color[2] / 255.0, 2.2), 255);
+
+	IDatasmithMaterialExpressionColor* IncandescenceColorExpression = MaterialElement->AddMaterialExpression<IDatasmithMaterialExpressionColor>();
+	IncandescenceColorExpression->SetName(TEXT("IncandescenceColor"));
+	IncandescenceColorExpression->GetColor() = FLinearColor(pow(IncandescenceColor[0] / 255.0, 2.2), pow(IncandescenceColor[1] / 255.0, 2.2), pow(IncandescenceColor[2] / 255.0, 2.2), 255);
+
+	IDatasmithMaterialExpressionColor* TransparencyColorExpression = MaterialElement->AddMaterialExpression<IDatasmithMaterialExpressionColor>();
+	TransparencyColorExpression->SetName(TEXT("TransparencyColor"));
+	TransparencyColorExpression->GetColor() = FLinearColor(pow(TransparencyColor[0] / 255.0, 2.2), pow(TransparencyColor[1] / 255.0, 2.2), pow(TransparencyColor[2] / 255.0, 2.2), 255);
+
+	IDatasmithMaterialExpressionScalar* GlowIntensityExpression = MaterialElement->AddMaterialExpression<IDatasmithMaterialExpressionScalar>();
+	GlowIntensityExpression->GetScalar() = GlowIntensity;
+	GlowIntensityExpression->SetName(TEXT("GlowIntensity"));
+
+	// Create aux expressions
+	IDatasmithMaterialExpressionGeneric* DiffuseLerp = MaterialElement->AddMaterialExpression<IDatasmithMaterialExpressionGeneric>();
+	DiffuseLerp->SetExpressionName(TEXT("LinearInterpolate"));
+
+	IDatasmithMaterialExpressionScalar* DiffuseLerpA = MaterialElement->AddMaterialExpression<IDatasmithMaterialExpressionScalar>();
+	DiffuseLerpA->GetScalar() = 0.04f;
+
+	IDatasmithMaterialExpressionScalar* DiffuseLerpB = MaterialElement->AddMaterialExpression<IDatasmithMaterialExpressionScalar>();
+	DiffuseLerpB->GetScalar() = 1.0f;
+
+	IDatasmithMaterialExpressionGeneric* BaseColorMultiply = MaterialElement->AddMaterialExpression<IDatasmithMaterialExpressionGeneric>();
+	BaseColorMultiply->SetExpressionName(TEXT("Multiply"));
+
+	IDatasmithMaterialExpressionGeneric* BaseColorAdd = MaterialElement->AddMaterialExpression<IDatasmithMaterialExpressionGeneric>();
+	BaseColorAdd->SetExpressionName(TEXT("Add"));
+
+	IDatasmithMaterialExpressionGeneric* BaseColorTransparencyMultiply = MaterialElement->AddMaterialExpression<IDatasmithMaterialExpressionGeneric>();
+	BaseColorTransparencyMultiply->SetExpressionName(TEXT("Multiply"));
+
+	IDatasmithMaterialExpressionGeneric* IncandescenceMultiply = MaterialElement->AddMaterialExpression<IDatasmithMaterialExpressionGeneric>();
+	IncandescenceMultiply->SetExpressionName(TEXT("Multiply"));
+
+	IDatasmithMaterialExpressionGeneric* IncandescenceScaleMultiply = MaterialElement->AddMaterialExpression<IDatasmithMaterialExpressionGeneric>();
+	IncandescenceScaleMultiply->SetExpressionName(TEXT("Multiply"));
+
+	IDatasmithMaterialExpressionScalar* IncandescenceScale = MaterialElement->AddMaterialExpression<IDatasmithMaterialExpressionScalar>();
+	IncandescenceScale->GetScalar() = 100.0f;
+
+	IDatasmithMaterialExpressionGeneric* TransparencyOneMinus = MaterialElement->AddMaterialExpression<IDatasmithMaterialExpressionGeneric>();
+	TransparencyOneMinus->SetExpressionName(TEXT("OneMinus"));
+
+	IDatasmithMaterialExpressionFunctionCall* BreakFloat3 = nullptr;
+	IDatasmithMaterialExpressionGeneric* AddRG = nullptr;
+	IDatasmithMaterialExpressionGeneric* AddRGB = nullptr;
+	IDatasmithMaterialExpressionGeneric* Divide = nullptr;
+	IDatasmithMaterialExpressionScalar* DivideConstant = nullptr;
+	if (bIsTransparent)
 	{
-		FLinearColor SpecularColor = FLinearColor::FromPow22Color(FColor(Specular[0], Specular[1], Specular[2], 255.));
-		IDatasmithMaterialExpressionColor* SpecularExpression = MaterialElement->AddMaterialExpression<IDatasmithMaterialExpressionColor>();
-		SpecularExpression->SetName(TEXT("Specular"));
-		SpecularExpression->GetColor() = SpecularColor;
-		MaterialElement->GetSpecular().SetExpression(SpecularExpression);
+		BreakFloat3 = MaterialElement->AddMaterialExpression<IDatasmithMaterialExpressionFunctionCall>();
+		BreakFloat3->SetFunctionPathName(TEXT("/Engine/Functions/Engine_MaterialFunctions02/Utility/BreakFloat3Components.BreakFloat3Components"));
+
+		AddRG = MaterialElement->AddMaterialExpression<IDatasmithMaterialExpressionGeneric>();
+		AddRG->SetExpressionName(TEXT("Add"));
+
+		AddRGB = MaterialElement->AddMaterialExpression<IDatasmithMaterialExpressionGeneric>();
+		AddRGB->SetExpressionName(TEXT("Add"));
+
+		Divide = MaterialElement->AddMaterialExpression<IDatasmithMaterialExpressionGeneric>();
+		Divide->SetExpressionName(TEXT("Divide"));
+
+		DivideConstant = MaterialElement->AddMaterialExpression<IDatasmithMaterialExpressionScalar>();
+		DivideConstant->GetScalar() = 3.0f;
 	}
 
-	IDatasmithMaterialExpressionScalar* ReflectivityScalar = MaterialElement->AddMaterialExpression<IDatasmithMaterialExpressionScalar>();
-	ReflectivityScalar->GetScalar() = Reflectivity;
-	ReflectivityScalar->SetName(TEXT("Reflectivity"));
-	MaterialElement->GetMetallic().SetExpression(ReflectivityScalar);
+	// Connect expressions
+	DiffuseLerpA->ConnectExpression(*DiffuseLerp->GetInput(0));
+	DiffuseLerpB->ConnectExpression(*DiffuseLerp->GetInput(1));
+	DiffuseExpression->ConnectExpression(*DiffuseLerp->GetInput(2));
 
-	IDatasmithMaterialExpressionGeneric* OneMinus = MaterialElement->AddMaterialExpression< IDatasmithMaterialExpressionGeneric >();
-	OneMinus->SetExpressionName(TEXT("OneMinus"));
-	ReflectivityScalar->ConnectExpression(*OneMinus->GetInput(0));
-	MaterialElement->GetRoughness().SetExpression(OneMinus);
+	ColorExpression->ConnectExpression(*BaseColorMultiply->GetInput(0));
+	DiffuseLerp->ConnectExpression(*BaseColorMultiply->GetInput(1));
 
+	BaseColorMultiply->ConnectExpression(*BaseColorAdd->GetInput(0));
+	IncandescenceColorExpression->ConnectExpression(*BaseColorAdd->GetInput(1));
 
-	MaterialElement->SetParentLabel(TEXT("BLINN"));
+	BaseColorAdd->ConnectExpression(*BaseColorTransparencyMultiply->GetInput(0));
+	TransparencyOneMinus->ConnectExpression(*BaseColorTransparencyMultiply->GetInput(1));
+
+	GlowIntensityExpression->ConnectExpression(*IncandescenceScaleMultiply->GetInput(0));
+	IncandescenceScale->ConnectExpression(*IncandescenceScaleMultiply->GetInput(1));
+
+	BaseColorTransparencyMultiply->ConnectExpression(*IncandescenceMultiply->GetInput(0));
+	IncandescenceScaleMultiply->ConnectExpression(*IncandescenceMultiply->GetInput(1));
+
+	TransparencyColorExpression->ConnectExpression(*TransparencyOneMinus->GetInput(0));
+
+	if (bIsTransparent)
+	{
+		TransparencyOneMinus->ConnectExpression(*BreakFloat3->GetInput(0));
+
+		BreakFloat3->ConnectExpression(*AddRG->GetInput(0), 0);
+		BreakFloat3->ConnectExpression(*AddRG->GetInput(1), 1);
+
+		AddRG->ConnectExpression(*AddRGB->GetInput(0));
+		BreakFloat3->ConnectExpression(*AddRGB->GetInput(1), 2);
+
+		AddRGB->ConnectExpression(*Divide->GetInput(0));
+		DivideConstant->ConnectExpression(*Divide->GetInput(1));
+	}
+
+	// Connect material outputs
+	MaterialElement->GetBaseColor().SetExpression(BaseColorTransparencyMultiply);
+	MaterialElement->GetEmissiveColor().SetExpression(IncandescenceMultiply);
+	if (bIsTransparent)
+	{
+		MaterialElement->GetOpacity().SetExpression(Divide);
+	}
+
+	MaterialElement->SetParentLabel(TEXT("LAMBERT"));
 }
 
-void FWireTranslatorImpl::AddAlLambertParameters(AlShader *Shader, TSharedRef<IDatasmithUEPbrMaterialElement> MaterialElement)
+void FWireTranslatorImpl::AddAlLightSourceParameters(AlShader *Shader, TSharedRef<IDatasmithUEPbrMaterialElement> MaterialElement)
 {
-	float Diffuse = 0.8;
-
-	AlList *list = Shader->fields();
-	AlShadingFieldItem *item = static_cast<AlShadingFieldItem *>(list->first());
-	while (item)
+	// Default values for a LightSource material
+	double Color[] = { 0.57, 0.58, 0.60 };
+	double TransparencyColor[] = { 0.0, 0.0, 0.0 };
+	double IncandescenceColor[] = { 0.0, 0.0, 0.0 };
+	double GlowIntensity = 0.0;
+
+	AlList* List = Shader->fields();
+	for (AlShadingFieldItem* Item = static_cast<AlShadingFieldItem *>(List->first()); Item; Item = Item->nextField())
 	{
-		double Value = 0.0;
-		statusCode errorCode = Shader->parameter(item->field(), Value);
-		if (errorCode != 0)
+		double Value = 0.0f;
+		statusCode ErrorCode = Shader->parameter(Item->field(), Value);
+		if (ErrorCode != 0)
 		{
-			item = item->nextField();
 			continue;
 		}
 
-		switch (item->field())
+		switch (Item->field())
 		{
-		case  AlShadingFields::kFLD_SHADING_LAMBERT_DIFFUSE:
-			Diffuse = Value;
+		case AlShadingFields::kFLD_SHADING_COMMON_COLOR_R:
+			Color[0] = Value;
+			break;
+		case AlShadingFields::kFLD_SHADING_COMMON_COLOR_G:
+			Color[1] = Value;
 			break;
+		case  AlShadingFields::kFLD_SHADING_COMMON_COLOR_B:
+			Color[2] = Value;
+			break;
+		case AlShadingFields::kFLD_SHADING_COMMON_INCANDESCENCE_R:
+			IncandescenceColor[0] = Value;
+			break;
+		case AlShadingFields::kFLD_SHADING_COMMON_INCANDESCENCE_G:
+			IncandescenceColor[1] = Value;
+			break;
+		case AlShadingFields::kFLD_SHADING_COMMON_INCANDESCENCE_B:
+			IncandescenceColor[2] = Value;
+			break;
+		case  AlShadingFields::kFLD_SHADING_COMMON_TRANSPARENCY_R:
+			TransparencyColor[0] = Value;
+			break;
+		case  AlShadingFields::kFLD_SHADING_COMMON_TRANSPARENCY_G:
+			TransparencyColor[1] = Value;
+			break;
+		case  AlShadingFields::kFLD_SHADING_COMMON_TRANSPARENCY_B:
+			TransparencyColor[2] = Value;
+			break;
+		case AlShadingFields::kFLD_SHADING_COMMON_GLOW_INTENSITY:
+			GlowIntensity = Value;
+			break;
+		default:
+			continue;
 		}
-		item = item->nextField();
 	}
 
-	IDatasmithMaterialExpressionScalar* DiffuseScalar = MaterialElement->AddMaterialExpression<IDatasmithMaterialExpressionScalar>();
-	DiffuseScalar->GetScalar() = Diffuse;
-	DiffuseScalar->SetName(TEXT("Diffuse"));
+	float Opacity = 1.0f - (TransparencyColor[0] + TransparencyColor[1] + TransparencyColor[2]) / 3.0f;
+	bool bIsTransparent = !FMath::IsNearlyEqual(Opacity, 1.0f);
 
-	IDatasmithMaterialExpressionGeneric* DiffuseMultiply = MaterialElement->AddMaterialExpression< IDatasmithMaterialExpressionGeneric >();
-	DiffuseMultiply->SetExpressionName(TEXT("Multiply"));
+	// Construct parameter expressions
+	IDatasmithMaterialExpressionColor* ColorExpression = MaterialElement->AddMaterialExpression<IDatasmithMaterialExpressionColor>();
+	ColorExpression->SetName(TEXT("Color"));
+	ColorExpression->GetColor() = FLinearColor(pow(Color[0] / 255.0, 2.2), pow(Color[1] / 255.0, 2.2), pow(Color[2] / 255.0, 2.2), 255);
 
-	//BaseColor->ConnectExpression(*DiffuseMultiply->GetInput(0));
-	DiffuseScalar->ConnectExpression(*DiffuseMultiply->GetInput(1));
+	IDatasmithMaterialExpressionColor* IncandescenceColorExpression = MaterialElement->AddMaterialExpression<IDatasmithMaterialExpressionColor>();
+	IncandescenceColorExpression->SetName(TEXT("IncandescenceColor"));
+	IncandescenceColorExpression->GetColor() = FLinearColor(pow(IncandescenceColor[0] / 255.0, 2.2), pow(IncandescenceColor[1] / 255.0, 2.2), pow(IncandescenceColor[2] / 255.0, 2.2), 255);
 
-	MaterialElement->GetBaseColor().SetExpression(DiffuseMultiply);
+	IDatasmithMaterialExpressionColor* TransparencyColorExpression = MaterialElement->AddMaterialExpression<IDatasmithMaterialExpressionColor>();
+	TransparencyColorExpression->SetName(TEXT("TransparencyColor"));
+	TransparencyColorExpression->GetColor() = FLinearColor(pow(TransparencyColor[0] / 255.0, 2.2), pow(TransparencyColor[1] / 255.0, 2.2), pow(TransparencyColor[2] / 255.0, 2.2), 255);
 
-	MaterialElement->SetParentLabel(TEXT("LAMBERT"));
-}
+	IDatasmithMaterialExpressionScalar* GlowIntensityExpression = MaterialElement->AddMaterialExpression<IDatasmithMaterialExpressionScalar>();
+	GlowIntensityExpression->GetScalar() = GlowIntensity;
+	GlowIntensityExpression->SetName(TEXT("GlowIntensity"));
+
+	// Create aux expressions
+	IDatasmithMaterialExpressionGeneric* BaseColorAdd = MaterialElement->AddMaterialExpression<IDatasmithMaterialExpressionGeneric>();
+	BaseColorAdd->SetExpressionName(TEXT("Add"));
+
+	IDatasmithMaterialExpressionGeneric* BaseColorTransparencyMultiply = MaterialElement->AddMaterialExpression<IDatasmithMaterialExpressionGeneric>();
+	BaseColorTransparencyMultiply->SetExpressionName(TEXT("Multiply"));
+
+	IDatasmithMaterialExpressionGeneric* IncandescenceMultiply = MaterialElement->AddMaterialExpression<IDatasmithMaterialExpressionGeneric>();
+	IncandescenceMultiply->SetExpressionName(TEXT("Multiply"));
+
+	IDatasmithMaterialExpressionGeneric* IncandescenceScaleMultiply = MaterialElement->AddMaterialExpression<IDatasmithMaterialExpressionGeneric>();
+	IncandescenceScaleMultiply->SetExpressionName(TEXT("Multiply"));
+
+	IDatasmithMaterialExpressionScalar* IncandescenceScale = MaterialElement->AddMaterialExpression<IDatasmithMaterialExpressionScalar>();
+	IncandescenceScale->GetScalar() = 100.0f;
+
+	IDatasmithMaterialExpressionGeneric* TransparencyOneMinus = MaterialElement->AddMaterialExpression<IDatasmithMaterialExpressionGeneric>();
+	TransparencyOneMinus->SetExpressionName(TEXT("OneMinus"));
+
+	IDatasmithMaterialExpressionFunctionCall* BreakFloat3 = nullptr;
+	IDatasmithMaterialExpressionGeneric* AddRG = nullptr;
+	IDatasmithMaterialExpressionGeneric* AddRGB = nullptr;
+	IDatasmithMaterialExpressionGeneric* Divide = nullptr;
+	IDatasmithMaterialExpressionScalar* DivideConstant = nullptr;
+	if (bIsTransparent)
+	{
+		BreakFloat3 = MaterialElement->AddMaterialExpression<IDatasmithMaterialExpressionFunctionCall>();
+		BreakFloat3->SetFunctionPathName(TEXT("/Engine/Functions/Engine_MaterialFunctions02/Utility/BreakFloat3Components.BreakFloat3Components"));
+
+		AddRG = MaterialElement->AddMaterialExpression<IDatasmithMaterialExpressionGeneric>();
+		AddRG->SetExpressionName(TEXT("Add"));
+
+		AddRGB = MaterialElement->AddMaterialExpression<IDatasmithMaterialExpressionGeneric>();
+		AddRGB->SetExpressionName(TEXT("Add"));
+
+		Divide = MaterialElement->AddMaterialExpression<IDatasmithMaterialExpressionGeneric>();
+		Divide->SetExpressionName(TEXT("Divide"));
+
+		DivideConstant = MaterialElement->AddMaterialExpression<IDatasmithMaterialExpressionScalar>();
+		DivideConstant->GetScalar() = 3.0f;
+	}
+
+	// Connect expressions
+	ColorExpression->ConnectExpression(*BaseColorAdd->GetInput(0));
+	IncandescenceColorExpression->ConnectExpression(*BaseColorAdd->GetInput(1));
+
+	BaseColorAdd->ConnectExpression(*BaseColorTransparencyMultiply->GetInput(0));
+	TransparencyOneMinus->ConnectExpression(*BaseColorTransparencyMultiply->GetInput(1));
+
+	GlowIntensityExpression->ConnectExpression(*IncandescenceScaleMultiply->GetInput(0));
+	IncandescenceScale->ConnectExpression(*IncandescenceScaleMultiply->GetInput(1));
+
+	BaseColorTransparencyMultiply->ConnectExpression(*IncandescenceMultiply->GetInput(0));
+	IncandescenceScaleMultiply->ConnectExpression(*IncandescenceMultiply->GetInput(1));
+
+	TransparencyColorExpression->ConnectExpression(*TransparencyOneMinus->GetInput(0));
+
+	if (bIsTransparent)
+	{
+		TransparencyOneMinus->ConnectExpression(*BreakFloat3->GetInput(0));
+
+		BreakFloat3->ConnectExpression(*AddRG->GetInput(0), 0);
+		BreakFloat3->ConnectExpression(*AddRG->GetInput(1), 1);
+
+		AddRG->ConnectExpression(*AddRGB->GetInput(0));
+		BreakFloat3->ConnectExpression(*AddRGB->GetInput(1), 2);
+
+		AddRGB->ConnectExpression(*Divide->GetInput(0));
+		DivideConstant->ConnectExpression(*Divide->GetInput(1));
+	}
+
+	// Connect material outputs
+	MaterialElement->GetBaseColor().SetExpression(BaseColorTransparencyMultiply);
+	MaterialElement->GetEmissiveColor().SetExpression(IncandescenceMultiply);
+	if (bIsTransparent)
+	{
+		MaterialElement->GetOpacity().SetExpression(Divide);
+	}
 
-void FWireTranslatorImpl::AddAlLightSourceParameters(AlShader *Shader, TSharedRef<IDatasmithUEPbrMaterialElement> MaterialElement)
-{
-	IDatasmithMaterialExpressionFlattenNormal* LightSource = MaterialElement->AddMaterialExpression< IDatasmithMaterialExpressionFlattenNormal >();
-	MaterialElement->GetNormal().SetExpression(LightSource);
 	MaterialElement->SetParentLabel(TEXT("LIGHTSOURCE"));
 }
 
 void FWireTranslatorImpl::AddAlPhongParameters(AlShader *Shader, TSharedRef<IDatasmithUEPbrMaterialElement> MaterialElement)
 {
-	float Specular[3];
-	Specular[0] = Specular[1] = Specular[2] = 0;
-	bool SpecularDefined = false;
-
-	float Reflectivity = 0;
-	float Diffuse = 0;
-	float Shinyness = 0;
-
-
-	AlList *list = Shader->fields();
-	AlShadingFieldItem *item = static_cast<AlShadingFieldItem *>(list->first());
-	while (item)
+	// Default values for a Phong material
+	double Color[] = { 0.57, 0.58, 0.60 };
+	double TransparencyColor[] = { 0.0, 0.0, 0.0 };
+	double IncandescenceColor[] = { 0.0, 0.0, 0.0 };
+	double SpecularColor[] = { 0.15, 0.15, 0.15 } ;
+	double Diffuse = 1.0;
+	double GlowIntensity = 0.0;
+	double Gloss = 0.8;
+	double Shinyness = 20.0;
+	double Specularity = 1.0;
+	double Reflectivity = 0.5;
+
+	AlList* List = Shader->fields();
+	for (AlShadingFieldItem* Item = static_cast<AlShadingFieldItem *>(List->first()); Item; Item = Item->nextField())
 	{
-		double Value = 0.0;
-		statusCode errorCode = Shader->parameter(item->field(), Value);
-		if (errorCode != 0)
+		double Value = 0.0f;
+		statusCode ErrorCode = Shader->parameter(Item->field(), Value);
+		if (ErrorCode != 0)
 		{
-			item = item->nextField();
 			continue;
 		}
 
-		switch (item->field())
+		switch (Item->field())
 		{
-		case  AlShadingFields::kFLD_SHADING_PHONG_DIFFUSE:
+		case AlShadingFields::kFLD_SHADING_PHONG_DIFFUSE:
 			Diffuse = Value;
 			break;
-		case  AlShadingFields::kFLD_SHADING_PHONG_REFLECTIVITY:
-			Reflectivity = Value;
+		case AlShadingFields::kFLD_SHADING_PHONG_GLOSS_:
+			Gloss = Value;
+			break;
+		case AlShadingFields::kFLD_SHADING_PHONG_SPECULAR_R:
+			SpecularColor[0] = Value;
+			break;
+		case AlShadingFields::kFLD_SHADING_PHONG_SPECULAR_G:
+			SpecularColor[1] = Value;
+			break;
+		case AlShadingFields::kFLD_SHADING_PHONG_SPECULAR_B:
+			SpecularColor[2] = Value;
+			break;
+		case AlShadingFields::kFLD_SHADING_PHONG_SPECULARITY_:
+			Specularity = Value;
 			break;
-		case  AlShadingFields::kFLD_SHADING_PHONG_SHINYNESS:
+		case AlShadingFields::kFLD_SHADING_PHONG_SHINYNESS:
 			Shinyness = Value;
 			break;
-		case  AlShadingFields::kFLD_SHADING_PHONG_SPECULAR_R:
-			Specular[0] = Value;
-			SpecularDefined = true;
+		case AlShadingFields::kFLD_SHADING_PHONG_REFLECTIVITY:
+			Reflectivity = Value;
+			break;
+		case AlShadingFields::kFLD_SHADING_COMMON_COLOR_R:
+			Color[0] = Value;
+			break;
+		case AlShadingFields::kFLD_SHADING_COMMON_COLOR_G:
+			Color[1] = Value;
+			break;
+		case  AlShadingFields::kFLD_SHADING_COMMON_COLOR_B:
+			Color[2] = Value;
+			break;
+		case AlShadingFields::kFLD_SHADING_COMMON_INCANDESCENCE_R:
+			IncandescenceColor[0] = Value;
+			break;
+		case AlShadingFields::kFLD_SHADING_COMMON_INCANDESCENCE_G:
+			IncandescenceColor[1] = Value;
+			break;
+		case AlShadingFields::kFLD_SHADING_COMMON_INCANDESCENCE_B:
+			IncandescenceColor[2] = Value;
+			break;
+		case  AlShadingFields::kFLD_SHADING_COMMON_TRANSPARENCY_R:
+			TransparencyColor[0] = Value;
+			break;
+		case  AlShadingFields::kFLD_SHADING_COMMON_TRANSPARENCY_G:
+			TransparencyColor[1] = Value;
 			break;
-		case  AlShadingFields::kFLD_SHADING_PHONG_SPECULAR_G:
-			Specular[1] = Value;
-			SpecularDefined = true;
+		case  AlShadingFields::kFLD_SHADING_COMMON_TRANSPARENCY_B:
+			TransparencyColor[2] = Value;
 			break;
-		case  AlShadingFields::kFLD_SHADING_PHONG_SPECULAR_B:
-			Specular[2] = Value;
-			SpecularDefined = true;
+		case AlShadingFields::kFLD_SHADING_COMMON_GLOW_INTENSITY:
+			GlowIntensity = Value;
 			break;
+			//case  AlShadingFields::kFLD_SHADING_COMMON_TRANSPARENCY_DEPTH:
+			//	TransparencyDepth = Value;
+			//	break;
+			//case  AlShadingFields::kFLD_SHADING_COMMON_TRANSPARENCY_SHADE:
+			//	TransparencyShade = Value;
+			//	break;
+		default:
+			continue;
 		}
-		item = item->nextField();
 	}
 
-	if (SpecularDefined)
+	float Opacity = 1.0f - (TransparencyColor[0] + TransparencyColor[1] + TransparencyColor[2]) / 3.0f;
+	bool bIsTransparent = !FMath::IsNearlyEqual(Opacity, 1.0f);
+
+	// Construct parameter expressions
+	IDatasmithMaterialExpressionScalar* DiffuseExpression = MaterialElement->AddMaterialExpression<IDatasmithMaterialExpressionScalar>();
+	DiffuseExpression->GetScalar() = Diffuse;
+	DiffuseExpression->SetName(TEXT("Diffuse"));
+
+	IDatasmithMaterialExpressionScalar* GlossExpression = MaterialElement->AddMaterialExpression<IDatasmithMaterialExpressionScalar>();
+	GlossExpression->GetScalar() = Gloss;
+	GlossExpression->SetName(TEXT("Gloss"));
+
+	IDatasmithMaterialExpressionColor* SpecularColorExpression = MaterialElement->AddMaterialExpression<IDatasmithMaterialExpressionColor>();
+	SpecularColorExpression->SetName(TEXT("SpecularColor"));
+	SpecularColorExpression->GetColor() = FLinearColor(pow(SpecularColor[0], 2.2), pow(SpecularColor[1], 2.2), pow(SpecularColor[2], 2.2), 1.0f);
+
+	IDatasmithMaterialExpressionScalar* SpecularityExpression = MaterialElement->AddMaterialExpression<IDatasmithMaterialExpressionScalar>();
+	SpecularityExpression->GetScalar() = Specularity * 0.3;
+	SpecularityExpression->SetName(TEXT("Specularity"));
+
+	IDatasmithMaterialExpressionScalar* ShinynessExpression = MaterialElement->AddMaterialExpression<IDatasmithMaterialExpressionScalar>();
+	ShinynessExpression->GetScalar() = Shinyness;
+	ShinynessExpression->SetName(TEXT("Shinyness"));
+
+	IDatasmithMaterialExpressionScalar* ReflectivityExpression = MaterialElement->AddMaterialExpression<IDatasmithMaterialExpressionScalar>();
+	ReflectivityExpression->GetScalar() = Reflectivity;
+	ReflectivityExpression->SetName(TEXT("Reflectivity"));
+
+	IDatasmithMaterialExpressionColor* ColorExpression = MaterialElement->AddMaterialExpression<IDatasmithMaterialExpressionColor>();
+	ColorExpression->SetName(TEXT("Color"));
+	ColorExpression->GetColor() = FLinearColor(pow(Color[0] / 255.0, 2.2), pow(Color[1] / 255.0, 2.2), pow(Color[2] / 255.0, 2.2), 255);
+
+	IDatasmithMaterialExpressionColor* IncandescenceColorExpression = MaterialElement->AddMaterialExpression<IDatasmithMaterialExpressionColor>();
+	IncandescenceColorExpression->SetName(TEXT("IncandescenceColor"));
+	IncandescenceColorExpression->GetColor() = FLinearColor(pow(IncandescenceColor[0] / 255.0, 2.2), pow(IncandescenceColor[1] / 255.0, 2.2), pow(IncandescenceColor[2] / 255.0, 2.2), 255);
+
+	IDatasmithMaterialExpressionColor* TransparencyColorExpression = MaterialElement->AddMaterialExpression<IDatasmithMaterialExpressionColor>();
+	TransparencyColorExpression->SetName(TEXT("TransparencyColor"));
+	TransparencyColorExpression->GetColor() = FLinearColor(pow(TransparencyColor[0] / 255.0, 2.2), pow(TransparencyColor[1] / 255.0, 2.2), pow(TransparencyColor[2] / 255.0, 2.2), 255);
+
+	IDatasmithMaterialExpressionScalar* GlowIntensityExpression = MaterialElement->AddMaterialExpression<IDatasmithMaterialExpressionScalar>();
+	GlowIntensityExpression->GetScalar() = GlowIntensity;
+	GlowIntensityExpression->SetName(TEXT("GlowIntensity"));
+
+	// Create aux expressions
+	IDatasmithMaterialExpressionGeneric* ColorSpecLerp = MaterialElement->AddMaterialExpression<IDatasmithMaterialExpressionGeneric>();
+	ColorSpecLerp->SetExpressionName(TEXT("LinearInterpolate"));
+
+	IDatasmithMaterialExpressionScalar* ColorSpecLerpValue = MaterialElement->AddMaterialExpression<IDatasmithMaterialExpressionScalar>();
+	ColorSpecLerpValue->GetScalar() = 0.96f;
+
+	IDatasmithMaterialExpressionGeneric* ColorMetallicLerp = MaterialElement->AddMaterialExpression<IDatasmithMaterialExpressionGeneric>();
+	ColorMetallicLerp->SetExpressionName(TEXT("LinearInterpolate"));
+
+	IDatasmithMaterialExpressionGeneric* DiffuseLerp = MaterialElement->AddMaterialExpression<IDatasmithMaterialExpressionGeneric>();
+	DiffuseLerp->SetExpressionName(TEXT("LinearInterpolate"));
+
+	IDatasmithMaterialExpressionScalar* DiffuseLerpA = MaterialElement->AddMaterialExpression<IDatasmithMaterialExpressionScalar>();
+	DiffuseLerpA->GetScalar() = 0.04f;
+
+	IDatasmithMaterialExpressionScalar* DiffuseLerpB = MaterialElement->AddMaterialExpression<IDatasmithMaterialExpressionScalar>();
+	DiffuseLerpB->GetScalar() = 1.0f;
+
+	IDatasmithMaterialExpressionGeneric* BaseColorMultiply = MaterialElement->AddMaterialExpression<IDatasmithMaterialExpressionGeneric>();
+	BaseColorMultiply->SetExpressionName(TEXT("Multiply"));
+
+	IDatasmithMaterialExpressionGeneric* BaseColorAdd = MaterialElement->AddMaterialExpression<IDatasmithMaterialExpressionGeneric>();
+	BaseColorAdd->SetExpressionName(TEXT("Add"));
+
+	IDatasmithMaterialExpressionGeneric* BaseColorTransparencyMultiply = MaterialElement->AddMaterialExpression<IDatasmithMaterialExpressionGeneric>();
+	BaseColorTransparencyMultiply->SetExpressionName(TEXT("Multiply"));
+
+	IDatasmithMaterialExpressionGeneric* IncandescenceMultiply = MaterialElement->AddMaterialExpression<IDatasmithMaterialExpressionGeneric>();
+	IncandescenceMultiply->SetExpressionName(TEXT("Multiply"));
+
+	IDatasmithMaterialExpressionGeneric* IncandescenceScaleMultiply = MaterialElement->AddMaterialExpression<IDatasmithMaterialExpressionGeneric>();
+	IncandescenceScaleMultiply->SetExpressionName(TEXT("Multiply"));
+
+	IDatasmithMaterialExpressionScalar* IncandescenceScale = MaterialElement->AddMaterialExpression<IDatasmithMaterialExpressionScalar>();
+	IncandescenceScale->GetScalar() = 100.0f;
+
+	IDatasmithMaterialExpressionGeneric* ShinynessSubtract = MaterialElement->AddMaterialExpression<IDatasmithMaterialExpressionGeneric>();
+	ShinynessSubtract->SetExpressionName(TEXT("Subtract"));
+
+	IDatasmithMaterialExpressionScalar* ShinynessSubtract2 = MaterialElement->AddMaterialExpression<IDatasmithMaterialExpressionScalar>();
+	ShinynessSubtract2->GetScalar() = 2.0f;
+
+	IDatasmithMaterialExpressionGeneric* ShinynessDivide = MaterialElement->AddMaterialExpression<IDatasmithMaterialExpressionGeneric>();
+	ShinynessDivide->SetExpressionName(TEXT("Divide"));
+
+	IDatasmithMaterialExpressionScalar* ShinynessDivide98 = MaterialElement->AddMaterialExpression<IDatasmithMaterialExpressionScalar>();
+	ShinynessDivide98->GetScalar() = 98.0f;
+
+	IDatasmithMaterialExpressionGeneric* SpecularityMultiply = MaterialElement->AddMaterialExpression<IDatasmithMaterialExpressionGeneric>();
+	SpecularityMultiply->SetExpressionName(TEXT("Multiply"));
+
+	IDatasmithMaterialExpressionGeneric* RoughnessOneMinus = MaterialElement->AddMaterialExpression<IDatasmithMaterialExpressionGeneric>();
+	RoughnessOneMinus->SetExpressionName(TEXT("OneMinus"));
+
+	IDatasmithMaterialExpressionGeneric* TransparencyOneMinus = MaterialElement->AddMaterialExpression<IDatasmithMaterialExpressionGeneric>();
+	TransparencyOneMinus->SetExpressionName(TEXT("OneMinus"));
+
+	IDatasmithMaterialExpressionFunctionCall* BreakFloat3 = nullptr;
+	IDatasmithMaterialExpressionGeneric* AddRG = nullptr;
+	IDatasmithMaterialExpressionGeneric* AddRGB = nullptr;
+	IDatasmithMaterialExpressionGeneric* Divide = nullptr;
+	IDatasmithMaterialExpressionScalar* DivideConstant = nullptr;
+	if (bIsTransparent)
 	{
-		FLinearColor SpecularColor = FLinearColor::FromPow22Color(FColor(Specular[0], Specular[1], Specular[2], 255));
-		IDatasmithMaterialExpressionColor* SpecularExpression = MaterialElement->AddMaterialExpression<IDatasmithMaterialExpressionColor>();
-		SpecularExpression->SetName(TEXT("Specular"));
-		SpecularExpression->GetColor() = SpecularColor;
-		MaterialElement->GetSpecular().SetExpression(SpecularExpression);
+		BreakFloat3 = MaterialElement->AddMaterialExpression<IDatasmithMaterialExpressionFunctionCall>();
+		BreakFloat3->SetFunctionPathName(TEXT("/Engine/Functions/Engine_MaterialFunctions02/Utility/BreakFloat3Components.BreakFloat3Components"));
+
+		AddRG = MaterialElement->AddMaterialExpression<IDatasmithMaterialExpressionGeneric>();
+		AddRG->SetExpressionName(TEXT("Add"));
+
+		AddRGB = MaterialElement->AddMaterialExpression<IDatasmithMaterialExpressionGeneric>();
+		AddRGB->SetExpressionName(TEXT("Add"));
+
+		Divide = MaterialElement->AddMaterialExpression<IDatasmithMaterialExpressionGeneric>();
+		Divide->SetExpressionName(TEXT("Divide"));
+
+		DivideConstant = MaterialElement->AddMaterialExpression<IDatasmithMaterialExpressionScalar>();
+		DivideConstant->GetScalar() = 3.0f;
 	}
 
-	IDatasmithMaterialExpressionScalar* ReflectivityScalar = MaterialElement->AddMaterialExpression<IDatasmithMaterialExpressionScalar>();
-	ReflectivityScalar->GetScalar() = Reflectivity;
-	ReflectivityScalar->SetName(TEXT("Reflectivity"));
+	// Connect expressions
+	SpecularColorExpression->ConnectExpression(*ColorSpecLerp->GetInput(0));
+	ColorExpression->ConnectExpression(*ColorSpecLerp->GetInput(1));
+	ColorSpecLerpValue->ConnectExpression(*ColorSpecLerp->GetInput(2));
+
+	ColorExpression->ConnectExpression(*ColorMetallicLerp->GetInput(0));
+	ColorSpecLerp->ConnectExpression(*ColorMetallicLerp->GetInput(1));
+	GlossExpression->ConnectExpression(*ColorMetallicLerp->GetInput(2));
+
+	DiffuseLerpA->ConnectExpression(*DiffuseLerp->GetInput(0));
+	DiffuseLerpB->ConnectExpression(*DiffuseLerp->GetInput(1));
+	DiffuseExpression->ConnectExpression(*DiffuseLerp->GetInput(2));
+
+	ColorMetallicLerp->ConnectExpression(*BaseColorMultiply->GetInput(0));
+	DiffuseLerp->ConnectExpression(*BaseColorMultiply->GetInput(1));
+
+	BaseColorMultiply->ConnectExpression(*BaseColorAdd->GetInput(0));
+	IncandescenceColorExpression->ConnectExpression(*BaseColorAdd->GetInput(1));
+
+	BaseColorAdd->ConnectExpression(*BaseColorTransparencyMultiply->GetInput(0));
+	TransparencyOneMinus->ConnectExpression(*BaseColorTransparencyMultiply->GetInput(1));
+
+	GlowIntensityExpression->ConnectExpression(*IncandescenceScaleMultiply->GetInput(0));
+	IncandescenceScale->ConnectExpression(*IncandescenceScaleMultiply->GetInput(1));
+
+	BaseColorTransparencyMultiply->ConnectExpression(*IncandescenceMultiply->GetInput(0));
+	IncandescenceScaleMultiply->ConnectExpression(*IncandescenceMultiply->GetInput(1));
+
+	ShinynessExpression->ConnectExpression(*ShinynessSubtract->GetInput(0));
+	ShinynessSubtract2->ConnectExpression(*ShinynessSubtract->GetInput(1));
+
+	ShinynessSubtract->ConnectExpression(*ShinynessDivide->GetInput(0));
+	ShinynessDivide98->ConnectExpression(*ShinynessDivide->GetInput(1));
 
-	//MaterialElement->GetMetallic().SetExpression(ReflectivityScalar);
+	ShinynessDivide->ConnectExpression(*SpecularityMultiply->GetInput(0));
+	SpecularityExpression->ConnectExpression(*SpecularityMultiply->GetInput(1));
 
-	IDatasmithMaterialExpressionGeneric* OneMinus = MaterialElement->AddMaterialExpression< IDatasmithMaterialExpressionGeneric >();
-	OneMinus->SetExpressionName(TEXT("OneMinus"));
-	ReflectivityScalar->ConnectExpression(*OneMinus->GetInput(0));
-	MaterialElement->GetRoughness().SetExpression(OneMinus);
+	SpecularityMultiply->ConnectExpression(*RoughnessOneMinus->GetInput(0));
+
+	TransparencyColorExpression->ConnectExpression(*TransparencyOneMinus->GetInput(0));
+
+	if (bIsTransparent)
+	{
+		TransparencyOneMinus->ConnectExpression(*BreakFloat3->GetInput(0));
+
+		BreakFloat3->ConnectExpression(*AddRG->GetInput(0), 0);
+		BreakFloat3->ConnectExpression(*AddRG->GetInput(1), 1);
+
+		AddRG->ConnectExpression(*AddRGB->GetInput(0));
+		BreakFloat3->ConnectExpression(*AddRGB->GetInput(1), 2);
+
+		AddRGB->ConnectExpression(*Divide->GetInput(0));
+		DivideConstant->ConnectExpression(*Divide->GetInput(1));
+	}
+
+	// Connect material outputs
+	MaterialElement->GetBaseColor().SetExpression(BaseColorTransparencyMultiply);
+	MaterialElement->GetMetallic().SetExpression(GlossExpression);
+	MaterialElement->GetSpecular().SetExpression(ReflectivityExpression);
+	MaterialElement->GetRoughness().SetExpression(RoughnessOneMinus);
+	MaterialElement->GetEmissiveColor().SetExpression(IncandescenceMultiply);
+	if (bIsTransparent)
+	{
+		MaterialElement->GetOpacity().SetExpression(Divide);
+	}
 
 	MaterialElement->SetParentLabel(TEXT("PHONG"));
 }
@@ -566,7 +1177,7 @@ bool FWireTranslatorImpl::GetShader()
 	{
 		FString ShaderName = Shader->name();
 		FString ShaderModelName = Shader->shadingModel();
-		
+
 		uint32 ShaderUUID = fabs((int32)GetTypeHash(*ShaderName));
 
 		TSharedRef<IDatasmithUEPbrMaterialElement> MaterialElement = FDatasmithSceneFactory::CreateUEPbrMaterial(*ShaderName);
@@ -574,27 +1185,22 @@ bool FWireTranslatorImpl::GetShader()
 		MaterialElement->SetLabel(*ShaderName);
 		MaterialElement->SetName(*FString::FromInt(ShaderUUID));
 
-		CreateAlCommonMaterial(Shader, MaterialElement);
-
-		if (ShaderModelName.Equals(TEXT("BLINN"))) 
+		if (ShaderModelName.Equals(TEXT("BLINN")))
 		{
 			AddAlBlinnParameters(Shader, MaterialElement);
 		}
-		else if (ShaderModelName.Equals(TEXT("LAMBERT"))) 
+		else if (ShaderModelName.Equals(TEXT("LAMBERT")))
 		{
 			AddAlLambertParameters(Shader, MaterialElement);
 		}
-		else if (ShaderModelName.Equals(TEXT("LIGHTSOURCE"))) 
+		else if (ShaderModelName.Equals(TEXT("LIGHTSOURCE")))
 		{
 			AddAlLightSourceParameters(Shader, MaterialElement);
 		}
-		else if (ShaderModelName.Equals(TEXT("PHONG"))) 
+		else if (ShaderModelName.Equals(TEXT("PHONG")))
 		{
 			AddAlPhongParameters(Shader, MaterialElement);
 		}
-		else
-		{
-		}
 
 		DatasmithScene->AddMaterial(MaterialElement);
 
@@ -744,7 +1350,7 @@ TSharedPtr< IDatasmithMeshElement > FWireTranslatorImpl::FindOrAddMeshElement(Al
 	//{
 	//	MeshElement->SetDimensions(BoundingBox[3] - BoundingBox[0], BoundingBox[4] - BoundingBox[1], BoundingBox[5] - BoundingBox[2], 0.0f);
 	//}
-	 
+
 
 	if (ShaderName)
 	{
@@ -785,7 +1391,7 @@ bool FWireTranslatorImpl::ProcessAlShellNode(AlDagNode& ShellNode, const FDagNod
 
 	SetActorTransform(ShellInfo.ActorElement, ShellNode);
 
-	//// Apply materials on the current part 
+	//// Apply materials on the current part
 	if (ShaderName)
 	{
 		TSharedPtr< IDatasmithMaterialIDElement > MaterialIDElement = ShaderNameToUEMaterialId[FString(ShaderName)];
@@ -808,7 +1414,7 @@ bool FWireTranslatorImpl::ProcessAlShellNode(AlDagNode& ShellNode, const FDagNod
 		{
 			DatasmithScene->AddActor(ActorElement);
 		}
-	}	
+	}
 	return true;
 }
 
@@ -842,7 +1448,7 @@ bool FWireTranslatorImpl::ProcessBodyNode(TSharedRef<BodyData> Body, const FDagN
 
 	//SetActorTransform(ShellInfo.ActorElement, ShellNode);
 
-	//// Apply materials on the current part 
+	//// Apply materials on the current part
 	if (*Body->ShaderName)
 	{
 		TSharedPtr< IDatasmithMaterialIDElement > MaterialIDElement = ShaderNameToUEMaterialId[FString(*Body->ShaderName)];
@@ -1262,7 +1868,7 @@ TOptional<FMeshDescription> FWireTranslatorImpl::GetMeshOfMeshBody(TSharedRef<Bo
 	for (auto DagNode : Body->ShellSet)
 	{
 		AlMeshNode *MeshNode = DagNode->asMeshNodePtr();
-		AlMesh *Mesh = MeshNode->mesh();		
+		AlMesh *Mesh = MeshNode->mesh();
 		if (Mesh)
 		{
 			//TransferAlMeshToMeshDescription(*Mesh, MeshDescription, MeshParameters, True);
@@ -1315,7 +1921,7 @@ TOptional<FMeshDescription> FWireTranslatorImpl::GetMeshDescription(TSharedRef<I
 	if(Body->bCadData)
 	{
 		return GetMeshOfShellBody(Body, MeshElement, MeshParameters);
-	} 
+	}
 	else
 	{
 		return GetMeshOfMeshBody(Body, MeshElement, MeshParameters);
@@ -1412,7 +2018,7 @@ void FDatasmithWireTranslator::Initialize(FDatasmithTranslatorCapabilities& OutC
 	{
 		OutCapabilities.SupportedFileFormats.Add(FFileFormatInfo{ TEXT("wire"), TEXT("AliasStudio, Model files") });
 		return;
-	} 
+	}
 #endif
 	OutCapabilities.bIsEnabled = false;
 }
diff --git a/Engine/Plugins/Enterprise/DatasmithIFCImporter/Source/DatasmithIFCTranslator/Private/DatasmithIFCImporter.cpp b/Engine/Plugins/Enterprise/DatasmithIFCImporter/Source/DatasmithIFCTranslator/Private/DatasmithIFCImporter.cpp
--- a/Engine/Plugins/Enterprise/DatasmithIFCImporter/Source/DatasmithIFCTranslator/Private/DatasmithIFCImporter.cpp
+++ b/Engine/Plugins/Enterprise/DatasmithIFCImporter/Source/DatasmithIFCTranslator/Private/DatasmithIFCImporter.cpp
@@ -250,7 +250,8 @@ TSharedPtr<IDatasmithActorElement> FDatasmithIFCImporter::ConvertNode(const IFC:
 			Containment->AddChild(ConvertNode(IFCReader->GetObjects()[PartIndex]));
 		}
 	}
-
+	
+	ImportedIFCInstances.Add(InObject.IfcInstance);
 	return ActorElement;
 }
 
@@ -321,14 +322,19 @@ bool FDatasmithIFCImporter::SendSceneToDatasmith()
 		DatasmithScene->AddActor(ConvertNode(Project));
 	}
 
+	// Import all unreferenced objects
+	TSharedPtr<IDatasmithActorElement> UnreferencedRoot = FDatasmithSceneFactory::CreateActor(TEXT("Unreferenced objects"));	
 	for (const IFC::FObject& IFCObjectRef : IFCReader->GetObjects())
 	{
-		if (IFCObjectRef.bRootObject == false)
+		if (IFCObjectRef.bRootObject && !ImportedIFCInstances.Contains(IFCObjectRef.IfcInstance))
 		{
-			continue;
+			UnreferencedRoot->AddChild(ConvertNode(IFCObjectRef));
 		}
 	}
-
+	if (UnreferencedRoot->GetChildrenCount() > 0)
+	{
+		DatasmithScene->AddActor(UnreferencedRoot);
+	}
 #endif
 
 	return true;
diff --git a/Engine/Plugins/Enterprise/DatasmithIFCImporter/Source/DatasmithIFCTranslator/Private/DatasmithIFCImporter.h b/Engine/Plugins/Enterprise/DatasmithIFCImporter/Source/DatasmithIFCTranslator/Private/DatasmithIFCImporter.h
--- a/Engine/Plugins/Enterprise/DatasmithIFCImporter/Source/DatasmithIFCTranslator/Private/DatasmithIFCImporter.h
+++ b/Engine/Plugins/Enterprise/DatasmithIFCImporter/Source/DatasmithIFCTranslator/Private/DatasmithIFCImporter.h
@@ -94,5 +94,9 @@ private:
 
 	TMap<IDatasmithMeshElement*, int64> MeshElementToIFCMeshIndex;
 	TMap<int64, TSharedRef<IDatasmithMeshElement>> IFCMeshIndexToMeshElement;
+
+	// Keep track of everything we import, because at the end we will import
+	// all unreferenced objects separately
+	TSet<int64> ImportedIFCInstances;
 };
 
