Refactored re-tessellation code to be shared between UCoreTechRetessellateAction and UDataprepTessellationOperation
Added a Blueprint function to perform re-tesselation in a Blueprint graph
Added some reporting to the tessellation Dataprep operation

diff --git a/Engine/Plugins/Enterprise/DatasmithCADImporter/Source/DatasmithCoreTechExtension/Private/CoreTechBlueprintLibrary.cpp b/Engine/Plugins/Enterprise/DatasmithCADImporter/Source/DatasmithCoreTechExtension/Private/CoreTechBlueprintLibrary.cpp
new file mode 100644
--- /dev/null
+++ b/Engine/Plugins/Enterprise/DatasmithCADImporter/Source/DatasmithCoreTechExtension/Private/CoreTechBlueprintLibrary.cpp
@@ -0,0 +1,106 @@
+// Copyright 1998-2019 Epic Games, Inc. All Rights Reserved.
+
+#include "CoreTechBlueprintLibrary.h"
+
+#ifdef CAD_LIBRARY
+#include "CoreTechHelper.h"
+#include "CoreTechMeshLoader.h"
+#endif
+#include "CoreTechParametricSurfaceExtension.h"
+#include "DatasmithAdditionalData.h"
+#include "DatasmithStaticMeshImporter.h" // Call to BuildStaticMesh
+#include "DatasmithUtils.h"
+#include "Translators/DatasmithTranslator.h"
+#include "UI/DatasmithDisplayHelper.h"
+
+#include "Algo/AnyOf.h"
+#include "AssetData.h"
+#include "Async/ParallelFor.h"
+#include "Engine/StaticMesh.h"
+#include "IStaticMeshEditor.h"
+#include "StaticMeshAttributes.h"
+#include "Misc/FileHelper.h"
+#include "Misc/ScopedSlowTask.h"
+#include "Toolkits/IToolkit.h"
+#include "Toolkits/ToolkitManager.h"
+#include "UObject/StrongObjectPtr.h"
+#include "Algo/Transform.h"
+
+
+#define LOCTEXT_NAMESPACE "CoreTechRetessellateAction"
+
+bool UCoreTechBlueprintLibrary::RetessellateStaticMesh( UStaticMesh* StaticMesh, const FDatasmithTessellationOptions& TessellationSettings, bool bPostChanges, FText& OutReason )
+{
+	bool bTessellationOutcome = false;
+
+#ifdef CAD_LIBRARY
+	int32 LODIndex = 0;
+
+	FAssetData AssetData( StaticMesh );
+	if( UCoreTechParametricSurfaceData* CoreTechData = Datasmith::GetAdditionalData<UCoreTechParametricSurfaceData>(AssetData) )
+	{
+		// Make sure MeshDescription exists
+		FMeshDescription* DestinationMeshDescription = StaticMesh->GetMeshDescription( LODIndex );
+		if (DestinationMeshDescription == nullptr)
+		{
+			DestinationMeshDescription = StaticMesh->CreateMeshDescription(0);
+		}
+
+		if( DestinationMeshDescription )
+		{
+			if(bPostChanges)
+			{
+				StaticMesh->Modify();
+				StaticMesh->PreEditChange( nullptr );
+			}
+
+			if( FCoreTechRetessellate_Impl::ApplyOnOneAsset( *StaticMesh, *CoreTechData, TessellationSettings ) )
+			{
+				// Post static mesh has changed
+				if(bPostChanges)
+				{
+					FDatasmithStaticMeshImporter::PreBuildStaticMesh(StaticMesh); // handle uvs stuff
+					FDatasmithStaticMeshImporter::BuildStaticMesh(StaticMesh);
+
+					StaticMesh->PostEditChange();
+					StaticMesh->MarkPackageDirty();
+
+					// Refresh associated editor
+					TSharedPtr<IToolkit> EditingToolkit = FToolkitManager::Get().FindEditorForAsset(StaticMesh);
+					if (IStaticMeshEditor* StaticMeshEditorInUse = StaticCastSharedPtr<IStaticMeshEditor>(EditingToolkit).Get())
+					{
+						StaticMeshEditorInUse->RefreshTool();
+					}
+				}
+				// No posting required, just make sure the new tessellation is committed
+				else
+				{
+					UStaticMesh::FCommitMeshDescriptionParams Params;
+					Params.bMarkPackageDirty = false;
+					Params.bUseHashAsGuid = true;
+					StaticMesh->CommitMeshDescription( LODIndex, Params );
+				}
+
+				// Save last tessellation settings
+				CoreTechData->LastTessellationOptions = TessellationSettings;
+
+				bTessellationOutcome = true;
+			}
+			else
+			{
+				OutReason = NSLOCTEXT("BlueprintRetessellation", "LoadFailed", "Cannot generate mesh from parametric surface data");
+			}
+		}
+		else
+		{
+			OutReason = NSLOCTEXT("BlueprintRetessellation", "MeshDescriptionMissing", "Cannot create mesh description");
+		}
+	}
+	else
+	{
+		OutReason = NSLOCTEXT("BlueprintRetessellation", "MissingData", "No tessellation data attached to the static mesh");
+	}
+#endif
+
+	return bTessellationOutcome;
+}
\ No newline at end of file
diff --git a/Engine/Plugins/Enterprise/DatasmithCADImporter/Source/DatasmithCoreTechExtension/Private/CoreTechParametricSurfaceExtension.cpp b/Engine/Plugins/Enterprise/DatasmithCADImporter/Source/DatasmithCoreTechExtension/Private/CoreTechParametricSurfaceExtension.cpp
--- a/Engine/Plugins/Enterprise/DatasmithCADImporter/Source/DatasmithCoreTechExtension/Private/CoreTechParametricSurfaceExtension.cpp
+++ b/Engine/Plugins/Enterprise/DatasmithCADImporter/Source/DatasmithCoreTechExtension/Private/CoreTechParametricSurfaceExtension.cpp
@@ -14,6 +14,7 @@
 
 #include "Algo/AnyOf.h"
 #include "AssetData.h"
+#include "Async/ParallelFor.h"
 #include "Engine/StaticMesh.h"
 #include "IStaticMeshEditor.h"
 #include "StaticMeshAttributes.h"
@@ -50,6 +51,9 @@ void FCoreTechRetessellate_Impl::ApplyOnAssets(const TArray<FAssetData>& Selecte
 	int32 NumAssetsToProcess = SelectedAssets.Num();
 	bool bAskForSameOption = NumAssetsToProcess > 1;
 
+	TArray<UStaticMesh*> TessellatedMeshes;
+	TessellatedMeshes.Reserve( SelectedAssets.Num() );
+
 	TUniquePtr<FScopedSlowTask> Progress;
 	int32 AssetIndex = -1;
 	for (const FAssetData& Asset : SelectedAssets)
@@ -106,23 +110,61 @@ void FCoreTechRetessellate_Impl::ApplyOnAssets(const TArray<FAssetData>& Selecte
 						Progress->EnterProgressFrame(1, Text);
 					}
 
-					ApplyOnOneAsset(*StaticMesh, *CoreTechData, *RetessellateOptions.Get());
+					if( StaticMesh->GetMeshDescription(0) == nullptr)
+					{
+						StaticMesh->CreateMeshDescription( 0 );
+					}
 
-					// Refresh associated editor
-					TSharedPtr<IToolkit> EditingToolkit = FToolkitManager::Get().FindEditorForAsset(StaticMesh);
-					if (IStaticMeshEditor* StaticMeshEditorInUse = StaticCastSharedPtr<IStaticMeshEditor>(EditingToolkit).Get())
+					if(StaticMesh->GetMeshDescription(0) != nullptr)
 					{
-						StaticMeshEditorInUse->RefreshTool();
+						StaticMesh->Modify();
+						StaticMesh->PreEditChange( nullptr );
+
+						if( ApplyOnOneAsset(*StaticMesh, *CoreTechData, RetessellateOptions->Options) )
+						{
+							TessellatedMeshes.Add(StaticMesh);
+						}
 					}
 				}
 			}
 		}
 	}
+
+	// Make sure lightmap settings are valid
+	if(TessellatedMeshes.Num() > 1)
+	{
+		ParallelFor(TessellatedMeshes.Num(), [&](int32 Index)
+		{
+			FDatasmithStaticMeshImporter::PreBuildStaticMesh(TessellatedMeshes[Index]); 
+		});
+	}
+	else if(TessellatedMeshes.Num() > 0)
+	{
+		FDatasmithStaticMeshImporter::PreBuildStaticMesh( TessellatedMeshes[0] );
+	}
+
+	FDatasmithStaticMeshImporter::BuildStaticMeshes( TessellatedMeshes );
+
+	for(UStaticMesh* StaticMesh : TessellatedMeshes)
+	{
+		StaticMesh->PostEditChange();
+		StaticMesh->MarkPackageDirty();
+
+		// Refresh associated editor
+		TSharedPtr<IToolkit> EditingToolkit = FToolkitManager::Get().FindEditorForAsset(StaticMesh);
+		if (IStaticMeshEditor* StaticMeshEditorInUse = StaticCastSharedPtr<IStaticMeshEditor>(EditingToolkit).Get())
+		{
+			StaticMeshEditorInUse->RefreshTool();
+		}
+	}
+
 #endif // CAD_LIBRARY
 }
 
-void FCoreTechRetessellate_Impl::ApplyOnOneAsset(UStaticMesh& StaticMesh, UCoreTechParametricSurfaceData& CoreTechData, UCoreTechRetessellateActionOptions& RetessellateOptions)
+bool FCoreTechRetessellate_Impl::ApplyOnOneAsset(UStaticMesh& StaticMesh, UCoreTechParametricSurfaceData& CoreTechData, const FDatasmithTessellationOptions& RetessellateOptions)
 {
+	bool bSuccessfulTessellation = false;
+
 #ifdef CAD_LIBRARY
 	// make a temporary file as GPure can only deal with files.
 	FString ResourceFile = CoreTechData.SourceFile.IsEmpty() ? FPaths::ProjectIntermediateDir() / "temp.ct" : CoreTechData.SourceFile;
@@ -134,10 +176,10 @@ void FCoreTechRetessellate_Impl::ApplyOnOneAsset(UStaticMesh& StaticMesh, UCoreT
 	CADLibrary::FImportParameters ImportParameters;
 	ImportParameters.MetricUnit = CoreTechData.SceneParameters.MetricUnit;
 	ImportParameters.ScaleFactor = CoreTechData.SceneParameters.ScaleFactor;
-	ImportParameters.ChordTolerance = RetessellateOptions.Options.ChordTolerance / CoreTechData.SceneParameters.ScaleFactor;
-	ImportParameters.MaxEdgeLength = RetessellateOptions.Options.MaxEdgeLength / CoreTechData.SceneParameters.ScaleFactor;
-	ImportParameters.MaxNormalAngle = RetessellateOptions.Options.NormalTolerance;
-	ImportParameters.StitchingTechnique = CADLibrary::EStitchingTechnique(RetessellateOptions.Options.StitchingTechnique);
+	ImportParameters.ChordTolerance = RetessellateOptions.ChordTolerance / CoreTechData.SceneParameters.ScaleFactor;
+	ImportParameters.MaxEdgeLength = RetessellateOptions.MaxEdgeLength / CoreTechData.SceneParameters.ScaleFactor;
+	ImportParameters.MaxNormalAngle = RetessellateOptions.NormalTolerance;
+	ImportParameters.StitchingTechnique = CADLibrary::EStitchingTechnique(RetessellateOptions.StitchingTechnique);
 
 	CADLibrary::FMeshParameters MeshParameters;
 	MeshParameters.ModelCoordSys = FDatasmithUtils::EModelCoordSystem(CoreTechData.SceneParameters.ModelCoordSys);
@@ -147,47 +189,33 @@ void FCoreTechRetessellate_Impl::ApplyOnOneAsset(UStaticMesh& StaticMesh, UCoreT
 	MeshParameters.SymmetricOrigin = CoreTechData.MeshParameters.SymmetricOrigin;
 
 	// Previous MeshDescription is get to be able to create a new one with the same order of PolygonGroup (the matching of color and partition is currently based on their order)
-	FMeshDescription* DestinationMeshDescription = StaticMesh.GetMeshDescription(0);
-	if (DestinationMeshDescription == nullptr)
-	{
-		DestinationMeshDescription = StaticMesh.CreateMeshDescription(0);
-	}
-	FStaticMeshAttributes DestinationMeshDescriptionAttributes(*DestinationMeshDescription);
-
-	FMeshDescription MeshDescription;
-	FStaticMeshAttributes MeshDescriptionAttributes(MeshDescription);
-	MeshDescriptionAttributes.Register();
-
-	TPolygonGroupAttributesRef<FName> PolygonGroupDestinationMeshSlotNames = DestinationMeshDescriptionAttributes.GetPolygonGroupMaterialSlotNames();
-	TPolygonGroupAttributesRef<FName> PolygonGroupImportedMaterialSlotNames = MeshDescriptionAttributes.GetPolygonGroupMaterialSlotNames();
-	for (FPolygonGroupID PolygonGroupID : DestinationMeshDescription->PolygonGroups().GetElementIDs())
-	{
-		FName ImportedSlotName = PolygonGroupDestinationMeshSlotNames[PolygonGroupID];
-		FPolygonGroupID PolyGroupID = MeshDescription.CreatePolygonGroup();
-		PolygonGroupImportedMaterialSlotNames[PolyGroupID] = ImportedSlotName;
-	}
-
-	if (!Loader.LoadFile(ResourceFile, MeshDescription, ImportParameters, MeshParameters))
+	if(FMeshDescription* DestinationMeshDescription = StaticMesh.GetMeshDescription(0))
 	{
-		return;
-	}
+		FStaticMeshAttributes DestinationMeshDescriptionAttributes(*DestinationMeshDescription);
 
-	if (DestinationMeshDescription)	// @TODO: check: should be: if (ret)? Checking DestinationMeshDescription is redundant.
-	{
-		StaticMesh.Modify();
-		StaticMesh.PreEditChange( nullptr );
-		*DestinationMeshDescription = MoveTemp(MeshDescription);
+		FMeshDescription MeshDescription;
+		FStaticMeshAttributes MeshDescriptionAttributes(MeshDescription);
+		MeshDescriptionAttributes.Register();
 
-		// @TODO: check: no commit?
-
-		// Make sure lightmap settings are valid
-		FDatasmithStaticMeshImporter::PreBuildStaticMesh(&StaticMesh); // handle uvs stuff
-		FDatasmithStaticMeshImporter::BuildStaticMesh(&StaticMesh);
+		TPolygonGroupAttributesRef<FName> PolygonGroupDestinationMeshSlotNames = DestinationMeshDescriptionAttributes.GetPolygonGroupMaterialSlotNames();
+		TPolygonGroupAttributesRef<FName> PolygonGroupImportedMaterialSlotNames = MeshDescriptionAttributes.GetPolygonGroupMaterialSlotNames();
+		for (FPolygonGroupID PolygonGroupID : DestinationMeshDescription->PolygonGroups().GetElementIDs())
+		{
+			FName ImportedSlotName = PolygonGroupDestinationMeshSlotNames[PolygonGroupID];
+			FPolygonGroupID PolyGroupID = MeshDescription.CreatePolygonGroup();
+			PolygonGroupImportedMaterialSlotNames[PolyGroupID] = ImportedSlotName;
+		}
 
-		StaticMesh.PostEditChange();
-		StaticMesh.MarkPackageDirty();
+		if ( Loader.LoadFile(ResourceFile, MeshDescription, ImportParameters, MeshParameters))
+		{
+			// @TODO: check: no commit?
+			*DestinationMeshDescription = MoveTemp(MeshDescription);
+			bSuccessfulTessellation = true;
+		}
 	}
 #endif // CAD_LIBRARY
+
+	return bSuccessfulTessellation;
 }
 
 #undef LOCTEXT_NAMESPACE
diff --git a/Engine/Plugins/Enterprise/DatasmithCADImporter/Source/DatasmithCoreTechExtension/Private/DataprepTessellationOperation.cpp b/Engine/Plugins/Enterprise/DatasmithCADImporter/Source/DatasmithCoreTechExtension/Private/DataprepTessellationOperation.cpp
--- a/Engine/Plugins/Enterprise/DatasmithCADImporter/Source/DatasmithCoreTechExtension/Private/DataprepTessellationOperation.cpp
+++ b/Engine/Plugins/Enterprise/DatasmithCADImporter/Source/DatasmithCoreTechExtension/Private/DataprepTessellationOperation.cpp
@@ -5,12 +5,15 @@
 #ifdef CAD_LIBRARY
 #include "CoreTechHelper.h"
 #include "CoreTechMeshLoader.h"
-#include "CoreTechParametricSurfaceExtension.h"
+#include "CoreTechBlueprintLibrary.h"
 #endif
 #include "DatasmithAdditionalData.h"
 #include "DatasmithAssetImportData.h"
 #include "DatasmithUtils.h"
 
+#include "IDataprepProgressReporter.h"
+
+#include "ActorEditorUtils.h"
 #include "Engine/StaticMesh.h"
 #include "StaticMeshAttributes.h"
 #include "GenericPlatform/GenericPlatformTime.h"
@@ -27,93 +30,76 @@ void UDataprepTessellationOperation::OnExecution_Implementation(const FDataprepC
 	// Collect start time to log amount of time spent to import incoming file
 	uint64 StartTime = FPlatformTime::Cycles64();
 
-	int32 StaticMeshCount = 0;
+	TSet<UStaticMesh*> SelectedMeshes;
 
-	for( UObject* Object : InContext.Objects )
+	for (UObject* Object : InContext.Objects)
 	{
-		if( UStaticMesh* StaticMesh = Cast< UStaticMesh >( Object ))
+		if ( UStaticMesh* StaticMesh = Cast<UStaticMesh>(Object) )
+		{
+			SelectedMeshes.Add( StaticMesh );
+		}
+		else if (AActor* Actor = Cast<AActor>(Object) )
 		{
-			if( !StaticMesh->IsMeshDescriptionValid( 0 ) )
+			TInlineComponentArray<UStaticMeshComponent*> StaticMeshComponents( Actor );
+			for (UStaticMeshComponent* StaticMeshComponent : StaticMeshComponents)
 			{
-				FText WarningMsg = FText::Format( LOCTEXT( "DataprepTessellationOperation_EmptyMesh", "No triangles in static mesh {0}" ), FText::FromString( StaticMesh->GetName() ) );
-				LogWarning( WarningMsg );
-				// #ueent_todo: Remove this log when Dataprep logging is operational
-				UE_LOG( LogCADLibrary, Log, TEXT("Retessellation failed : %s "), *WarningMsg.ToString() );
-				continue;
+				if((StaticMesh = StaticMeshComponent->GetStaticMesh()) != nullptr)
+				{
+					SelectedMeshes.Add( StaticMesh );
+				}
 			}
+		}
+	}
 
-			// Borrowed from UCoreTechRetessellateAction class
-			// #ueent_todo: Create CAD blueprint library to expose this feature
-			FAssetData AssetData( StaticMesh );
-			if( UCoreTechParametricSurfaceData* CoreTechData = Datasmith::GetAdditionalData<UCoreTechParametricSurfaceData>(AssetData) )
-			{
-				FString ResourceFile = CoreTechData->SourceFile.IsEmpty() ? FPaths::ProjectIntermediateDir() / "temp.ct" : CoreTechData->SourceFile;
-				FFileHelper::SaveArrayToFile(CoreTechData->RawData, *ResourceFile);
-
-				CADLibrary::CoreTechMeshLoader Loader;
-
-				CADLibrary::FImportParameters Parameters;
-				Parameters.MetricUnit = CoreTechData->SceneParameters.MetricUnit;
-				Parameters.ChordTolerance = TessellationSettings.ChordTolerance;
-				Parameters.MaxEdgeLength = TessellationSettings.MaxEdgeLength;
-				Parameters.MaxNormalAngle = TessellationSettings.NormalTolerance;
-				Parameters.StitchingTechnique = CADLibrary::EStitchingTechnique(TessellationSettings.StitchingTechnique);
+	if(	!IsCancelled() && SelectedMeshes.Num() > 0)
+	{
+		TSharedPtr<FDataprepWorkReporter> Task = CreateTask( LOCTEXT( "LogCADLibrary_Tessellating", "Tessellating meshes ..." ), (float)SelectedMeshes.Num() );
 
-				CADLibrary::FMeshParameters MeshParameters;
-				MeshParameters.ModelCoordSys = FDatasmithUtils::EModelCoordSystem(CoreTechData->SceneParameters.ModelCoordSys);
+		TArray<UObject*> ModifiedStaticMeshes;
+		ModifiedStaticMeshes.Reserve( SelectedMeshes.Num() );
 
-				// Previous MeshDescription is get to be able to create a new one with the same order of PolygonGroup (the matching of color and partition is currently based on their order)
-				FMeshDescription* DestinationMeshDescription = StaticMesh->GetMeshDescription(0);
-				if (DestinationMeshDescription == nullptr)
-				{
-					DestinationMeshDescription = StaticMesh->CreateMeshDescription(0);
-				}
-				FStaticMeshAttributes DestinationMeshAttributes(*DestinationMeshDescription);
+		for( UStaticMesh* StaticMesh : SelectedMeshes )
+		{
+			if( IsCancelled() )
+			{
+				break;
+			}
 
-				// Create MeshDescription from tessellated mesh
-				FMeshDescription MeshDescription;
-				FStaticMeshAttributes MeshDescriptionAttributes(MeshDescription);
-				MeshDescriptionAttributes.Register();
+			Task->ReportNextStep( FText::Format( LOCTEXT( "LogCADLibrary_Tessellating_One_Mesh", "Tessellating {0} ..."), FText::FromString( StaticMesh->GetName() ) ) );
 
-				TPolygonGroupAttributesRef<FName> PolygonGroupDestinationMeshSlotNames = DestinationMeshAttributes.GetPolygonGroupMaterialSlotNames();
-				TPolygonGroupAttributesRef<FName> PolygonGroupImportedMaterialSlotNames = MeshDescriptionAttributes.GetPolygonGroupMaterialSlotNames();
-				for (FPolygonGroupID PolygonGroupID : DestinationMeshDescription->PolygonGroups().GetElementIDs())
+			if( StaticMesh->IsMeshDescriptionValid( 0 ) )
+			{
+				FText OutReason;
+				if( UCoreTechBlueprintLibrary::RetessellateStaticMesh( StaticMesh, TessellationSettings, false, OutReason ) )
 				{
-					FName ImportedSlotName = PolygonGroupDestinationMeshSlotNames[PolygonGroupID];
-					FPolygonGroupID PolyGroupID = MeshDescription.CreatePolygonGroup();
-					PolygonGroupImportedMaterialSlotNames[PolyGroupID] = ImportedSlotName;
+					ModifiedStaticMeshes.Add( StaticMesh );
 				}
-
-				if (!Loader.LoadFile(ResourceFile, MeshDescription, Parameters, MeshParameters))
+				else
 				{
-					UE_LOG(LogCADLibrary, Log, TEXT("Retessellation of %s failed : Cannot generate mesh from parametric surface data"), *StaticMesh->GetName());
-					continue;
+					FText WarningMsg = FText::Format( LOCTEXT( "DataprepTessellationOperation_TessellationFailed", "{0}" ), OutReason );
+					LogWarning( WarningMsg );
 				}
-
-				// Move result of tessellation into static mesh LOD 0
-				*DestinationMeshDescription = MoveTemp(MeshDescription);
-
-				// @TODO: check: no CommitStaticMeshDescription?
-
-				// Save last tessellation settings
-				CoreTechData->LastTessellationOptions = TessellationSettings;
 			}
 			else
 			{
-				UE_LOG( LogCADLibrary, Log, TEXT("Retessellation of %s failed : No tessellation data attached to the static mesh"), *StaticMesh->GetName() );
-				continue;
+				FText WarningMsg = FText::Format( LOCTEXT( "DataprepTessellationOperation_EmptyMesh", "No triangles in static mesh {0}" ), FText::FromString( StaticMesh->GetName() ) );
+				LogWarning( WarningMsg );
 			}
+		}
+
+		// Log time spent to import incoming file in minutes and seconds
+		double ElapsedSeconds = FPlatformTime::ToSeconds64(FPlatformTime::Cycles64() - StartTime);
 
-			++StaticMeshCount;
+		int ElapsedMin = int(ElapsedSeconds / 60.0);
+		ElapsedSeconds -= 60.0 * (double)ElapsedMin;
+		UE_LOG( LogCADLibrary, Log, TEXT("Tessellation of %d out of %d static mesh(es) in [%d min %.3f s]"), ModifiedStaticMeshes.Num(), SelectedMeshes.Num(), ElapsedMin, ElapsedSeconds );
+
+		if(ModifiedStaticMeshes.Num() > 0)
+		{
+			AssetsModified( MoveTemp( ModifiedStaticMeshes ) );
 		}
 	}
 
-	// Log time spent to import incoming file in minutes and seconds
-	double ElapsedSeconds = FPlatformTime::ToSeconds64(FPlatformTime::Cycles64() - StartTime);
-
-	int ElapsedMin = int(ElapsedSeconds / 60.0);
-	ElapsedSeconds -= 60.0 * (double)ElapsedMin;
-	UE_LOG( LogCADLibrary, Log, TEXT("Tessellation of %d static mesh(es) in [%d min %.3f s]"), StaticMeshCount, ElapsedMin, ElapsedSeconds );
 #else
 	UE_LOG( LogCADLibrary, Warning, TEXT("Tessellation not performed") );
 #endif
diff --git a/Engine/Plugins/Enterprise/DatasmithCADImporter/Source/DatasmithCoreTechExtension/Public/CoreTechBlueprintLibrary.h b/Engine/Plugins/Enterprise/DatasmithCADImporter/Source/DatasmithCoreTechExtension/Public/CoreTechBlueprintLibrary.h
new file mode 100644
--- /dev/null
+++ b/Engine/Plugins/Enterprise/DatasmithCADImporter/Source/DatasmithCoreTechExtension/Public/CoreTechBlueprintLibrary.h
@@ -0,0 +1,32 @@
+// Copyright 1998-2019 Epic Games, Inc. All Rights Reserved.
+
+#pragma once
+
+#include "CoreMinimal.h"
+
+#include "DatasmithImportOptions.h"
+
+#include "Kismet/BlueprintFunctionLibrary.h"
+#include "UObject/TextProperty.h"
+
+#include "CoreTechBlueprintLibrary.generated.h"
+
+class UStaticMesh;
+
+UCLASS(Blueprintable, BlueprintType, meta = (DisplayName = "CAD Surface Operations Library"))
+class DATASMITHCORETECHEXTENSION_API UCoreTechBlueprintLibrary : public UBlueprintFunctionLibrary
+{
+	GENERATED_BODY()
+
+public:
+	/**
+	 * Re-tessellate LOD 0 of a static mesh if it contains parametric surface data.
+	 * @param	StaticMesh				Static mesh to re-tessellate.
+	 * @param	TessellationSettings	Tessellation settings to use.
+	 * @param	bPostChanges			If true, .
+	 * @param	OutReason				Text describing the reason of failure.
+	 * @return True if successful, false otherwise
+	 */
+	UFUNCTION(BlueprintPure, Category = "Datasmith | Surface Operations" )
+	static bool RetessellateStaticMesh( UStaticMesh* StaticMesh, const FDatasmithTessellationOptions& TessellationSettings, bool bPostChanges, FText& OutReason );
+};
diff --git a/Engine/Plugins/Enterprise/DatasmithCADImporter/Source/DatasmithCoreTechExtension/Public/CoreTechParametricSurfaceExtension.h b/Engine/Plugins/Enterprise/DatasmithCADImporter/Source/DatasmithCoreTechExtension/Public/CoreTechParametricSurfaceExtension.h
--- a/Engine/Plugins/Enterprise/DatasmithCADImporter/Source/DatasmithCoreTechExtension/Public/CoreTechParametricSurfaceExtension.h
+++ b/Engine/Plugins/Enterprise/DatasmithCADImporter/Source/DatasmithCoreTechExtension/Public/CoreTechParametricSurfaceExtension.h
@@ -77,9 +77,7 @@ public:
 
 	static bool CanApplyOnAssets(const TArray<FAssetData>& SelectedAssets);
 	static void ApplyOnAssets(const TArray<FAssetData>& SelectedAssets);
-
-private:
-	static void ApplyOnOneAsset(UStaticMesh& StaticMesh, UCoreTechParametricSurfaceData& CoreTechData, UCoreTechRetessellateActionOptions& RetesselateOptions);
+	static bool ApplyOnOneAsset(UStaticMesh& StaticMesh, UCoreTechParametricSurfaceData& CoreTechData, const FDatasmithTessellationOptions& RetesselateOptions);
 };
 
 
