Fix missing undef LOCTEXT_NAMESPACE
#rnx

diff --git a/Engine/Plugins/Enterprise/DatasmithCADImporter/Source/DatasmithCoreTechExtension/Private/CoreTechBlueprintLibrary.cpp b/Engine/Plugins/Enterprise/DatasmithCADImporter/Source/DatasmithCoreTechExtension/Private/CoreTechBlueprintLibrary.cpp
--- a/Engine/Plugins/Enterprise/DatasmithCADImporter/Source/DatasmithCoreTechExtension/Private/CoreTechBlueprintLibrary.cpp
+++ b/Engine/Plugins/Enterprise/DatasmithCADImporter/Source/DatasmithCoreTechExtension/Private/CoreTechBlueprintLibrary.cpp
@@ -103,4 +103,6 @@ bool UCoreTechBlueprintLibrary::RetessellateStaticMesh( UStaticMesh* StaticMesh,
 #endif
 
 	return bTessellationOutcome;
-}
\ No newline at end of file
+}
+
+#undef LOCTEXT_NAMESPACE
\ No newline at end of file
