Merging CL #9539074 , #9542922, #9543233, #9547054 and #9547513 from //UE4/Dev-Enterprise

Fix WireTranslator compilation
OpenModel ThirdParty

diff --git a/Engine/Plugins/Enterprise/DatasmithCADImporter/Source/DatasmithWireTranslator/DatasmithWireTranslator.Build.cs b/Engine/Plugins/Enterprise/DatasmithCADImporter/Source/DatasmithWireTranslator/DatasmithWireTranslator.Build.cs
--- a/Engine/Plugins/Enterprise/DatasmithCADImporter/Source/DatasmithWireTranslator/DatasmithWireTranslator.Build.cs
+++ b/Engine/Plugins/Enterprise/DatasmithCADImporter/Source/DatasmithWireTranslator/DatasmithWireTranslator.Build.cs
@@ -18,15 +18,16 @@ public class DatasmithWireTranslator : ModuleRules
 		PrivateDependencyModuleNames.AddRange(
 			new string[]
 			{
-				"Engine",
+				"CADLibrary",
+				"DatasmithContent",
 				"DatasmithCore",
 				"DatasmithCoreTechExtension",
-				"DatasmithContent",
 				"DatasmithImporter",
+				"Engine",
 				"MeshDescription",
-				"CADLibrary",
-			}
-		);
+                "StaticMeshDescription",
+            }
+        );
 
 		if (System.Type.GetType("OpenModel") != null)
 		{
diff --git a/Engine/Plugins/Enterprise/DatasmithCADImporter/Source/DatasmithWireTranslator/Private/AliasCoretechWrapper.cpp b/Engine/Plugins/Enterprise/DatasmithCADImporter/Source/DatasmithWireTranslator/Private/AliasCoretechWrapper.cpp
--- a/Engine/Plugins/Enterprise/DatasmithCADImporter/Source/DatasmithWireTranslator/Private/AliasCoretechWrapper.cpp
+++ b/Engine/Plugins/Enterprise/DatasmithCADImporter/Source/DatasmithWireTranslator/Private/AliasCoretechWrapper.cpp
@@ -156,7 +156,7 @@ CT_OBJECT_ID FAliasCoretechWrapper::Add3DCurve(AlCurve& Curve)
 
 CT_OBJECT_ID FAliasCoretechWrapper::AddTrimCurve(AlTrimCurve& TrimCurve)
 {
-	bool bIsReversed = TrimCurve.isReversed();
+	boolean bIsReversed = TrimCurve.isReversed();
 	curveFormType Form = TrimCurve.form();
 	CT_UINT32 Order = TrimCurve.degree() + 1;
 	CT_UINT32 ControlPointSize = TrimCurve.numberOfCVs();
diff --git a/Engine/Plugins/Enterprise/DatasmithCADImporter/Source/DatasmithWireTranslator/Private/OpenModelUtils.cpp b/Engine/Plugins/Enterprise/DatasmithCADImporter/Source/DatasmithWireTranslator/Private/OpenModelUtils.cpp
--- a/Engine/Plugins/Enterprise/DatasmithCADImporter/Source/DatasmithWireTranslator/Private/OpenModelUtils.cpp
+++ b/Engine/Plugins/Enterprise/DatasmithCADImporter/Source/DatasmithWireTranslator/Private/OpenModelUtils.cpp
@@ -15,8 +15,9 @@
 #include "AlTesselate.h"
 #include "MeshAttributes.h"
 #include "MeshDescription.h"
+#include "StaticMeshAttributes.h"
+#include "StaticMeshOperations.h"
 #include "Utility/DatasmithMathUtils.h"
-
 using namespace OpenModelUtils;
 
 
@@ -1659,10 +1660,13 @@ bool OpenModelUtils::TransferAlMeshToMeshDescription(const AlMesh& AliasMesh, FM
 		SymmetricMatrix = FDatasmithUtils::GetSymmetricMatrix(MeshParameters.SymmetricOrigin, MeshParameters.SymmetricNormal);
 	}
 
+	// Gather all array data
+	FStaticMeshAttributes Attributes(MeshDescription);
+	TVertexInstanceAttributesRef<FVector> VertexInstanceNormals = Attributes.GetVertexInstanceNormals();
+	TVertexInstanceAttributesRef<FVector2D> VertexInstanceUVs = Attributes.GetVertexInstanceUVs();
+	TPolygonGroupAttributesRef<FName> PolygonGroupImportedMaterialSlotNames = Attributes.GetPolygonGroupMaterialSlotNames();
+
 	TVertexAttributesRef<FVector> VertexPositions = MeshDescription.VertexAttributes().GetAttributesRef<FVector>(MeshAttribute::Vertex::Position);
-	TVertexInstanceAttributesRef<FVector> VertexInstanceNormals = MeshDescription.VertexInstanceAttributes().GetAttributesRef<FVector>(MeshAttribute::VertexInstance::Normal);
-	TVertexInstanceAttributesRef<FVector2D> VertexInstanceUVs = MeshDescription.VertexInstanceAttributes().GetAttributesRef<FVector2D>(MeshAttribute::VertexInstance::TextureCoordinate);
-	TPolygonGroupAttributesRef<FName> PolygonGroupImportedMaterialSlotNames = MeshDescription.PolygonGroupAttributes().GetAttributesRef<FName>(MeshAttribute::PolygonGroup::ImportedMaterialSlotName);
 
 	// Prepared for static mesh usage ?
 	if (!VertexPositions.IsValid() || !VertexInstanceNormals.IsValid() || !VertexInstanceUVs.IsValid() || !PolygonGroupImportedMaterialSlotNames.IsValid())
@@ -1823,10 +1827,7 @@ bool OpenModelUtils::TransferAlMeshToMeshDescription(const AlMesh& AliasMesh, FM
 	}
 
 	// Build edge meta data
-	MeshDescription.DetermineEdgeHardnessesFromVertexInstanceNormals();
-
-	MeshDescription.EdgeAttributes().RegisterAttribute<bool>(MeshAttribute::Edge::IsUVSeam, 1, false);
-	MeshDescription.DetermineUVSeamsFromUVs(0);
+	FStaticMeshOperations::DetermineEdgeHardnessesFromVertexInstanceNormals(MeshDescription);
 
 	return true;
 }
