Merging CL #9598147 and #9612832 from //UE4/Dev-Enterprise

diff --git a/Engine/Plugins/Enterprise/DatasmithCADImporter/Source/DatasmithWireTranslator/Private/DatasmithWireTranslator.cpp b/Engine/Plugins/Enterprise/DatasmithCADImporter/Source/DatasmithWireTranslator/Private/DatasmithWireTranslator.cpp
--- a/Engine/Plugins/Enterprise/DatasmithCADImporter/Source/DatasmithWireTranslator/Private/DatasmithWireTranslator.cpp
+++ b/Engine/Plugins/Enterprise/DatasmithCADImporter/Source/DatasmithWireTranslator/Private/DatasmithWireTranslator.cpp
@@ -1257,9 +1257,6 @@ void FWireTranslatorImpl::GetDagNodeInfo(TSharedRef<BodyData> CurrentNode, const
 	CurrentNodeInfo.Label = ParentInfo.Label; // +TEXT(" ") + CurrentNode->LayerName + TEXT(" ") + CurrentNode->ShaderName;
 	CurrentNode->Label = CurrentNodeInfo.Label;
 
-	FString ThisGroupNodeID;
-	AlPersistentID* GroupNodeId = new AlPersistentID();
-
 	// Limit length of UUID by combining hash of parent UUID and container's UUID if ParentUuid is not empty
 	CurrentNodeInfo.UEuuid = GetUEUUIDFromAIPersistentID(ParentInfo.UEuuid, CurrentNodeInfo.Label);
 }
