Merging CL #9614944 and #9616277 from //UE4/Dev-Enterprise

Wire: Fix wrong metric and some material issues

diff --git a/Engine/Plugins/Enterprise/DatasmithCADImporter/Source/DatasmithWireTranslator/Private/AliasCoretechWrapper.h b/Engine/Plugins/Enterprise/DatasmithCADImporter/Source/DatasmithWireTranslator/Private/AliasCoretechWrapper.h
--- a/Engine/Plugins/Enterprise/DatasmithCADImporter/Source/DatasmithWireTranslator/Private/AliasCoretechWrapper.h
+++ b/Engine/Plugins/Enterprise/DatasmithCADImporter/Source/DatasmithWireTranslator/Private/AliasCoretechWrapper.h
@@ -32,7 +32,7 @@ public:
 	 * eg. For a file in inches, arg should be 0.0254
 	 */
 	FAliasCoretechWrapper(const TCHAR* InOwner)
-		: CTSession(InOwner, 0.001, 1)
+		: CTSession(InOwner, 0.01, 1) // Alias files have cm units
 	{
 	}
 
