Wire: Fix NURBs geometry not building due to too small scale #jira UE-81753 #rb JeanMichel.Dignard, David.Lesage

diff --git a/Engine/Plugins/Enterprise/DatasmithCADImporter/Source/DatasmithWireTranslator/Private/AliasCoretechWrapper.h b/Engine/Plugins/Enterprise/DatasmithCADImporter/Source/DatasmithWireTranslator/Private/AliasCoretechWrapper.h
--- a/Engine/Plugins/Enterprise/DatasmithCADImporter/Source/DatasmithWireTranslator/Private/AliasCoretechWrapper.h
+++ b/Engine/Plugins/Enterprise/DatasmithCADImporter/Source/DatasmithWireTranslator/Private/AliasCoretechWrapper.h
@@ -5,9 +5,13 @@
 #ifdef CAD_LIBRARY
 #include "CoreMinimal.h"
 
-
 #include "CTSession.h"
 
+// Additional scale factor required when tessellating NURBS as Alias has extremely small geometry,
+// originally tessellating to triangles with area in the order of 10^-5, failing the
+// FourSquaredTriangleArea > SMALL_NUMBER test in DatasmithMeshHelper.cpp::IsMeshValid
+#define ALIAS_BUILD_SCALE 100
+
 class AlDagNode;
 class AlShell;
 class AlSurface;
@@ -32,7 +36,7 @@ public:
 	 * eg. For a file in inches, arg should be 0.0254
 	 */
 	FAliasCoretechWrapper(const TCHAR* InOwner)
-		: CTSession(InOwner, 0.01, 1) // Alias files have cm units
+		: CTSession(InOwner, 0.01, ALIAS_BUILD_SCALE) // Alias files are in cm units
 	{
 	}
 
diff --git a/Engine/Plugins/Enterprise/DatasmithCADImporter/Source/DatasmithWireTranslator/Private/DatasmithWireTranslator.cpp b/Engine/Plugins/Enterprise/DatasmithCADImporter/Source/DatasmithWireTranslator/Private/DatasmithWireTranslator.cpp
--- a/Engine/Plugins/Enterprise/DatasmithCADImporter/Source/DatasmithWireTranslator/Private/DatasmithWireTranslator.cpp
+++ b/Engine/Plugins/Enterprise/DatasmithCADImporter/Source/DatasmithWireTranslator/Private/DatasmithWireTranslator.cpp
@@ -1388,6 +1388,15 @@ bool FWireTranslatorImpl::ProcessAlShellNode(AlDagNode& ShellNode, const FDagNod
 
 	SetActorTransform(ShellInfo.ActorElement, ShellNode);
 
+	// If the mesh is NURBS, it will be built with an additional ALIAS_BUILD_SCALE (see AliasCoretechWrapper constructor),
+	// so we need to compensate the same scale on the actor. Alias doesn't seem to allow parenting geometry
+	// to other geometry, so we don't have to handle consequences on any potential children of this actor by doing this
+	AlObjectType Type = ShellNode.type();
+	if (Type == AlObjectType::kSurfaceNodeType || Type == AlObjectType::kShellNodeType)
+	{
+		ActorElement->SetScale(ActorElement->GetScale() / ALIAS_BUILD_SCALE);
+	}
+
 	//// Apply materials on the current part
 	if (ShaderName)
 	{
diff --git a/Engine/Plugins/Enterprise/DatasmithCADImporter/Source/DatasmithWireTranslator/Private/DatasmithWireTranslator.h b/Engine/Plugins/Enterprise/DatasmithCADImporter/Source/DatasmithWireTranslator/Private/DatasmithWireTranslator.h
--- a/Engine/Plugins/Enterprise/DatasmithCADImporter/Source/DatasmithWireTranslator/Private/DatasmithWireTranslator.h
+++ b/Engine/Plugins/Enterprise/DatasmithCADImporter/Source/DatasmithWireTranslator/Private/DatasmithWireTranslator.h
@@ -47,7 +47,4 @@ protected:
 
 private:
 	TSharedPtr<FWireTranslatorImpl> Translator;
-
-	const float MetricUnit = 0.001; // As CT don`t take care of unit input data, metric unit is set with default value
-	const float ScaleFactor = 1; // Wire internal unit is cm so no scale needed
 };
