Merging c9606593 //UE4/Dev-Enterprise to Release-4.24 (//UE4/Release-4.24)

diff --git a/Engine/Plugins/Enterprise/DatasmithCADImporter/Source/DatasmithDispatcher/Private/DatasmithCommands.cpp b/Engine/Plugins/Enterprise/DatasmithCADImporter/Source/DatasmithDispatcher/Private/DatasmithCommands.cpp
new file mode 100644
--- /dev/null
+++ b/Engine/Plugins/Enterprise/DatasmithCADImporter/Source/DatasmithDispatcher/Private/DatasmithCommands.cpp
@@ -0,0 +1,170 @@
+// Copyright 1998-2019 Epic Games, Inc. All Rights Reserved.
+#include "DatasmithCommands.h"
+
+#include "DatasmithDispatcherSocket.h"
+#include "DatasmithDispatcher.h"
+
+namespace DatasmithDispatcher
+{
+	FDatasmithCommandManager::FDatasmithCommandManager(FDatasmithDispatcherSocket& InSocket)
+		: Socket(InSocket)
+		, CurrentCommandeSize(0)
+	{
+		Initialize();
+	}
+
+	FDatasmithCommandManager::~FDatasmithCommandManager()
+	{
+		Terminate();
+	}
+
+	void FDatasmithCommandManager::Initialize()
+	{
+		CommandMap.Reserve((int32)DatasmithCommandType::Last);
+
+		CommandMap.Add(Ping, new FDatasmithPingCommand());
+		CommandMap.Add(BackPing, new FDatasmithBackPingCommand());
+		CommandMap.Add(RunTask, new FDatasmithRunTaskCommand());
+		CommandMap.Add(NotifyEndTask, new FDatasmithNotifyEndTaskCommand());
+		CommandMap.Add(ImportParams, new FDatasmithImportParametersCommand());
+	}
+
+	void FDatasmithCommandManager::Terminate()
+	{
+		for (auto ItemCommand : CommandMap)
+		{
+			delete ItemCommand.Value;
+		}
+		CommandMap.Empty();
+	}
+
+	ICommand * FDatasmithCommandManager::GetNextCommand()
+	{
+		uint32 DataSize;
+		if (!Socket.HasPendingData(DataSize))
+		{
+			return nullptr;
+		}
+
+		if (!CurrentCommandeSize)
+		{
+			Socket >> CurrentCommandeSize;
+		}
+
+		if (!Socket.HasPendingData(DataSize))
+		{
+			return nullptr;
+		}
+
+		// All data is not yet arrived, wait
+		if (DataSize < CurrentCommandeSize)
+		{
+			return nullptr;
+		}
+
+		DatasmithCommandType CommandType;
+		Socket >> CommandType;
+
+		ICommand** Command = CommandMap.Find(CommandType);
+		if (Command == nullptr)
+		{
+			return nullptr;
+		}
+		 
+		(*Command)->Initialize();
+		(*Command)->Read(Socket);
+
+		CurrentCommandeSize = 0;
+
+		return (*Command);
+	}
+
+	void FDatasmithBackPingCommand::Write(FDatasmithDispatcherSocket &CurrentSocket)
+	{
+		CurrentSocket << GetType();
+		CurrentSocket.SendData();
+	}
+
+	void FDatasmithBackPingCommand::Read(FDatasmithDispatcherSocket &CurrentSocket)
+	{
+	}
+
+	void FDatasmithPingCommand::Write(FDatasmithDispatcherSocket &CurrentSocket)
+	{
+		CurrentSocket << GetType();
+		CurrentSocket.SendData();
+	}
+
+	void FDatasmithPingCommand::Read(FDatasmithDispatcherSocket &CurrentSocket)
+	{
+	}
+
+	void FDatasmithRunTaskCommand::Write(FDatasmithDispatcherSocket&  CurrentSocket)
+	{
+		CurrentSocket << GetType();
+		CurrentSocket << JobFilePath;
+		CurrentSocket << JobIndex;					
+		CurrentSocket.SendData();
+	}
+
+	void FDatasmithRunTaskCommand::Read(FDatasmithDispatcherSocket& CurrentSocket)
+	{
+		CurrentSocket >> JobFilePath;
+		CurrentSocket >> JobIndex;
+	}
+
+  	void FDatasmithNotifyEndTaskCommand::Write(FDatasmithDispatcherSocket& CurrentSocket)
+	{
+		CurrentSocket << GetType();
+
+		CurrentSocket << (uint32)ExternalReferenceSet.Num();
+		for (const FString& ExternalReference : ExternalReferenceSet)
+		{
+			CurrentSocket << ExternalReference;
+		}
+		CurrentSocket << ProcessResult;
+		CurrentSocket << SceneGraphFileName;
+		CurrentSocket << GeomFileName;
+		CurrentSocket.SendData();
+	}
+
+	void FDatasmithNotifyEndTaskCommand::Read(FDatasmithDispatcherSocket& CurrentSocket)
+	{
+		uint32 NbReferences;
+
+		CurrentSocket >> NbReferences;
+		FString Buffer;
+		for (uint32 ite = 0; ite < NbReferences; ++ite)
+		{
+			CurrentSocket >> Buffer;
+			ExternalReferenceSet.Add(Buffer);
+		}
+		CurrentSocket >> ProcessResult;
+		CurrentSocket >> SceneGraphFileName;
+		CurrentSocket >> GeomFileName;
+	}
+
+
+	void FDatasmithImportParametersCommand::Write(FDatasmithDispatcherSocket&  CurrentSocket)
+	{
+		CurrentSocket << GetType();
+		CurrentSocket << ImportParameters.ChordTolerance;
+		CurrentSocket << ImportParameters.MaxEdgeLength;
+		CurrentSocket << ImportParameters.MaxNormalAngle;
+		CurrentSocket << ImportParameters.MetricUnit;
+		CurrentSocket << ImportParameters.ScaleFactor;
+		CurrentSocket << ImportParameters.StitchingTechnique;
+		CurrentSocket.SendData();
+	}
+
+	void FDatasmithImportParametersCommand::Read(FDatasmithDispatcherSocket&  CurrentSocket)
+	{
+		CurrentSocket >> ImportParameters.ChordTolerance;
+		CurrentSocket >> ImportParameters.MaxEdgeLength;
+		CurrentSocket >> ImportParameters.MaxNormalAngle;
+		CurrentSocket >> ImportParameters.MetricUnit;
+		CurrentSocket >> ImportParameters.ScaleFactor;
+		CurrentSocket >> ImportParameters.StitchingTechnique;
+	}
+
+}
\ No newline at end of file
diff --git a/Engine/Plugins/Enterprise/DatasmithCADImporter/Source/DatasmithDispatcher/Private/DatasmithDispatcher.cpp b/Engine/Plugins/Enterprise/DatasmithCADImporter/Source/DatasmithDispatcher/Private/DatasmithDispatcher.cpp
new file mode 100644
--- /dev/null
+++ b/Engine/Plugins/Enterprise/DatasmithCADImporter/Source/DatasmithDispatcher/Private/DatasmithDispatcher.cpp
@@ -0,0 +1,141 @@
+// Copyright 1998-2019 Epic Games, Inc. All Rights Reserved.
+#include "DatasmithDispatcher.h"
+
+#include "Misc/ScopeLock.h"
+#include "HAL/Thread.h"
+#include "Containers/Array.h"
+#include "Misc/Optional.h"
+
+namespace DatasmithDispatcher
+{
+	FDatasmithDispatcher::FDatasmithDispatcher(const TCHAR* InCacheDir, int32 InNumberOfWorkers, TMap<FString, FString>& OutCADFileToUnrealFileMap, TMap<FString, FString>& OutCADFileToUnrealGeomMap)
+		: ProcessCacheFolder(InCacheDir)
+		, CADFileToUnrealFileMap(OutCADFileToUnrealFileMap)
+		, CADFileToUnrealGeomMap(OutCADFileToUnrealGeomMap)
+		, NumberOfWorkers(InNumberOfWorkers)
+		, LastLaunchTask(0)
+		, NumOfFinishedTask(0)
+ 	{
+		WorkerHandlerSet.Reserve(NumberOfWorkers);
+		for (int32 Index = 0; Index < NumberOfWorkers; Index++)
+		{
+			WorkerHandlerSet.Emplace(*this, ProcessCacheFolder);
+		}
+	}
+
+	FDatasmithDispatcher::~FDatasmithDispatcher()
+	{
+	}
+
+
+	void FDatasmithDispatcher::RunHandlers()
+	{
+		WorkersRunning.Reserve(NumberOfWorkers);
+		int32 WorkerIndex = 0;
+		for (FDatasmithWorkerHandler& Handler : WorkerHandlerSet)
+		{
+			FString ProcessName = TEXT("DatasmithWorker ") + FString::FromInt(WorkerIndex);
+			WorkerIndex++;
+			WorkersRunning.Emplace(*ProcessName, [&]() { Handler.Run(); });
+		}
+	}
+
+	void FDatasmithDispatcher::Close()
+	{
+		for (FDatasmithWorkerHandler& Handler : WorkerHandlerSet)
+		{
+			Handler.NotifyKill();
+		}
+
+		for (FThread& Process : WorkersRunning)
+		{
+			Process.Join();
+		}
+	}
+
+	void FDatasmithDispatcher::Clear()
+	{
+		TaskPool.Empty();
+		LastLaunchTask = 0;
+		NumOfFinishedTask = 0;
+	}
+
+
+	void FDatasmithDispatcher::Process()
+	{
+		RunHandlers();
+
+		while (!IsOver())
+		{
+			FWindowsPlatformProcess::Sleep(0.1);
+		} 
+
+		Close();
+	}
+
+	void FDatasmithDispatcher::AddTask(const FString& InFile)
+	{
+		FScopeLock Lock(&TaskPoolCriticalSection);
+		for (const FTask& Task : TaskPool)
+		{
+			if (Task.FileName == InFile)
+			{
+				return;
+			}
+		}
+
+		int32 TaskIndex = TaskPool.Emplace(InFile);
+		TaskPool[TaskIndex].Index = TaskIndex;
+	}
+
+	TOptional<FTask> FDatasmithDispatcher::GetTask()
+	{
+		FScopeLock Lock(&TaskPoolCriticalSection);
+
+		if (!TaskPool.IsValidIndex(LastLaunchTask))
+		{
+			return TOptional<FTask>();
+		}
+
+		TaskPool[LastLaunchTask].State = Running;
+		return TaskPool[LastLaunchTask++];
+	}
+
+	bool FDatasmithDispatcher::IsOver() 
+	{
+		FScopeLock Lock(&TaskPoolCriticalSection);
+		return NumOfFinishedTask == TaskPool.Num();
+	}
+
+	void FDatasmithDispatcher::SetTaskState(int32 TaskIndex, EProcessState TaskState) 
+	{
+		FScopeLock Lock(&TaskPoolCriticalSection);
+
+		if (TaskIndex < 0)
+		{
+			return;
+		}
+		if (TaskIndex < TaskPool.Num())
+		{
+			TaskPool[TaskIndex].State = TaskState;
+			if (TaskState & (ProcessOk | ProcessFailed | FileNotFound))
+			{
+				NumOfFinishedTask++;
+			}
+		}
+	}
+
+	void FDatasmithDispatcher::LinkCTFileToUnrealCacheFile(const FString& CTFile, const FString& UnrealSceneGraphFile, const FString& UnrealMeshFile)
+	{
+		FScopeLock Lock(&TaskPoolCriticalSection);
+		if (!UnrealSceneGraphFile.IsEmpty())
+		{
+			CADFileToUnrealFileMap.Add(CTFile, UnrealSceneGraphFile);
+		}
+		if (!UnrealMeshFile.IsEmpty())
+		{
+			CADFileToUnrealGeomMap.Add(CTFile, UnrealMeshFile);
+		}
+
+	}
+}
diff --git a/Engine/Plugins/Enterprise/DatasmithCADImporter/Source/DatasmithDispatcher/Private/DatasmithDispatcherModule.cpp b/Engine/Plugins/Enterprise/DatasmithCADImporter/Source/DatasmithDispatcher/Private/DatasmithDispatcherModule.cpp
new file mode 100644
--- /dev/null
+++ b/Engine/Plugins/Enterprise/DatasmithCADImporter/Source/DatasmithDispatcher/Private/DatasmithDispatcherModule.cpp
@@ -0,0 +1,29 @@
+// Copyright 1998-2019 Epic Games, Inc. All Rights Reserved.
+#include "DatasmithDispatcherModule.h"
+
+#include "Modules/ModuleInterface.h"
+#include "Modules/ModuleManager.h"
+
+
+
+#define LOCTEXT_NAMESPACE "FDatasmithDispatcherModule"
+
+
+FDatasmithDispatcherModule& FDatasmithDispatcherModule::Get()
+{
+	return FModuleManager::LoadModuleChecked< FDatasmithDispatcherModule >(DATASMITHDISPATCHER_MODULE_NAME);
+}
+
+bool FDatasmithDispatcherModule::IsAvailable()
+{
+	return FModuleManager::Get().IsModuleLoaded(DATASMITHDISPATCHER_MODULE_NAME);
+}
+
+void FDatasmithDispatcherModule::StartupModule()
+{
+}
+
+IMPLEMENT_MODULE(FDatasmithDispatcherModule, DatasmithDispatcher);
+
+#undef LOCTEXT_NAMESPACE // "DatasmithDispatcherModule"
+
diff --git a/Engine/Plugins/Enterprise/DatasmithCADImporter/Source/DatasmithDispatcher/Private/DatasmithDispatcherSocket.cpp b/Engine/Plugins/Enterprise/DatasmithCADImporter/Source/DatasmithDispatcher/Private/DatasmithDispatcherSocket.cpp
new file mode 100644
--- /dev/null
+++ b/Engine/Plugins/Enterprise/DatasmithCADImporter/Source/DatasmithDispatcher/Private/DatasmithDispatcherSocket.cpp
@@ -0,0 +1,236 @@
+// Copyright 1998-2019 Epic Games, Inc. All Rights Reserved.
+#include "DatasmithDispatcherSocket.h"
+
+#include "CoreMinimal.h"
+
+#include "SocketSubsystem.h"
+#include "Sockets.h"
+
+using namespace DatasmithDispatcher;
+
+DEFINE_LOG_CATEGORY_STATIC(LogDatasmithCoretechWorker, Log, All);
+
+
+
+FDatasmithDispatcherSocket::FDatasmithDispatcherSocket(const TCHAR* ServerAddress)
+	: SocketAdress(ServerAddress)
+	, SocketId(-1)
+	, Socket(nullptr)
+	, bServerSide(true)
+	, bOpen(false)
+	, ErrorCode(0)
+{
+	Cache.Reserve(1024 * 128);
+}
+
+FDatasmithDispatcherSocket::FDatasmithDispatcherSocket()
+	: SocketId(-1)
+	, Socket(nullptr)
+	, bServerSide(true)
+	, bOpen(false)
+	, ErrorCode(0)
+{
+	Cache.Reserve(1024 * 128);
+}
+
+void FDatasmithDispatcherSocket::SetSocket(FSocket* InSocket)
+{
+	Socket = InSocket;
+	bOpen = true;
+}
+
+bool FDatasmithDispatcherSocket::Listen()
+{
+	return Socket->Listen(1);
+}
+
+FSocket* FDatasmithDispatcherSocket::Accept()
+{
+	FString Message = TEXT("DispatcherSocket");
+	return Socket->Accept(Message);
+}
+
+#include "CoreGlobals.h"
+
+void FDatasmithDispatcherSocket::Bind()
+{
+	ISocketSubsystem* SocketSubsystem = ISocketSubsystem::Get();
+	if (nullptr == SocketSubsystem)
+	{
+		return;
+	}
+
+	TSharedPtr<FInternetAddr> InternetAddress = SocketSubsystem->CreateInternetAddr();
+
+	bool bIsValid = true;
+	InternetAddress->SetIp(*SocketAdress, bIsValid);
+	if (!bIsValid)
+	{
+		return;
+	}
+
+	Socket = SocketSubsystem->CreateSocket(NAME_Stream, TEXT("FNetworkFileServer tcp-listen"), InternetAddress->GetProtocolType());
+	if (nullptr == Socket)
+	{
+		return;
+	}
+
+	int32 SocketPort = SocketSubsystem->BindNextPort(Socket, *InternetAddress, 5000, 1);
+	if (SocketPort == 0)
+	{
+		return;
+	}
+
+	bOpen = true;
+}
+
+bool FDatasmithDispatcherSocket::IsConnected()
+{
+	uint32 DataSize;
+	Socket->HasPendingData(DataSize);
+
+	bool bHasPendingConnection;
+	Socket->HasPendingConnection(bHasPendingConnection);
+	return !bHasPendingConnection || DataSize;
+}
+
+bool FDatasmithDispatcherSocket::HasPendingData(uint32& DataSize)
+{
+	return Socket->HasPendingData(DataSize);
+}
+
+
+void FDatasmithDispatcherSocket::Connect(const int32& InServerPort)
+{
+	bServerSide = false;
+	bOpen = false;
+
+	ISocketSubsystem* SocketSubsystem = ISocketSubsystem::Get();
+	if (nullptr == SocketSubsystem)
+	{
+		return;
+	}
+
+	TSharedPtr<FInternetAddr> InternetAddress = SocketSubsystem->CreateInternetAddr();
+	bool bIsValid = true;
+	InternetAddress->SetIp(*SocketAdress, bIsValid);
+	if (!bIsValid)
+	{
+		return;
+	}
+	InternetAddress->SetPort(InServerPort);
+
+	Socket = SocketSubsystem->CreateSocket(NAME_Stream, TEXT("FNetworkFileServer tcp-listen"), InternetAddress->GetProtocolType());
+	if (nullptr == Socket)
+	{
+		return;
+	}
+	Socket->SetNonBlocking(true);
+
+	if (!Socket->Connect(*InternetAddress))
+	{
+		return;
+	}
+	UE_LOG(LogDatasmithCoretechWorker, Display, TEXT("Is connected"));
+
+	bool bHasPendingConnection = true;
+	Socket->WaitForPendingConnection(bHasPendingConnection, FTimespan(0, 0, 5));
+
+	UE_LOG(LogDatasmithCoretechWorker, Display, TEXT("WaitForPendingConnection return = %d"), bHasPendingConnection);
+	bOpen = !bHasPendingConnection;
+}
+
+void FDatasmithDispatcherSocket::Close()
+{
+	if (Socket)
+	{
+		Socket->Close();
+		ISocketSubsystem* SocketSubsystem = ISocketSubsystem::Get();
+		SocketSubsystem->DestroySocket(Socket);
+		Socket = nullptr;
+	}
+}
+
+
+void FDatasmithDispatcherSocket::Read(int32 DataSize, uint8* Data)
+{
+	if (!bOpen) return;
+
+	int32 Total = 0;
+	memset(Data, 0, DataSize);
+	int32 BytesRead;
+
+	while (Total < DataSize)
+	{
+		bool bStatus = Socket->Recv(Data + Total, DataSize - Total, BytesRead);
+		if (!bStatus)
+		{
+			UE_LOG(LogDatasmithCoretechWorker, Display, TEXT("Close socket Read %d %d "), DataSize, Total);
+
+			Socket->Close();
+			ErrorCode = UnableToReadOnSocket;
+			bOpen = false;
+		}
+		Total += BytesRead;
+	}
+}
+
+void FDatasmithDispatcherSocket::Write(int32 DataSize, uint8 *Data)
+{
+	Cache.Append(Data, DataSize);
+}
+
+void FDatasmithDispatcherSocket::SendData()
+{
+	if (!bOpen) return;
+	int32 BytesSend;
+
+	uint32 BufferSize = Cache.Num();
+	uint8 *BufferSizeData = (uint8*)&BufferSize;
+	uint32 Total = 0;
+	while (Total < sizeof(uint32))
+	{
+		bool bStatus = Socket->Send(BufferSizeData + Total, sizeof(uint32) - Total, BytesSend);
+		Total += (uint32)BytesSend;
+	}
+
+	uint8 *Data = Cache.GetData();
+	Total = 0;
+	while (Total < BufferSize)
+	{
+		bool bStatus = Socket->Send(Data + Total, BufferSize - Total, BytesSend);
+		if (!bStatus)
+		{
+			Socket->Close();
+			ErrorCode = UnableToSendData;
+			bOpen = false;
+			UE_LOG(LogDatasmithCoretechWorker, Display, TEXT("Close socket Write %d %d "), BufferSize, Total);
+		}
+		Total += (uint32) BytesSend;
+	}
+	Cache.Empty();
+}
+
+void FDatasmithDispatcherSocket::operator >>(FString& Data)
+{
+	int Size;
+	*this >> Size;
+
+	Data.Empty(Size);
+	Data.Reserve(Size);
+	for (int k = 0; k < Size; k++)
+	{
+		TCHAR Char;
+		*this >> Char;
+		Data.AppendChar(Char);
+	}
+}
+
+void FDatasmithDispatcherSocket::operator <<(const FString& Data)
+{
+	*this << (int32)Data.Len();
+	for (int32 k = 0; k < Data.Len(); k++)
+	{
+		*this << Data[k];
+	}
+}
diff --git a/Engine/Plugins/Enterprise/DatasmithCADImporter/Source/DatasmithDispatcher/Private/DatasmithWorkerHandler.cpp b/Engine/Plugins/Enterprise/DatasmithCADImporter/Source/DatasmithDispatcher/Private/DatasmithWorkerHandler.cpp
new file mode 100644
--- /dev/null
+++ b/Engine/Plugins/Enterprise/DatasmithCADImporter/Source/DatasmithDispatcher/Private/DatasmithWorkerHandler.cpp
@@ -0,0 +1,234 @@
+// Copyright 1998-2019 Epic Games, Inc. All Rights Reserved.
+#include "DatasmithWorkerHandler.h"
+
+#include "HAL/FileManager.h"
+#include "HAL/PlatformProcess.h"
+#include "Logging/LogMacros.h"
+#include "Misc/Paths.h"
+#include "SocketSubsystem.h"
+#include "Sockets.h"
+#include "DatasmithDispatcher.h"
+#include "DatasmithCommands.h"
+
+//DEFINE_LOG_CATEGORY_STATIC(LogDatasmithDispatcher, Log, All);
+
+
+namespace DatasmithDispatcher
+{
+	FString FDatasmithWorkerHandler::GetExePath()
+	{
+		FString ProcessorPath;
+#if PLATFORM_MAC
+		ProcessorPath = FPaths::Combine(FPaths::EngineDir(), TEXT("Binaries/Mac/NotForLicensees/DatasmithCoretechWorker"));
+#elif PLATFORM_LINUX
+		ProcessorPath = FPaths::Combine(FPaths::EngineDir(), TEXT("Binaries/Linux/NotForLicensees/DatasmithCoretechWorker"));
+#else
+		ProcessorPath = FPaths::Combine(FPaths::EngineDir(), TEXT("Binaries/Win64/NotForLicensees/DatasmithCoretechWorker.exe"));
+#endif
+
+		if (FPaths::FileExists(ProcessorPath))
+		{
+			return ProcessorPath;
+		}
+		return FString();
+	}
+
+	FDatasmithWorkerHandler::FDatasmithWorkerHandler(FDatasmithDispatcher& InDispatcher, const FString& InCachePath)
+		: Dispatcher(InDispatcher)
+		, WorkerProcessId(0)
+		, ServerSocket(TEXT("127.0.0.1"))
+		, ClientListener()
+		, CachePath(InCachePath)
+		, bShouldTerminate(false)
+		//, Stopwatch(0)
+		//, StartTime(0)
+	{
+	}
+
+	bool FDatasmithWorkerHandler::StartWorker()
+	{
+		FString ProcessorPath = GetExePath();
+		if (!ProcessorPath.Len())
+		{
+			return false;
+		}
+
+		//StartTime = Stopwatch.GetDuration();
+
+		FString WorkerAbsolutePath = IFileManager::Get().ConvertToAbsolutePathForExternalAppForWrite(*ProcessorPath);
+		FPaths::NormalizeDirectoryName(WorkerAbsolutePath);
+
+		FString CommandToProcess = TEXT(" -Listen ") + FString::FromInt(ServerSocket.GetPort()) + TEXT(" -CacheDir \"") + GetCachePath() + TEXT("\"");
+
+		WorkerHandle = FPlatformProcess::CreateProc(*WorkerAbsolutePath, *CommandToProcess, true, false, false, &WorkerProcessId, 0, nullptr, nullptr);
+		//WorkerHandle = FPlatformProcess::CreateProc(*WorkerAbsolutePath, *CommandToProcess, false, false, false, &WorkerProcessId, 0, nullptr, nullptr);
+
+		if (!WorkerHandle.IsValid())
+		{
+			return false;
+		}
+
+		FSocket *ClientSocket = ServerSocket.Accept();
+		if (!ClientSocket)
+		{
+			return false;
+		}
+
+		ClientListener.SetSocket(ClientSocket);
+		return true;
+	}
+
+
+	void FDatasmithWorkerHandler::RestartProcessor()
+	{
+		FPlatformProcess::TerminateProc(WorkerHandle, true);
+		StartWorker();
+	}
+
+	bool FDatasmithWorkerHandler::InitializeSocket()
+	{
+		ServerSocket.Bind();
+		if(!ServerSocket.IsOpen())
+		{
+			return false;
+		}
+		return true; 
+	}
+
+	bool FDatasmithWorkerHandler::Run()
+	{
+
+		if (!InitializeSocket())
+		{
+			return false;
+		}
+
+		if(!ServerSocket.Listen())
+		{
+			return false;
+		}
+
+		if (!StartWorker())
+		{
+			return false;
+		}
+
+		DatasmithDispatcher::FDatasmithCommandManager CommandManager(ClientListener);
+
+		int test = 0;
+		bool bTaskIsRunning = true;
+
+		while (NeedToRun())
+		{
+
+			if (!ClientListener.IsConnected())
+			{
+				if (CurrentTask.IsSet())
+				{
+					Dispatcher.SetTaskState(CurrentTask->Index, ProcessFailed);
+					CurrentTask.Reset();
+				}
+				RestartProcessor();
+				bTaskIsRunning = true;
+				continue;
+			}
+
+			// If No task running, launch one
+			uint32 DataSize;
+			ClientListener.HasPendingData(DataSize);
+			if (!bTaskIsRunning && !DataSize)
+			{
+				CurrentTask = Dispatcher.GetTask();
+				if (CurrentTask.IsSet())
+				{
+					FDatasmithRunTaskCommand AddJob(CurrentTask.GetValue().FileName, CurrentTask.GetValue().Index);
+					AddJob.Write(ClientListener);
+					bTaskIsRunning = true;
+					continue;
+				}
+			}
+
+			if (DataSize)
+			{
+				if (ICommand* Command = CommandManager.GetNextCommand())
+				{
+					switch (Command->GetType())
+					{
+					case Ping:
+						PingProcess(StaticCast<FDatasmithPingCommand*>(Command));
+						bTaskIsRunning = false;
+						break;
+					case BackPing:
+						BackPingProcess(StaticCast<FDatasmithBackPingCommand*>(Command));
+						break;
+					case NotifyEndTask:
+					{
+						EndTaskProcess(StaticCast<FDatasmithNotifyEndTaskCommand*>(Command));
+						bTaskIsRunning = false;
+						break;
+					}
+
+					case ImportParams:
+					case RunTask:
+					default:
+						break;
+					}
+				}
+			}
+
+			FWindowsPlatformProcess::Sleep(0.001);
+		}
+
+		Kill();
+		return true;
+	}
+
+	void FDatasmithWorkerHandler::PingProcess(FDatasmithPingCommand* PingCommand)
+	{
+		DatasmithDispatcher::FDatasmithBackPingCommand NewCommand;
+		NewCommand.Write(ClientListener);
+	}
+
+	void FDatasmithWorkerHandler::BackPingProcess(FDatasmithBackPingCommand* PingCommand)
+	{
+		//FTimespan PingTime = Stopwatch.GetDuration() - StartTime;
+		//UE_LOG(LogDatasmithDispatcher, Display, TEXT("Connected in %d s"), PingTime.GetTotalSeconds());
+
+		return ;
+	}
+
+	void FDatasmithWorkerHandler::EndTaskProcess(FDatasmithNotifyEndTaskCommand* RunTaskCommand)
+	{
+		if (!CurrentTask.IsSet())
+		{
+			return;
+		}
+
+		FString CurrentPath = FPaths::GetPath(CurrentTask->FileName);
+		for (const FString& ExternalReferenceFile : RunTaskCommand->GetExternalReferences())
+		{
+			Dispatcher.AddTask(FPaths::Combine(CurrentPath, ExternalReferenceFile));
+		}
+		FString CurrentFileName = FPaths::GetCleanFilename(CurrentTask->FileName);
+		Dispatcher.SetTaskState(CurrentTask->Index, RunTaskCommand->GetProcessResult());
+		Dispatcher.LinkCTFileToUnrealCacheFile(CurrentFileName, RunTaskCommand->GetSceneGraphFile(), RunTaskCommand->GetGeomFile());
+		CurrentTask.Reset();
+	}
+
+	void FDatasmithWorkerHandler::NotifyKill()
+	{
+		bShouldTerminate = true;
+	}
+
+	bool FDatasmithWorkerHandler::NeedToRun()
+	{
+		return !bShouldTerminate;
+	}
+
+	void FDatasmithWorkerHandler::Kill()
+	{
+		ClientListener.Close();
+		ServerSocket.Close();
+		FPlatformProcess::TerminateProc(WorkerHandle, true);
+	}
+}
diff --git a/Engine/Plugins/Enterprise/DatasmithCADImporter/Source/DatasmithDispatcher/Public/DatasmithCommands.h b/Engine/Plugins/Enterprise/DatasmithCADImporter/Source/DatasmithDispatcher/Public/DatasmithCommands.h
new file mode 100644
--- /dev/null
+++ b/Engine/Plugins/Enterprise/DatasmithCADImporter/Source/DatasmithDispatcher/Public/DatasmithCommands.h
@@ -0,0 +1,300 @@
+// Copyright 1998-2019 Epic Games, Inc. All Rights Reserved.
+#pragma once
+
+#include "CoreMinimal.h"
+
+#include "CADOptions.h"
+#include "HAL/Runnable.h"
+#include "Misc/DateTime.h"
+
+
+class FCoreTechFileParser;
+
+/**
+ * Serializable interface.
+ * Should allow writing/Reading from Socket/file/any stream
+ */
+
+namespace DatasmithDispatcher
+{
+
+	enum EProcessState
+	{
+		Unknown = 0x00,
+		UnTreated = 0x01,
+		Running = 0x02,
+		ProcessOk = 0x04,
+		ProcessFailed = 0x08,
+		FileNotFound = 0x10
+	};
+
+	struct FTask
+	{
+		FTask(const FString& InFile)
+		{
+			FileName = InFile;
+			State = UnTreated;
+		}
+
+		FTask()
+		{
+		}
+
+		FString FileName;
+		int32 Index = -1;
+		EProcessState State = Unknown;
+	};
+
+	class FDatasmithDispatcherSocket;
+	class FDatasmithDispatcher;
+
+	/**
+	 * Enum with all command types
+	 */
+	enum DatasmithCommandType
+	{
+		Ping,
+		BackPing,
+		RunTask,
+		NotifyEndTask,
+		ImportParams,
+		Last
+	};
+
+	/**
+	 * Command interface
+	 * Serializable for Socket
+	 */
+	class ICommand
+	{
+	public:
+		//Initialize command
+		virtual void Initialize() = 0;
+
+		//Get command type
+		virtual DatasmithCommandType GetType() = 0;
+
+		//Write to stream
+		virtual void Write(FDatasmithDispatcherSocket&) = 0;
+
+		//Read from stream
+		virtual void Read(FDatasmithDispatcherSocket&) = 0;
+	};
+
+	/**
+	 * Process commands received by client listener (here GPCTClientListener)
+	 */
+	class FDatasmithCommandManager
+	{
+	public:
+		//Initialize command manager (add default commands)
+		FDatasmithCommandManager(FDatasmithDispatcherSocket& CurrentSocket);
+
+		//Destructor
+		~FDatasmithCommandManager();
+
+		void Initialize();
+		void Terminate();
+
+		ICommand *GetNextCommand();
+
+	private:
+		//Command map (map command type with command implementation -> for Reading)
+		TMap<DatasmithCommandType, ICommand*> CommandMap;
+
+		FDatasmithDispatcherSocket& Socket;
+		uint32 CurrentCommandeSize;
+	};
+
+	/**
+	 * "Template" command. Use this command if you want to implement new ones -> this is one of the simplest possible commands
+	 */
+	class FDatasmithPingCommand : public ICommand
+	{
+	public:
+		//Initialize command
+		FDatasmithPingCommand() {}
+
+		virtual void Initialize() override
+		{
+		}
+
+		// Inherited via ICommand
+		virtual DatasmithCommandType GetType() override
+		{
+			return Ping;
+		}
+		virtual void Write(FDatasmithDispatcherSocket&) override;
+		virtual void Read(FDatasmithDispatcherSocket&) override;
+
+	private:
+
+	};
+
+	class FDatasmithBackPingCommand : public ICommand
+	{
+	public:
+		//Initialize command
+		FDatasmithBackPingCommand() {}
+
+		virtual void Initialize() override
+		{
+		}
+
+		// Inherited via ICommand
+		virtual DatasmithCommandType GetType() override
+		{
+			return BackPing;
+		}
+		virtual void Write(FDatasmithDispatcherSocket&) override;
+		virtual void Read(FDatasmithDispatcherSocket&) override;
+
+	private:
+
+	};
+
+	class FDatasmithRunTaskCommand : public ICommand
+	{
+	public:
+		//Constructors
+		FDatasmithRunTaskCommand()
+		{
+		}
+		FDatasmithRunTaskCommand(const FString& InOutputFile, const int32 InJobIndex) 
+			: JobFilePath(InOutputFile)
+			, JobIndex(InJobIndex)
+		{
+		}
+
+		virtual ~FDatasmithRunTaskCommand()
+		{
+		}
+
+		const FString& GetFileToProcess() const
+		{
+			return JobFilePath;
+		}
+
+		virtual void Initialize() override
+		{
+			JobFilePath.Empty();
+			JobIndex = 0;
+		}
+
+		// Inherited via ICommand
+		virtual DatasmithCommandType GetType() override
+		{
+			return RunTask;
+		}
+
+		virtual void Write(FDatasmithDispatcherSocket&) override;
+		virtual void Read(FDatasmithDispatcherSocket&) override;
+
+	private:
+		FString JobFilePath;
+		int32 JobIndex;
+	};
+
+	class FDatasmithNotifyEndTaskCommand : public ICommand
+	{
+	public:
+		virtual ~FDatasmithNotifyEndTaskCommand()
+		{
+		}
+
+		void SetExternalReferences(const TArray<FString>& ExternalRef)
+		{
+			ExternalReferenceSet = ExternalRef;
+		}
+	
+		const TArray<FString>& GetExternalReferences()
+		{
+			return ExternalReferenceSet;
+		}
+
+		void SetProcessResult(EProcessState InProcessResult)
+		{
+			ProcessResult = InProcessResult;
+		}
+
+		EProcessState GetProcessResult()
+		{
+			return ProcessResult;
+		}
+
+		void SetSceneGraphFile(const FString& InSceneGraphFileName)
+		{
+			SceneGraphFileName = InSceneGraphFileName;
+		}
+
+		const FString& GetSceneGraphFile()
+		{
+			return SceneGraphFileName;
+		}
+
+		void SetGeomFile(const FString& InGeomFile)
+		{
+			GeomFileName = InGeomFile;
+		}
+
+		const FString& GetGeomFile()
+		{
+			return GeomFileName;
+		}
+
+		virtual void Initialize() override
+		{
+			ExternalReferenceSet.Empty();
+			ProcessResult = UnTreated;
+		}
+
+		virtual DatasmithCommandType GetType() override
+		{
+			return NotifyEndTask;
+		}
+
+		virtual void Write(FDatasmithDispatcherSocket&) override;
+		virtual void Read(FDatasmithDispatcherSocket&) override;
+
+	private:
+		TArray<FString> ExternalReferenceSet;
+		EProcessState ProcessResult;
+		FString SceneGraphFileName;
+		FString GeomFileName;
+	};
+
+	class FDatasmithImportParametersCommand : public ICommand
+	{
+	public:
+		virtual ~FDatasmithImportParametersCommand()
+		{
+		}
+
+		virtual DatasmithCommandType GetType() override
+		{
+			return ImportParams;
+		}
+		
+		virtual void Initialize() override
+		{
+			ImportParameters.init();
+		}
+
+		virtual void Set(const CAD::FImportParameters& InImportParameters)
+		{
+			ImportParameters = InImportParameters;
+		}
+
+		virtual void Get(CAD::FImportParameters& OutImportParameters)
+		{
+			OutImportParameters = ImportParameters;
+		}
+
+		virtual void Write(FDatasmithDispatcherSocket&) override;
+		virtual void Read(FDatasmithDispatcherSocket&) override;
+
+	private:
+		CAD::FImportParameters ImportParameters;
+	};
+
+} // NS DatasmithDispatcher
diff --git a/Engine/Plugins/Enterprise/DatasmithCADImporter/Source/DatasmithDispatcher/Public/DatasmithDispatcher.h b/Engine/Plugins/Enterprise/DatasmithCADImporter/Source/DatasmithDispatcher/Public/DatasmithDispatcher.h
new file mode 100644
--- /dev/null
+++ b/Engine/Plugins/Enterprise/DatasmithCADImporter/Source/DatasmithDispatcher/Public/DatasmithDispatcher.h
@@ -0,0 +1,63 @@
+// Copyright 1998-2019 Epic Games, Inc. All Rights Reserved.
+#pragma once
+
+#include "DatasmithWorkerHandler.h"
+
+#include "Containers/Queue.h"
+#include "CoreMinimal.h"
+#include "HAL/CriticalSection.h"
+#include "HAL/Thread.h"
+
+namespace DatasmithDispatcher
+{
+
+
+	//Wrapper class to launch processors
+	class DATASMITHDISPATCHER_API FDatasmithDispatcher
+	{
+
+	public:
+		//Constructor
+		FDatasmithDispatcher(const TCHAR* InCacheDir, int32 InNumberOfWorkers, TMap<FString, FString>& CADFileToUnrealFileMap, TMap<FString, FString>& CADFileToUnrealGeomMap);
+		~FDatasmithDispatcher();
+
+		/**
+		 * Spawns worker handlers
+		 */
+		void RunHandlers();
+		void Close();
+
+		void Process();
+		void Clear();
+
+		void AddTask(const FString& FileName); // Thread safe with FScopeLock Lock(&TaskPoolCriticalSection);
+
+		TOptional<FTask> GetTask();   // Thread safe with FScopeLock Lock(&TaskPoolCriticalSection);
+		bool IsOver(); // Thread safe with FScopeLock Lock(&TaskPoolCriticalSection);
+
+		void SetTaskState(int32 TaskIndex, EProcessState TaskState);   // Thread safe
+		void LinkCTFileToUnrealCacheFile(const FString& CTFile, const FString& UnrealSceneGraphFile, const FString& UnrealGeomFile);   // Thread safe
+
+
+	private:
+		FCriticalSection TaskPoolCriticalSection;
+
+		FString ProcessCacheFolder;
+		TMap<FString, FString>& CADFileToUnrealFileMap;
+		TMap<FString, FString>& CADFileToUnrealGeomMap;
+
+		//Started
+		bool bStarted;
+
+		int32 NumberOfWorkers;
+		TArray<FDatasmithWorkerHandler> WorkerHandlerSet;
+		TArray<FThread> WorkersRunning;
+
+		TArray<FTask> TaskPool;
+		int32 LastLaunchTask;
+		int32 NumOfFinishedTask;
+
+
+	}; // FDatasmithDispatcher
+
+} // NS DatasmithDispatcher
diff --git a/Engine/Plugins/Enterprise/DatasmithCADImporter/Source/DatasmithDispatcher/Public/DatasmithDispatcherModule.h b/Engine/Plugins/Enterprise/DatasmithCADImporter/Source/DatasmithDispatcher/Public/DatasmithDispatcherModule.h
new file mode 100644
--- /dev/null
+++ b/Engine/Plugins/Enterprise/DatasmithCADImporter/Source/DatasmithDispatcher/Public/DatasmithDispatcherModule.h
@@ -0,0 +1,19 @@
+// Copyright 1998-2019 Epic Games, Inc. All Rights Reserved.
+#pragma once
+
+#include "Modules/ModuleInterface.h"
+
+#define DATASMITHDISPATCHER_MODULE_NAME TEXT("DatasmithDispatcher")
+
+/**
+ * This module exposes additional features for assets containing CoreTech data.
+ */
+class FDatasmithDispatcherModule : public IModuleInterface
+{
+public:
+	static FDatasmithDispatcherModule& Get();
+	static bool IsAvailable();
+
+private:
+	virtual void StartupModule() override;
+};
diff --git a/Engine/Plugins/Enterprise/DatasmithCADImporter/Source/DatasmithDispatcher/Public/DatasmithDispatcherSocket.h b/Engine/Plugins/Enterprise/DatasmithCADImporter/Source/DatasmithDispatcher/Public/DatasmithDispatcherSocket.h
new file mode 100644
--- /dev/null
+++ b/Engine/Plugins/Enterprise/DatasmithCADImporter/Source/DatasmithDispatcher/Public/DatasmithDispatcherSocket.h
@@ -0,0 +1,117 @@
+// Copyright 1998-2019 Epic Games, Inc. All Rights Reserved.
+#pragma once
+
+#include "CoreMinimal.h"
+
+#include "Sockets.h"
+#include "GenericPlatform/GenericPlatformProcess.h"
+
+#define DEFAULT_DATASMITHSERVER_PORT 0xCAFE
+#define MAX_DATASMITHSERVER_PORT 0xFFFF 
+
+namespace DatasmithDispatcher
+{
+
+	enum SocketErrorCode
+	{
+		NoError = 0,
+		UnableToReadOnSocket = 1,
+		UnableToSendData = 2,
+		CouldNotStartWSA = 3,
+		UnableToGetLocalAddress = 4,
+		ConnectionToServerFailed = 5
+	};
+
+	class FDatasmithDispatcherSocket 
+	{
+
+	public:
+		FDatasmithDispatcherSocket(const TCHAR* ServerAddress);
+		FDatasmithDispatcherSocket();
+		FDatasmithDispatcherSocket& operator=(const FDatasmithDispatcherSocket&) = delete;
+			
+		void SetSocket(FSocket* InSocket);
+
+		int32 GetErrorCode() const
+		{
+			return ErrorCode;
+		}
+
+		void Connect(const int32& ServerPort);
+
+		bool IsConnected();
+
+		void Bind();
+
+		bool Listen();
+		FSocket* Accept();
+
+		int32 GetPort() const
+		{
+			if (!Socket)
+			{
+				return -1;
+			}
+			return Socket->GetPortNo();
+		}
+
+		void Close();
+
+		template <typename T>
+		void operator >>(T &val)
+		{
+			Read(sizeof(val), (uint8*)&val);
+		}
+
+		template <typename T>
+		void operator <<(T val)
+		{
+			Write(sizeof(val), (uint8*)&val);
+		}
+
+		void operator >>(FString& str);
+		void operator <<(const FString& str);
+
+		bool IsClosed() const 
+		{
+			return !bOpen;
+		}
+
+		bool IsOpen() const
+		{
+			return bOpen;
+		}
+
+		FSocket* GetSocket()
+		{
+			return Socket;
+		}
+
+		bool HasPendingData(uint32& DataSize);
+
+		/**
+		 * Actually send the buffered output cache
+		 */
+		void SendData();
+
+	private:
+		// low-level reading function
+		void Read(int32 Num, uint8 *Buffer);
+
+		// low-level writing function
+		void Write(int32 Num, uint8 *Buffer);
+
+	private:
+		TArray<uint8> Cache;
+
+		FString SocketAdress;
+		int32 SocketId;
+		FSocket* Socket;
+
+		bool bServerSide;
+		bool bOpen;
+
+		int32 ErrorCode;
+	};
+
+} //namespace DatasmithDispatcher
diff --git a/Engine/Plugins/Enterprise/DatasmithCADImporter/Source/DatasmithDispatcher/Public/DatasmithWorkerHandler.h b/Engine/Plugins/Enterprise/DatasmithCADImporter/Source/DatasmithDispatcher/Public/DatasmithWorkerHandler.h
new file mode 100644
--- /dev/null
+++ b/Engine/Plugins/Enterprise/DatasmithCADImporter/Source/DatasmithDispatcher/Public/DatasmithWorkerHandler.h
@@ -0,0 +1,83 @@
+// Copyright 1998-2019 Epic Games, Inc. All Rights Reserved.
+#pragma once
+
+#include "CoreMinimal.h"
+
+#include "DatasmithCommands.h"
+#include "DatasmithDispatcherSocket.h"
+#include "Misc/Timespan.h"
+
+
+#if PLATFORM_MAC
+#include "Mac/MacPlatformProcess.h"
+#elif PLATFORM_LINUX
+#include "Unix/UnixPlatformProcess.h"
+#else
+#include "Windows/WindowsPlatformProcess.h"
+#endif
+
+
+
+class FSocket;
+
+namespace DatasmithDispatcher
+{
+	class FDatasmithBackPingCommand;
+	class FDatasmithDispatcher;
+	class FDatasmithNotifyEndTaskCommand;
+	class FDatasmithPingCommand;
+
+	//Handle a Worker by socket communication
+	class FDatasmithWorkerHandler
+	{
+
+	public:
+		FDatasmithWorkerHandler(FDatasmithDispatcher& InDispatcher, const FString& InCachePath);
+
+		bool Run();
+		void NotifyKill();
+
+	private:
+		const FString& GetCachePath() const
+		{
+			return CachePath;
+		}
+
+		FString GetExePath();
+
+		bool InitializeSocket();
+		bool StartWorker();
+		void Kill();
+
+		void RestartProcessor();
+
+		bool NeedToRun();
+
+		void PingProcess(FDatasmithPingCommand* PingCommand);
+		void BackPingProcess(FDatasmithBackPingCommand* PingCommand);
+		void EndTaskProcess(FDatasmithNotifyEndTaskCommand* RunTaskCommand);
+
+
+	private:
+		FDatasmithDispatcher& Dispatcher;
+		uint32 WorkerProcessId;
+
+		FDatasmithDispatcherSocket ServerSocket;  
+		/**
+		 *  Template socket used to accept connections
+		 */
+		FDatasmithDispatcherSocket ClientListener; 
+
+		FString CachePath;
+
+		TOptional<FTask> CurrentTask;
+
+		FProcHandle WorkerHandle;
+
+		bool bShouldTerminate;
+		//FTimespan Stopwatch;
+		//FTimespan StartPingTime;
+		//FTimespan StartTime;
+
+	};
+} // namespace DatasmithDispatcher
