Merging 9745597 //UE4/Dev-Enterprise to Release-4.24 (//UE4/Release-4.24)
DatasmithImporterDispatcher

diff --git a/Engine/Plugins/Enterprise/DatasmithCADImporter/DatasmithCADImporter.uplugin b/Engine/Plugins/Enterprise/DatasmithCADImporter/DatasmithCADImporter.uplugin
--- a/Engine/Plugins/Enterprise/DatasmithCADImporter/DatasmithCADImporter.uplugin
+++ b/Engine/Plugins/Enterprise/DatasmithCADImporter/DatasmithCADImporter.uplugin
@@ -41,6 +41,30 @@
 				"Win64"
 			]
 		},
+		{
+			"Name": "DatasmithDispatcher",
+			"Type": "Editor",
+			"LoadingPhase": "Default",
+			"WhitelistPlatforms": [
+				"Win64"
+			]
+		},
+		{
+			"Name": "CADInterfaces",
+			"Type": "Editor",
+			"LoadingPhase": "Default",
+			"WhitelistPlatforms": [
+				"Win64"
+			]
+		},
+		{
+			"Name": "CADTools",
+			"Type": "Editor",
+			"LoadingPhase": "Default",
+			"WhitelistPlatforms": [
+				"Win64"
+			]
+		},
 		{
 			"Name": "DatasmithWireTranslator",
 			"Type": "Editor",
diff --git a/Engine/Plugins/Enterprise/DatasmithCADImporter/Source/CADInterfaces/CADInterfaces.Build.cs b/Engine/Plugins/Enterprise/DatasmithCADImporter/Source/CADInterfaces/CADInterfaces.Build.cs
new file mode 100644
--- /dev/null
+++ b/Engine/Plugins/Enterprise/DatasmithCADImporter/Source/CADInterfaces/CADInterfaces.Build.cs
@@ -0,0 +1,38 @@
+// Copyright 1998-2019 Epic Games, Inc. All Rights Reserved.
+
+using System.IO;
+
+namespace UnrealBuildTool.Rules
+{
+	public class CADInterfaces : ModuleRules
+	{
+		public CADInterfaces(ReadOnlyTargetRules Target) : base(Target)
+		{
+			PublicDependencyModuleNames.AddRange(
+				new string[]
+				{
+					"Core",
+					"CADTools"
+				}
+			);
+
+			// Module is effectively empty without non redistributable content
+			bool bHasNonRedistSources = Directory.Exists(Path.Combine(ModuleDirectory, "Public", "NotForLicensees"));
+
+			// CAD library is only available if CoreTech is available too
+			bool bHasCoretech = System.Type.GetType("CoreTech") != null;
+
+			// Support for Windows only
+			bool bIsPlateformSupported = Target.Platform == UnrealTargetPlatform.Win64;
+
+			if (bHasNonRedistSources && bIsPlateformSupported && bHasCoretech)
+			{
+				PublicDefinitions.Add("CAD_INTERFACE");
+				PublicDependencyModuleNames.Add("CoreTech");
+
+				PublicIncludePaths.Add(ModuleDirectory + "/Public/NotForLicensees");
+				PrivateIncludePaths.Add(ModuleDirectory + "/Private/NotForLicensees");
+			}
+		}
+	}
+}
diff --git a/Engine/Plugins/Enterprise/DatasmithCADImporter/Source/CADInterfaces/Private/CADInterfacesModule.cpp b/Engine/Plugins/Enterprise/DatasmithCADImporter/Source/CADInterfaces/Private/CADInterfacesModule.cpp
new file mode 100644
--- /dev/null
+++ b/Engine/Plugins/Enterprise/DatasmithCADImporter/Source/CADInterfaces/Private/CADInterfacesModule.cpp
@@ -0,0 +1,26 @@
+// Copyright 1998-2019 Epic Games, Inc. All Rights Reserved.
+#include "CADInterfacesModule.h"
+
+#include "Modules/ModuleInterface.h"
+#include "Modules/ModuleManager.h"
+
+#define LOCTEXT_NAMESPACE "CADInterfacesModule"
+
+FCADInterfacesModule& FCADInterfacesModule::Get()
+{
+	return FModuleManager::LoadModuleChecked< FCADInterfacesModule >(CADINTERFACES_MODULE_NAME);
+}
+
+bool FCADInterfacesModule::IsAvailable()
+{
+	return FModuleManager::Get().IsModuleLoaded(CADINTERFACES_MODULE_NAME);
+}
+
+void FCADInterfacesModule::StartupModule()
+{
+}
+
+IMPLEMENT_MODULE(FCADInterfacesModule, CADInterfaces);
+
+#undef LOCTEXT_NAMESPACE // "CADInterfacesModule"
+
diff --git a/Engine/Plugins/Enterprise/DatasmithCADImporter/Source/CADInterfaces/Public/CADInterfacesModule.h b/Engine/Plugins/Enterprise/DatasmithCADImporter/Source/CADInterfaces/Public/CADInterfacesModule.h
new file mode 100644
--- /dev/null
+++ b/Engine/Plugins/Enterprise/DatasmithCADImporter/Source/CADInterfaces/Public/CADInterfacesModule.h
@@ -0,0 +1,16 @@
+// Copyright 1998-2019 Epic Games, Inc. All Rights Reserved.
+#pragma once
+
+#include "Modules/ModuleInterface.h"
+
+#define CADINTERFACES_MODULE_NAME TEXT("WorkerDataManagement")
+
+class FCADInterfacesModule : public IModuleInterface
+{
+public:
+	static FCADInterfacesModule& Get();
+	static bool IsAvailable();
+
+private:
+	virtual void StartupModule() override;
+};
diff --git a/Engine/Plugins/Enterprise/DatasmithCADImporter/Source/CADLibrary/CADLibrary.Build.cs b/Engine/Plugins/Enterprise/DatasmithCADImporter/Source/CADLibrary/CADLibrary.Build.cs
--- a/Engine/Plugins/Enterprise/DatasmithCADImporter/Source/CADLibrary/CADLibrary.Build.cs
+++ b/Engine/Plugins/Enterprise/DatasmithCADImporter/Source/CADLibrary/CADLibrary.Build.cs
@@ -14,7 +14,6 @@ public class CADLibrary : ModuleRules
 			new string[]
 			{
 				"Core",
-				"CoreUObject",
 				"DatasmithCore",
 				"MeshDescription",
 				"StaticMeshDescription"
@@ -38,6 +37,8 @@ public class CADLibrary : ModuleRules
 				new string[]
 				{
 					"CoreTech",
+					"CADInterfaces",
+					"CADTools"
 				}
 			);
 		}
diff --git a/Engine/Plugins/Enterprise/DatasmithCADImporter/Source/CADLibrary/Public/CADLibraryOptions.h b/Engine/Plugins/Enterprise/DatasmithCADImporter/Source/CADLibrary/Public/CADLibraryOptions.h
deleted file mode 100644
diff --git a/Engine/Plugins/Enterprise/DatasmithCADImporter/Source/CADTools/CADTools.Build.cs b/Engine/Plugins/Enterprise/DatasmithCADImporter/Source/CADTools/CADTools.Build.cs
new file mode 100644
--- /dev/null
+++ b/Engine/Plugins/Enterprise/DatasmithCADImporter/Source/CADTools/CADTools.Build.cs
@@ -0,0 +1,25 @@
+// Copyright 1998-2019 Epic Games, Inc. All Rights Reserved.
+
+using System.IO;
+
+namespace UnrealBuildTool.Rules
+{
+	public class CADTools : ModuleRules
+	{
+		public CADTools(ReadOnlyTargetRules Target) : base(Target)
+		{
+            PublicDependencyModuleNames.AddRange(
+				new string[]
+				{
+                    "Core",
+                }
+            );
+
+			PrivateDependencyModuleNames.AddRange(
+				new string[]
+				{
+                }
+			);
+        }
+    }
+}
diff --git a/Engine/Plugins/Enterprise/DatasmithCADImporter/Source/CADTools/Private/CADData.cpp b/Engine/Plugins/Enterprise/DatasmithCADImporter/Source/CADTools/Private/CADData.cpp
new file mode 100644
--- /dev/null
+++ b/Engine/Plugins/Enterprise/DatasmithCADImporter/Source/CADTools/Private/CADData.cpp
@@ -0,0 +1,200 @@
+// Copyright 1998-2019 Epic Games, Inc. All Rights Reserved.
+#include "CADData.h"
+
+#include "CADOptions.h"
+
+#include "Misc/FileHelper.h"
+
+
+namespace CADLibrary
+{
+uint32 BuildFastColorHash(uint32 ColorId, uint8 Alpha)
+{
+	if (Alpha == 0)
+	{
+		Alpha = 1;
+	}
+	return ColorId | Alpha << 24;
+}
+
+void UnhashFastColorHash(uint32 ColorHash, uint32& ColorId, uint8& Alpha)
+{
+	ColorId = ColorHash & 0x00ffffff;
+	Alpha = (uint8)((ColorHash & 0xff000000) >> 24);
+}
+
+int32 BuildColorHash(const FColor& Color)
+{
+	FString HashString = FString::Printf(TEXT("%02x%02x%02x%02x"), Color.R, Color.G, Color.B, Color.A);
+	return FGenericPlatformMath::Abs((int32)GetTypeHash(HashString));
+}
+
+int32 BuildMaterialHash(const FCADMaterial& Material)
+{
+	FString Hash;
+	if (!Material.MaterialName.IsEmpty())
+	{
+		Hash += Material.MaterialName;  // we add material name because it could be used by the end user so two material with same parameters but different name are different.
+	}
+	Hash += FString::Printf(TEXT("%02x%02x%02x "), Material.Diffuse.R, Material.Diffuse.G, Material.Diffuse.B);
+	Hash += FString::Printf(TEXT("%02x%02x%02x "), Material.Ambient.R, Material.Ambient.G, Material.Ambient.B);
+	Hash += FString::Printf(TEXT("%02x%02x%02x "), Material.Specular.R, Material.Specular.G, Material.Specular.B);
+	Hash += FString::Printf(TEXT("%02x%02x%02x"), (int)(Material.Shininess * 255.0), (int)(Material.Transparency * 255.0), (int)(Material.Reflexion * 255.0));
+
+	if (!Material.TextureName.IsEmpty())
+	{
+		Hash += Material.TextureName;
+	}
+	return FMath::Abs((int32)GetTypeHash(Hash));
+}
+
+
+FBody::FBody(uint32 InBodyID, int32 FaceNum)
+	: TriangleCount(0)
+	, BodyID(InBodyID)
+{
+	FaceTessellationSet.Reserve(FaceNum);
+}
+
+
+FCTRawGeomFile::FCTRawGeomFile(FString& InFileName, TMap< uint32, FBody* >& InBodyUuidToCTBodyMap)
+	: BodyUuidToBodyMap(InBodyUuidToCTBodyMap)
+{
+	uint32 NbBody = 0;
+
+	TArray<uint8> RawData;
+	FFileHelper::LoadFileToArray(RawData, *InFileName);
+
+	ReadTessellationSetFromFile(RawData, NbBody, CTBodySet);
+
+	for(FBody& Body : CTBodySet)
+	{
+		//CTIdToBodyUuidMap.Emplace(Body.GetBodyId(), Body.GetBodyUuid());
+		BodyUuidToBodyMap.Add(Body.GetBodyUuid(), &Body);
+	}
+}
+
+enum ETessellationContent
+{
+	HasVertex = 1 << 1,
+	HasNormal = 1 << 2,
+	HasIndex = 1 << 3,
+	HasTexCoord = 1 << 4,
+
+	Default = HasVertex | HasNormal | HasIndex | HasTexCoord
+};
+
+void WriteTessellationInRawData(FTessellationData& Tessellation, TArray<uint8>& GlobalRawData)
+{
+	int32 INFO[LastLine];
+
+	int32 Types = ETessellationContent::Default;
+	uint32 VertexSize = Tessellation.VertexArray.Num();
+	uint32 NormalSize = Tessellation.NormalArray.Num();
+	uint32 IndexSize = Tessellation.IndexArray.Num();
+	uint32 TexCoordSize = Tessellation.TexCoordArray.Num();
+
+	uint32 TessellationRawDataSize = VertexSize + NormalSize + IndexSize + TexCoordSize + sizeof(INFO); // VertexCount, VertexType, NormalCount, NormalType, IndexCount, IndexType, TexCoordType;
+
+	INFO[LineType] = Types;
+	INFO[LineRawSize] = TessellationRawDataSize;
+	INFO[LineBodyId] = Tessellation.BodyId;
+	INFO[LineBodyUuId] = Tessellation.BodyUuId;
+	INFO[LineBodyFaceNum] = Tessellation.BodyFaceNum;
+	INFO[LineMaterialId] = Tessellation.MaterialId;
+	INFO[LineMaterialHash] = Tessellation.MaterialHash;
+	INFO[LineVertexCount] = Tessellation.VertexCount;
+	INFO[LineVertexType] = Tessellation.SizeOfVertexType;
+	INFO[LineNormalCount] = Tessellation.NormalCount;
+	INFO[LineNormalType] = Tessellation.SizeOfNormalType;
+	INFO[LineIndexCount] = Tessellation.IndexCount;
+	INFO[LineIndexType] = Tessellation.SizeOfIndexType;
+	INFO[LineTexCoordType] = Tessellation.SizeOfTexCoordType;
+
+	GlobalRawData.Append((uint8*)INFO, sizeof(INFO));
+	GlobalRawData.Append(Tessellation.VertexArray);
+	GlobalRawData.Append(Tessellation.NormalArray);
+	GlobalRawData.Append(Tessellation.IndexArray);
+	GlobalRawData.Append(Tessellation.TexCoordArray);
+}
+
+bool ReadTessellationSetFromFile(TArray<uint8>& RawData, uint32& OutNbBodies, TArray<FBody>& CTBodySet)
+{
+	uint32 INFO[LastLine];
+	const uint32 InfoSize = sizeof(INFO);
+	int32 RawDataOffset = sizeof(uint32);
+	memcpy(&OutNbBodies, RawData.GetData(), RawDataOffset);
+
+	CTBodySet.Reserve(OutNbBodies);
+
+	FBody* CurrentBody = nullptr;
+
+	int32 CurrentBodyId = -1;
+	int32 CurrentBodyIndex = -1;
+	while (RawDataOffset < RawData.Num())
+	{
+		memcpy(INFO, RawData.GetData() + RawDataOffset, InfoSize);
+		RawDataOffset += InfoSize;
+
+		int32 Types = INFO[LineType];
+		int32 RawSize = INFO[LineRawSize];
+
+		if (Types != ETessellationContent::Default)
+		{
+			return false;
+		}
+
+		if (INFO[LineBodyId] != CurrentBodyId)
+		{
+			CurrentBodyId = INFO[LineBodyId];
+			CTBodySet.Emplace(INFO[LineBodyId], INFO[LineBodyFaceNum]);
+			CurrentBodyIndex++;
+
+			CurrentBody = &CTBodySet[CurrentBodyIndex];
+			CurrentBody->SetBodyId(INFO[LineBodyId]);
+			CurrentBody->SetBodyUuid(INFO[LineBodyUuId]);
+			CurrentBody->GetTessellationSet().Reserve(INFO[LineBodyFaceNum]);
+		}
+
+		int32 Index = CurrentBody->GetTessellationSet().Emplace();
+		FTessellationData& Tessellation = CurrentBody->GetTessellationSet()[Index];
+
+		Tessellation.BodyId = INFO[LineBodyId];
+		Tessellation.BodyUuId = INFO[LineBodyUuId];
+		Tessellation.BodyFaceNum = INFO[LineBodyFaceNum];
+		Tessellation.MaterialId = INFO[LineMaterialId];
+		Tessellation.MaterialHash = INFO[LineMaterialHash];
+		Tessellation.VertexCount = INFO[LineVertexCount];
+		Tessellation.SizeOfVertexType = INFO[LineVertexType];
+		Tessellation.NormalCount = INFO[LineNormalCount];
+		Tessellation.SizeOfNormalType = INFO[LineNormalType];
+		Tessellation.IndexCount = INFO[LineIndexCount];
+		Tessellation.SizeOfIndexType = INFO[LineIndexType];
+		Tessellation.SizeOfTexCoordType = INFO[LineTexCoordType];
+
+		uint32 VertexSize = Tessellation.VertexCount * Tessellation.SizeOfVertexType * 3;
+		uint32 NormalSize = Tessellation.NormalCount * Tessellation.SizeOfNormalType * 3;
+		uint32 IndexSize = Tessellation.IndexCount * Tessellation.SizeOfIndexType;
+		uint32 TexCoordSize = Tessellation.VertexCount * Tessellation.SizeOfTexCoordType * 2;
+		uint32 DataByteSize = VertexSize + NormalSize + IndexSize + TexCoordSize + InfoSize;
+
+
+		Tessellation.VertexArray.Empty(VertexSize);
+		Tessellation.NormalArray.Empty(NormalSize);
+		Tessellation.IndexArray.Empty(IndexSize);
+		Tessellation.TexCoordArray.Empty(TexCoordSize);
+
+		Tessellation.VertexArray.Append(RawData.GetData() + RawDataOffset, VertexSize);
+		RawDataOffset += VertexSize;
+		Tessellation.NormalArray.Append(RawData.GetData() + RawDataOffset, NormalSize);
+		RawDataOffset += NormalSize;
+		Tessellation.IndexArray.Append(RawData.GetData() + RawDataOffset, IndexSize);
+		RawDataOffset += IndexSize;
+		Tessellation.TexCoordArray.Append(RawData.GetData() + RawDataOffset, TexCoordSize);
+		RawDataOffset += TexCoordSize;
+	}
+	return true;
+}
+
+
+}
diff --git a/Engine/Plugins/Enterprise/DatasmithCADImporter/Source/CADTools/Private/CADToolsModule.cpp b/Engine/Plugins/Enterprise/DatasmithCADImporter/Source/CADTools/Private/CADToolsModule.cpp
new file mode 100644
--- /dev/null
+++ b/Engine/Plugins/Enterprise/DatasmithCADImporter/Source/CADTools/Private/CADToolsModule.cpp
@@ -0,0 +1,26 @@
+// Copyright 1998-2019 Epic Games, Inc. All Rights Reserved.
+#include "CADToolsModule.h"
+
+#include "Modules/ModuleInterface.h"
+#include "Modules/ModuleManager.h"
+
+#define LOCTEXT_NAMESPACE "FCADToolsModule"
+
+FCADToolsModule& FCADToolsModule::Get()
+{
+	return FModuleManager::LoadModuleChecked< FCADToolsModule >(CADTOOLS_MODULE_NAME);
+}
+
+bool FCADToolsModule::IsAvailable()
+{
+	return FModuleManager::Get().IsModuleLoaded(CADTOOLS_MODULE_NAME);
+}
+
+void FCADToolsModule::StartupModule()
+{
+}
+
+IMPLEMENT_MODULE(FCADToolsModule, CADTools);
+
+#undef LOCTEXT_NAMESPACE
+
diff --git a/Engine/Plugins/Enterprise/DatasmithCADImporter/Source/CADTools/Public/CADData.h b/Engine/Plugins/Enterprise/DatasmithCADImporter/Source/CADTools/Public/CADData.h
new file mode 100644
--- /dev/null
+++ b/Engine/Plugins/Enterprise/DatasmithCADImporter/Source/CADTools/Public/CADData.h
@@ -0,0 +1,243 @@
+// Copyright 1998-2019 Epic Games, Inc. All Rights Reserved.
+#pragma once
+
+#include "CoreMinimal.h"
+
+#include "Containers/Array.h"
+#include "Math/Color.h"
+
+namespace CADLibrary
+{
+
+enum ECTTessellationLine
+{
+	LineType = 0,
+	LineRawSize,
+	LineBodyId,
+	LineBodyUuId,
+	LineBodyFaceNum,
+	LineMaterialId,
+	LineMaterialHash,
+	LineVertexCount,
+	LineVertexType,
+	LineNormalCount,
+	LineNormalType,
+	LineIndexCount,
+	LineIndexType,
+	LineTexCoordType,
+	LastLine
+};
+
+struct CADTOOLS_API FCADMaterial
+{
+	int32 MaterialId;  //#ueent_CAD
+	FString MaterialName;
+	FColor Diffuse;
+	FColor Ambient;
+	FColor Specular;
+	float Shininess;
+	float Transparency;
+	float Reflexion;
+	FString TextureName;
+};
+
+struct CADTOOLS_API FObjectDisplayDataId
+{
+	uint32 MaterialId = 0;
+	uint32 ColorHId = 0; // => FastHash == ColorId+Transparency
+};
+
+/**
+ * Helper struct to store tessellation data from CoreTech
+ */
+struct CADTOOLS_API FTessellationData
+{
+	TArray<uint8> VertexArray;
+	TArray<uint8> NormalArray;
+	TArray<uint8> IndexArray;
+	TArray<uint8> TexCoordArray;
+	uint32        VertexCount = 0;
+	uint32        NormalCount = 0;
+	uint32        IndexCount = 0;
+
+	uint32    StartVertexIndex = 0;
+
+	uint8 SizeOfVertexType;
+	uint8 SizeOfTexCoordType;
+	uint8 SizeOfNormalType;
+	uint8 SizeOfIndexType;
+	bool  bHasRGBColor;
+
+	uint32 BodyId = 0;
+	uint32 BodyUuId = 0;
+	uint32 BodyFaceNum = 0;
+	uint32 MaterialId = 0;
+	uint32 MaterialHash = 0;
+
+	TArray<int32> VertexIdSet;  // StaticMesh FVertexID
+	TArray<int32> SymVertexIdSet; // StaticMesh FVertexID for sym part
+};
+
+class CADTOOLS_API FBody
+{
+public:
+	FBody(uint32 BodyID, int32 FaceNum);
+
+	TArray<FTessellationData>& GetTessellationSet()
+	{
+		return FaceTessellationSet;
+	}
+
+	uint32 GetTriangleCount() const
+	{
+		return TriangleCount;
+	}
+
+	uint32 GetBodyId() const
+	{
+		return BodyID;
+	}
+
+	uint32 GetBodyUuid() const 
+	{
+		return BodyUuid;
+	}
+
+	void SetBodyId(uint32 Id)
+	{
+		BodyID = Id;
+	}
+
+	void SetBodyUuid(uint32 Uuid)
+	{
+		BodyUuid = Uuid;
+	}
+
+protected:
+
+	/**
+	 * Array of FCTTessellation
+	 * FCTTessellation is an elementary structure of Mesh
+	 */
+	TArray<FTessellationData> FaceTessellationSet;
+
+	/**
+	 * Structure use to map Material ID to Material Hash. Material hash is an unique value by material
+	 */
+	TMap<uint32, uint32> MaterialIdToMaterialHash;
+	uint32 TriangleCount;
+	uint32 BodyID;
+	uint32 BodyUuid;
+};
+
+class CADTOOLS_API FCTRawGeomFile
+{
+public:
+	FCTRawGeomFile(FString& InFileName, TMap< uint32, FBody* >& InBodyUuidToCTBodyMap);
+
+	FString& GetFileName()
+	{
+		return FileName;
+	}
+
+protected:
+	FString FileName;
+	TArray<FBody> CTBodySet;
+	TMap<uint32, FBody* >& BodyUuidToBodyMap;
+};
+
+CADTOOLS_API uint32 BuildFastColorHash(uint32 ColorId, uint8 Alpha);
+CADTOOLS_API void UnhashFastColorHash(uint32 ColorHash, uint32& OutColorId, uint8& OutAlpha);
+
+CADTOOLS_API int32 BuildColorHash(const FColor& Color);
+CADTOOLS_API int32 BuildMaterialHash(const FCADMaterial& Material);
+
+CADTOOLS_API void WriteTessellationInRawData(FTessellationData& Tessellation, TArray<uint8>& GlobalRawData);
+CADTOOLS_API bool ReadTessellationSetFromFile(TArray<uint8>& RawData, uint32& OutNbBody, TArray<FBody>& FaceTessellations);
+
+inline void CopyValue(const uint8* Source, int Offset, uint8 Size, bool bIs3D, FVector& Dest, const FMatrix* Matrix = nullptr)
+{
+	if (Source == nullptr)
+	{
+		return;
+	}
+
+	switch (Size)
+	{
+		case sizeof(uint8) :
+		{
+			Dest[0] = ((uint8*)Source)[Offset + 0] / 255.;
+			Dest[1] = ((uint8*)Source)[Offset + 1] / 255.;
+			Dest[2] = bIs3D ? ((uint8*)Source)[Offset + 2] / 255. : 0.;
+			break;
+		}
+		case sizeof(float) :
+		{
+			Dest[0] = ((float*)Source)[Offset + 0];
+			Dest[1] = ((float*)Source)[Offset + 1];
+			Dest[2] = bIs3D ? ((float*)Source)[Offset + 2] : 0.;
+			break;
+		}
+		case sizeof(double) :
+		{
+			Dest[0] = ((double*)Source)[Offset + 0];
+			Dest[1] = ((double*)Source)[Offset + 1];
+			Dest[2] = bIs3D ? ((double*)Source)[Offset + 2] : 0.;
+			break;
+		}
+		default:
+		{
+			Dest[0] = 0.;
+			Dest[1] = 0.;
+			Dest[2] = 0.;
+			break;
+		}
+	}
+
+	if (bIs3D && Matrix != nullptr)
+	{
+		Dest = Matrix->TransformPosition(Dest);
+	}
+}
+
+inline void CopyValue(const void* Source, int Offset, uint8 Size, int32 Dest[3])
+{
+	if (Source == nullptr)
+	{
+		return;
+	}
+
+	switch (Size)
+	{
+		case sizeof(uint8) :
+		{
+			Dest[0] = ((uint8*)Source)[Offset + 0];
+			Dest[1] = ((uint8*)Source)[Offset + 1];
+			Dest[2] = ((uint8*)Source)[Offset + 2];
+			break;
+		}
+		case sizeof(uint16) :
+		{
+			Dest[0] = ((uint16*)Source)[Offset + 0];
+			Dest[1] = ((uint16*)Source)[Offset + 1];
+			Dest[2] = ((uint16*)Source)[Offset + 2];
+			break;
+		}
+		case sizeof(uint32) :
+		{
+			Dest[0] = ((uint32*)Source)[Offset + 0];
+			Dest[1] = ((uint32*)Source)[Offset + 1];
+			Dest[2] = ((uint32*)Source)[Offset + 2];
+			break;
+		}
+		default:
+		{
+			Dest[0] = 0.;
+			Dest[1] = 0.;
+			Dest[2] = 0.;
+			break;
+		}
+	}
+}
+
+}
\ No newline at end of file
diff --git a/Engine/Plugins/Enterprise/DatasmithCADImporter/Source/CADTools/Public/CADOptions.h b/Engine/Plugins/Enterprise/DatasmithCADImporter/Source/CADTools/Public/CADOptions.h
new file mode 100644
--- /dev/null
+++ b/Engine/Plugins/Enterprise/DatasmithCADImporter/Source/CADTools/Public/CADOptions.h
@@ -0,0 +1,42 @@
+// Copyright 1998-2019 Epic Games, Inc. All Rights Reserved.
+
+#pragma once
+
+#include "Math/Vector.h"
+
+namespace CADLibrary
+{
+	enum EStitchingTechnique
+	{
+		StitchingNone = 0,
+		StitchingHeal,
+		StitchingSew,
+	};
+
+	enum class EModelCoordSystem : uint8
+	{
+		ZUp_LeftHanded,
+		ZUp_RightHanded,
+		YUp_LeftHanded,
+		YUp_RightHanded,
+	};
+
+	struct FImportParameters
+	{
+		double MetricUnit = 0.001;
+		double ScaleFactor = 0.1;
+		double ChordTolerance = 0.2;
+		double MaxEdgeLength = 0.0;
+		double MaxNormalAngle = 20.0;
+		EStitchingTechnique StitchingTechnique = EStitchingTechnique::StitchingNone;
+		EModelCoordSystem ModelCoordSys = EModelCoordSystem::ZUp_RightHanded;
+	};
+
+	struct FMeshParameters
+	{
+		bool bNeedSwapOrientation = false;
+		bool bIsSymmetric = false;
+		FVector SymmetricOrigin;
+		FVector SymmetricNormal;
+	};
+}
diff --git a/Engine/Plugins/Enterprise/DatasmithCADImporter/Source/CADTools/Public/CADToolsModule.h b/Engine/Plugins/Enterprise/DatasmithCADImporter/Source/CADTools/Public/CADToolsModule.h
new file mode 100644
--- /dev/null
+++ b/Engine/Plugins/Enterprise/DatasmithCADImporter/Source/CADTools/Public/CADToolsModule.h
@@ -0,0 +1,16 @@
+// Copyright 1998-2019 Epic Games, Inc. All Rights Reserved.
+#pragma once
+
+#include "Modules/ModuleInterface.h"
+
+#define CADTOOLS_MODULE_NAME TEXT("CADTools")
+
+class FCADToolsModule : public IModuleInterface
+{
+public:
+	static FCADToolsModule& Get();
+	static bool IsAvailable();
+
+private:
+	virtual void StartupModule() override;
+};
diff --git a/Engine/Plugins/Enterprise/DatasmithCADImporter/Source/DatasmithCADTranslator/DatasmithCADTranslator.Build.cs b/Engine/Plugins/Enterprise/DatasmithCADImporter/Source/DatasmithCADTranslator/DatasmithCADTranslator.Build.cs
--- a/Engine/Plugins/Enterprise/DatasmithCADImporter/Source/DatasmithCADTranslator/DatasmithCADTranslator.Build.cs
+++ b/Engine/Plugins/Enterprise/DatasmithCADImporter/Source/DatasmithCADTranslator/DatasmithCADTranslator.Build.cs
@@ -19,22 +19,24 @@ namespace UnrealBuildTool.Rules
 			PrivateDependencyModuleNames.AddRange(
 				new string[]
 				{
-					"Engine",
+					"CADInterfaces",
+					"CADLibrary",
+					"CADTools",
 					"DatasmithCore",
 					"DatasmithCoreTechExtension",
 					"DatasmithContent",
 					"DatasmithImporter",
+					"DatasmithDispatcher",
+					"Engine",
 					"MeshDescription",
-                    "CADLibrary",
-                }
+					"Sockets",
+				}
 			);
 
 			if (System.Type.GetType("CoreTech") != null)
 			{
 				PublicDependencyModuleNames.Add("CoreTech");
 			}
-
-            //PublicDefinitions.Add("USE_CORETECH_MT_PARSER");
-        }
-    }
+		}
+	}
 }
diff --git a/Engine/Plugins/Enterprise/DatasmithCADImporter/Source/DatasmithCADTranslator/Private/CoreTechParser.cpp b/Engine/Plugins/Enterprise/DatasmithCADImporter/Source/DatasmithCADTranslator/Private/CoreTechParser.cpp
deleted file mode 100644
diff --git a/Engine/Plugins/Enterprise/DatasmithCADImporter/Source/DatasmithCADTranslator/Private/CoreTechParser.h b/Engine/Plugins/Enterprise/DatasmithCADImporter/Source/DatasmithCADTranslator/Private/CoreTechParser.h
deleted file mode 100644
diff --git a/Engine/Plugins/Enterprise/DatasmithCADImporter/Source/DatasmithCADTranslator/Private/DatasmithCADTranslator.cpp b/Engine/Plugins/Enterprise/DatasmithCADImporter/Source/DatasmithCADTranslator/Private/DatasmithCADTranslator.cpp
--- a/Engine/Plugins/Enterprise/DatasmithCADImporter/Source/DatasmithCADTranslator/Private/DatasmithCADTranslator.cpp
+++ b/Engine/Plugins/Enterprise/DatasmithCADImporter/Source/DatasmithCADTranslator/Private/DatasmithCADTranslator.cpp
@@ -3,19 +3,16 @@
 #include "DatasmithCADTranslator.h"
 
 #ifdef CAD_LIBRARY
-#include "CADLibraryOptions.h"
 #include "CoreTechParametricSurfaceExtension.h"
 #include "DatasmithCADTranslatorModule.h"
-#include "DatasmithMeshHelper.h"
-#include "DatasmithSceneFactory.h"
-#include "DatasmithUtils.h"
-#include "IDatasmithSceneElements.h"
+#include "DatasmithDispatcher.h"
+#include "DatasmithMeshBuilder.h"
+#include "DatasmithSceneGraphBuilder.h"
 #include "Misc/FileHelper.h"
-#include "Misc/Paths.h"
 
 
 FDatasmithCADTranslator::FDatasmithCADTranslator()
-	: Translator(nullptr)
+	: MeshBuilder(CADFileToUE4GeomMap, MeshElementToCADBRepUuidMap)
 {
 }
 
@@ -67,46 +64,58 @@ bool FDatasmithCADTranslator::IsSourceSupported(const FDatasmithSceneSource& Sou
 
 bool FDatasmithCADTranslator::LoadScene(TSharedRef<IDatasmithScene> DatasmithScene)
 {
-	FString CachePath = FPaths::ConvertRelativePathToFull(FDatasmithCADTranslatorModule::Get().GetCacheDir());
-
-	double FileMetricUnit = 0.001;
-	double ScaleFactor = 0.1;
+	ImportParameters.MetricUnit = 0.001;
+	ImportParameters.ScaleFactor = 0.1;
+	ImportParameters.ChordTolerance = GetCommonTessellationOptions().ChordTolerance;
+	ImportParameters.MaxEdgeLength = GetCommonTessellationOptions().MaxEdgeLength;
+	ImportParameters.MaxNormalAngle = GetCommonTessellationOptions().NormalTolerance;
+	ImportParameters.StitchingTechnique = (CADLibrary::EStitchingTechnique) GetCommonTessellationOptions().StitchingTechnique;
 
 	FString FileExtension = GetSource().GetSourceFileExtension();
 	if (FileExtension == TEXT("jt"))
 	{
-		FileMetricUnit = 1.;
-		ScaleFactor = 100.;
+		ImportParameters.MetricUnit = 1.;
+		ImportParameters.ScaleFactor = 100.;
 	}
 
-	Translator = MakeShared<FDatasmithCADTranslatorImpl>(GetSource(), DatasmithScene, CachePath, FileMetricUnit, ScaleFactor);
-	if (!Translator)
+	ImportParameters.ModelCoordSys = CADLibrary::EModelCoordSystem::ZUp_RightHanded;
+	if (FileExtension == TEXT("sldprt") || FileExtension == TEXT("sldasm"))
 	{
-		return false;
+		ImportParameters.ModelCoordSys = CADLibrary::EModelCoordSystem::YUp_RightHanded;
 	}
 
-	FString OutputPath = FPaths::ConvertRelativePathToFull(FPaths::Combine(FDatasmithCADTranslatorModule::Get().GetTempDir(), TEXT("Cache"), GetSource().GetSceneName()));
-	IFileManager::Get().MakeDirectory(*OutputPath, true);
+	FString CachePath = FPaths::ConvertRelativePathToFull(FDatasmithCADTranslatorModule::Get().GetCacheDir());
+
+	int32 NumCores = FPlatformMisc::NumberOfCores();
+	DatasmithDispatcher::FDatasmithDispatcher Dispatcher(CachePath, NumCores, CADFileToUE4FileMap, CADFileToUE4GeomMap);
+
+	Dispatcher.Init(ImportParameters);
+	Dispatcher.AddTask(FPaths::ConvertRelativePathToFull(GetSource().GetSourceFile()));
 
-	Translator->SetOutputPath(OutputPath);
-	Translator->SetTessellationOptions(GetCommonTessellationOptions());
+	bool bEnableMultiProcess = (NumCores != 1);
+	Dispatcher.Process(bEnableMultiProcess);
+	Dispatcher.Clear();
 
-	return Translator->Read();
+	FDatasmithSceneGraphBuilder SceneGraphBuilder(CADFileToUE4FileMap, MeshElementToCADBRepUuidMap,  CachePath, DatasmithScene, GetSource(), ImportParameters);
+	SceneGraphBuilder.Build();
+	SceneGraphBuilder.Clear();
+
+	MeshBuilder.Init(CachePath);
+	MeshBuilder.SetImportParameters(ImportParameters);
+	MeshBuilder.LoadRawDataGeom();
 
 	return true;
 }
 
 void FDatasmithCADTranslator::UnloadScene()
 {
-	Translator->UnloadScene();
 }
 
 bool FDatasmithCADTranslator::LoadStaticMesh(const TSharedRef<IDatasmithMeshElement> MeshElement, FDatasmithMeshElementPayload& OutMeshPayload)
 {
 	CADLibrary::FMeshParameters MeshParameters;
-	MeshParameters.ModelCoordSys = FDatasmithUtils::EModelCoordSystem::ZUp_RightHanded;
 
-	if (TOptional< FMeshDescription > Mesh = Translator->GetMeshDescription(MeshElement, MeshParameters))
+	if (TOptional< FMeshDescription > Mesh = MeshBuilder.GetMeshDescription(MeshElement, MeshParameters))
 	{
 		OutMeshPayload.LodMeshes.Add(MoveTemp(Mesh.GetValue()));
 
@@ -120,9 +129,9 @@ bool FDatasmithCADTranslator::LoadStaticMesh(const TSharedRef<IDatasmithMeshElem
 				UCoreTechParametricSurfaceData* CoreTechData = Datasmith::MakeAdditionalData<UCoreTechParametricSurfaceData>();
 				CoreTechData->SourceFile = CoretechFile;
 				CoreTechData->RawData = MoveTemp(ByteArray);
-				CoreTechData->SceneParameters.ModelCoordSys = uint8(FDatasmithUtils::EModelCoordSystem::ZUp_RightHanded);
-				CoreTechData->SceneParameters.MetricUnit = 0.001;
-				CoreTechData->SceneParameters.ScaleFactor= 0.1;
+				CoreTechData->SceneParameters.ModelCoordSys = uint8(ImportParameters.ModelCoordSys);
+				CoreTechData->SceneParameters.MetricUnit = ImportParameters.MetricUnit;
+				CoreTechData->SceneParameters.ScaleFactor = ImportParameters.ScaleFactor;
 
 				CoreTechData->MeshParameters.bNeedSwapOrientation = MeshParameters.bNeedSwapOrientation;
 				CoreTechData->MeshParameters.bIsSymmetric = MeshParameters.bIsSymmetric;
@@ -140,11 +149,6 @@ bool FDatasmithCADTranslator::LoadStaticMesh(const TSharedRef<IDatasmithMeshElem
 void FDatasmithCADTranslator::SetSceneImportOptions(TArray<TStrongObjectPtr<UObject>>& Options)
 {
 	FDatasmithCoreTechTranslator::SetSceneImportOptions(Options);
-
-	if (Translator)
-	{
-		Translator->SetTessellationOptions(GetCommonTessellationOptions());
-	}
 }
 #endif
 
diff --git a/Engine/Plugins/Enterprise/DatasmithCADImporter/Source/DatasmithCADTranslator/Private/DatasmithCADTranslator.h b/Engine/Plugins/Enterprise/DatasmithCADImporter/Source/DatasmithCADTranslator/Private/DatasmithCADTranslator.h
--- a/Engine/Plugins/Enterprise/DatasmithCADImporter/Source/DatasmithCADTranslator/Private/DatasmithCADTranslator.h
+++ b/Engine/Plugins/Enterprise/DatasmithCADImporter/Source/DatasmithCADTranslator/Private/DatasmithCADTranslator.h
@@ -4,13 +4,15 @@
 
 #include "CoreMinimal.h"
 
+#include "CADOptions.h"
 #include "DatasmithCoreTechTranslator.h"
-#ifdef CAD_LIBRARY
-#include "CADLibraryOptions.h"
-#include "DatasmithCADTranslatorImpl.h"
+#include "DatasmithDispatcher.h"
+#include "DatasmithMeshBuilder.h"
+#include "DatasmithSceneGraphBuilder.h"
 #include "UObject/ObjectMacros.h"
-#endif
 
+class IDatasmithMeshElement;
+class FCoreTechParser;
 
 class FDatasmithCADTranslator : public FDatasmithCoreTechTranslator
 {
@@ -40,7 +42,13 @@ public:
 	// End ADatasmithCoreTechTranslator overrides
 
 private:
-	TSharedPtr<FDatasmithCADTranslatorImpl> Translator;
+	TMap<FString, FString> CADFileToUE4FileMap;
+	TMap<FString, FString> CADFileToUE4GeomMap;
+	TMap< TSharedPtr< IDatasmithMeshElement >, uint32 > MeshElementToCADBRepUuidMap;
+
+	CADLibrary::FImportParameters ImportParameters;
+
+	FDatasmithMeshBuilder MeshBuilder;
 #endif
 };
 
diff --git a/Engine/Plugins/Enterprise/DatasmithCADImporter/Source/DatasmithCADTranslator/Private/DatasmithCADTranslatorImpl.cpp b/Engine/Plugins/Enterprise/DatasmithCADImporter/Source/DatasmithCADTranslator/Private/DatasmithCADTranslatorImpl.cpp
deleted file mode 100644
diff --git a/Engine/Plugins/Enterprise/DatasmithCADImporter/Source/DatasmithCADTranslator/Private/DatasmithCADTranslatorImpl.h b/Engine/Plugins/Enterprise/DatasmithCADImporter/Source/DatasmithCADTranslator/Private/DatasmithCADTranslatorImpl.h
deleted file mode 100644
diff --git a/Engine/Plugins/Enterprise/DatasmithCADImporter/Source/DatasmithCADTranslator/Private/DatasmithDataHelper.cpp b/Engine/Plugins/Enterprise/DatasmithCADImporter/Source/DatasmithCADTranslator/Private/DatasmithDataHelper.cpp
new file mode 100644
--- /dev/null
+++ b/Engine/Plugins/Enterprise/DatasmithCADImporter/Source/DatasmithCADTranslator/Private/DatasmithDataHelper.cpp
@@ -0,0 +1,111 @@
+// Copyright 1998-2019 Epic Games, Inc. All Rights Reserved.
+#include "DatasmithDataHelper.h"
+
+#include "CADData.h"
+#include "DatasmithSceneFactory.h"
+#include "Math/Color.h"
+
+TSharedPtr<IDatasmithUEPbrMaterialElement> CreateUEPbrMaterialFromColor(FColor& Color)
+{
+	FString Name = FString::FromInt(CADLibrary::BuildColorHash(Color));
+	FString Label = FString::Printf(TEXT("%02x%02x%02x%02x"), Color.R, Color.G, Color.B, Color.A);
+
+	// Take the Material diffuse color and connect it to the BaseColor of a UEPbrMaterial
+	TSharedRef<IDatasmithUEPbrMaterialElement> MaterialElement = FDatasmithSceneFactory::CreateUEPbrMaterial(*Name);
+	MaterialElement->SetLabel(*Label);
+
+	FLinearColor LinearColor = FLinearColor::FromPow22Color(Color);
+
+	IDatasmithMaterialExpressionColor* ColorExpression = MaterialElement->AddMaterialExpression<IDatasmithMaterialExpressionColor>();
+	ColorExpression->SetName(TEXT("Diffuse Color"));
+	ColorExpression->GetColor() = LinearColor;
+
+	MaterialElement->GetBaseColor().SetExpression(ColorExpression);
+	MaterialElement->SetParentLabel(TEXT("CAD Color"));
+
+	if (LinearColor.A < 1.0f)
+	{
+		MaterialElement->SetBlendMode(/*EBlendMode::BLEND_Translucent*/2);
+
+		IDatasmithMaterialExpressionScalar* Scalar = MaterialElement->AddMaterialExpression<IDatasmithMaterialExpressionScalar>();
+		Scalar->GetScalar() = LinearColor.A;
+		Scalar->SetName(TEXT("Opacity Level"));
+
+		MaterialElement->GetOpacity().SetExpression(Scalar);
+		MaterialElement->SetParentLabel(TEXT("CAD Transparent Color"));
+	}
+	return MaterialElement;
+}
+
+TSharedPtr<IDatasmithUEPbrMaterialElement> CreateUEPbrMaterialFromMaterial(CADLibrary::FCADMaterial& Material, TSharedRef<IDatasmithScene> Scene)
+{
+	FString Name = FString::FromInt(CADLibrary::BuildMaterialHash(Material));
+
+	// Take the Material diffuse color and connect it to the BaseColor of a UEPbrMaterial
+	TSharedRef<IDatasmithUEPbrMaterialElement> MaterialElement = FDatasmithSceneFactory::CreateUEPbrMaterial(*Name);
+	FString MaterialLabel = Material.MaterialName;
+	if (MaterialLabel.IsEmpty())
+	{
+		MaterialLabel = TEXT("Material");
+	}
+	MaterialElement->SetLabel(*MaterialLabel);
+
+	// Set a diffuse color if there's nothing in the BaseColor
+	if (MaterialElement->GetBaseColor().GetExpression() == nullptr)
+	{
+		FLinearColor LinearColor = FLinearColor::FromPow22Color(Material.Diffuse);
+
+		IDatasmithMaterialExpressionColor* ColorExpression = MaterialElement->AddMaterialExpression<IDatasmithMaterialExpressionColor>();
+		ColorExpression->SetName(TEXT("Diffuse Color"));
+		ColorExpression->GetColor() = LinearColor;
+
+		MaterialElement->GetBaseColor().SetExpression(ColorExpression);
+	}
+
+	if (Material.Transparency > 0.0)
+	{
+		MaterialElement->SetBlendMode(/*EBlendMode::BLEND_Translucent*/2);
+		IDatasmithMaterialExpressionScalar* Scalar = MaterialElement->AddMaterialExpression<IDatasmithMaterialExpressionScalar>();
+		Scalar->GetScalar() = Material.Transparency;
+		MaterialElement->GetOpacity().SetExpression(Scalar);
+	}
+
+	// Set a Emissive color 
+	if (MaterialElement->GetEmissiveColor().GetExpression() == nullptr)
+	{
+		// Doc CT => TODO
+		//GLfloat Specular[4] = { specular.rgb[0] / 255., specular.rgb[1] / 255., specular.rgb[2] / 255., 1. - transparency };
+		//GLfloat Shininess[1] = { (float)(128 * shininess) };
+		//glMaterialfv(GL_FRONT_AND_BACK, GL_SPECULAR, Specular);
+		//glMaterialfv(GL_FRONT_AND_BACK, GL_SHININESS, Shininess);
+
+		FLinearColor LinearColor = FLinearColor::FromPow22Color(Material.Specular);
+
+		IDatasmithMaterialExpressionColor* ColorExpression = MaterialElement->AddMaterialExpression<IDatasmithMaterialExpressionColor>();
+		ColorExpression->SetName(TEXT("Specular Color"));
+		ColorExpression->GetColor() = LinearColor;
+
+		MaterialElement->GetEmissiveColor().SetExpression(ColorExpression);
+	}
+
+	// Simple conversion of shininess and reflectivity to PBR roughness and metallic values; model could be improved to properly blend the values
+	if (!FMath::IsNearlyZero(Material.Shininess))
+	{
+		IDatasmithMaterialExpressionScalar* Scalar = static_cast<IDatasmithMaterialExpressionScalar*>(MaterialElement->AddMaterialExpression(EDatasmithMaterialExpressionType::ConstantScalar));
+		Scalar->GetScalar() = 1.f - Material.Shininess;
+		Scalar->SetName(TEXT("Shininess"));
+		MaterialElement->GetRoughness().SetExpression(Scalar);
+	}
+
+	if (!FMath::IsNearlyZero(Material.Reflexion))
+	{
+		IDatasmithMaterialExpressionScalar* Scalar = static_cast<IDatasmithMaterialExpressionScalar*>(MaterialElement->AddMaterialExpression(EDatasmithMaterialExpressionType::ConstantScalar));
+		Scalar->GetScalar() = Material.Reflexion;
+		Scalar->SetName(TEXT("Reflexion"));
+		MaterialElement->GetMetallic().SetExpression(Scalar);
+	}
+
+	return MaterialElement;
+}
+
+
diff --git a/Engine/Plugins/Enterprise/DatasmithCADImporter/Source/DatasmithCADTranslator/Private/DatasmithDataHelper.h b/Engine/Plugins/Enterprise/DatasmithCADImporter/Source/DatasmithCADTranslator/Private/DatasmithDataHelper.h
new file mode 100644
--- /dev/null
+++ b/Engine/Plugins/Enterprise/DatasmithCADImporter/Source/DatasmithCADTranslator/Private/DatasmithDataHelper.h
@@ -0,0 +1,15 @@
+// Copyright 1998-2019 Epic Games, Inc. All Rights Reserved.
+#pragma once
+
+#include "CoreMinimal.h"
+
+#include "CADData.h"
+#include "MeshTypes.h"
+
+struct FColor;
+class IDatasmithUEPbrMaterialElement;
+class IDatasmithScene;
+
+TSharedPtr<IDatasmithUEPbrMaterialElement> CreateDefaultUEPbrMaterial();
+TSharedPtr<IDatasmithUEPbrMaterialElement> CreateUEPbrMaterialFromColor(FColor& Color);
+TSharedPtr<IDatasmithUEPbrMaterialElement> CreateUEPbrMaterialFromMaterial(CADLibrary::FCADMaterial& Material, TSharedRef<IDatasmithScene> Scene);
diff --git a/Engine/Plugins/Enterprise/DatasmithCADImporter/Source/DatasmithCADTranslator/Private/DatasmithMeshBuilder.cpp b/Engine/Plugins/Enterprise/DatasmithCADImporter/Source/DatasmithCADTranslator/Private/DatasmithMeshBuilder.cpp
--- a/Engine/Plugins/Enterprise/DatasmithCADImporter/Source/DatasmithCADTranslator/Private/DatasmithMeshBuilder.cpp
+++ b/Engine/Plugins/Enterprise/DatasmithCADImporter/Source/DatasmithCADTranslator/Private/DatasmithMeshBuilder.cpp
@@ -1,39 +1,30 @@
 // Copyright 1998-2019 Epic Games, Inc. All Rights Reserved.
 #include "DatasmithMeshBuilder.h"
 
-#ifdef USE_CORETECH_MT_PARSER
-
 #include "CoreTechHelper.h"
+#include "CoreTechFileParser.h"
 #include "DatasmithMeshHelper.h"
+#include "HAL/FileManager.h"
 #include "IDatasmithSceneElements.h"
 #include "MeshDescription.h"
-#include "Misc/FileHelper.h"
 #include "Misc/Paths.h"
-
 using namespace CADLibrary;
 
-FCoretechBody::FCoretechBody(uint32 InBodyID)
-	: TriangleCount(0)
-	, BodyID(InBodyID)
+FDatasmithMeshBuilder::FDatasmithMeshBuilder(TMap<FString, FString>& InCADFileToUE4GeomMap, TMap< TSharedPtr< IDatasmithMeshElement >, uint32 >& InMeshElementToCTBodyUuidMap)
+	: CADFileToUE4GeomMap(InCADFileToUE4GeomMap)
+	, MeshElementToBodyUuidMap(InMeshElementToCTBodyUuidMap)
 {
 }
 
-void FCoretechBody::Add(FCTTessellation* InFaceTessellation)
+void FDatasmithMeshBuilder::Init(const FString& InCachePath)
 {
-	if (InFaceTessellation == nullptr)
-	{
-		return;
-	}
-	FaceTessellationSet.Add(InFaceTessellation);
-	TriangleCount += InFaceTessellation->IndexCount / 3;
-	MaterialToPartition.LinkMaterialId2MaterialHash(InFaceTessellation->MaterialId, InFaceTessellation->MaterialHash);
+	CachePath = InCachePath;
 }
 
-FDatasmithMeshBuilder::FDatasmithMeshBuilder(const FString& InCachePath, TMap<FString, FString>& InCADFileToUE4GeomMap, TMap< TSharedPtr< IDatasmithMeshElement >, uint32 >& InMeshElementToCTBodyUuidMap)
-	: CachePath(InCachePath)
-	, CADFileToUE4GeomMap(InCADFileToUE4GeomMap)
-	, MeshElementToCTBodyUuidMap(InMeshElementToCTBodyUuidMap)
+void FDatasmithMeshBuilder::Clear()
 {
+	RawDataArray.Empty();
+	BodyUuidToCADBRepMap.Empty();
 }
 
 void FDatasmithMeshBuilder::LoadRawDataGeom()
@@ -47,61 +38,35 @@ void FDatasmithMeshBuilder::LoadRawDataGeom()
 			continue;
 		}
 
-		int32 index = RawDataArray.Emplace(RawDataFile, BodyUuidToCoretechBodyMap);
-	}
-}
-
-FCTRawGeomFile::FCTRawGeomFile(FString& InFileName, TMap< uint32, FCoretechBody* >& InBodyUuidToCTBodyMap)
-	: BodyUuidToCTBodyMap(InBodyUuidToCTBodyMap)
-{
-	uint32 NbBody = 0;
-
-	FFileHelper::LoadFileToArray(RawData, *InFileName);
-	ReadTessellationSetFromFile(RawData, NbBody, FaceTessellationSet);
-	CTBodySet.Reserve(NbBody);
-
-	for (int index = 0; index < FaceTessellationSet.Num(); ++index)
-	{
-		uint32* PBodyIndex = CTIdToBodyMap.Find(FaceTessellationSet[index].BodyId);
-		if (!PBodyIndex)
-		{
-			uint32 BodyIndex = CTBodySet.Emplace(FaceTessellationSet[index].BodyId);
-			PBodyIndex = &BodyIndex;
-			CTIdToBodyMap.Emplace(FaceTessellationSet[index].BodyId, BodyIndex);
-			BodyUuidToCTBodyMap.Emplace(FaceTessellationSet[index].BodyUuId, &CTBodySet[BodyIndex]);
-		}
-		CTBodySet[*PBodyIndex].Add(&FaceTessellationSet[index]);
+		int32 index = RawDataArray.Emplace(RawDataFile, BodyUuidToCADBRepMap);
 	}
 }
 
 TOptional<FMeshDescription> FDatasmithMeshBuilder::GetMeshDescription(TSharedRef<IDatasmithMeshElement> OutMeshElement, CADLibrary::FMeshParameters& OutMeshParameters)
 {
-	uint32* BodyUuid = MeshElementToCTBodyUuidMap.Find(OutMeshElement);
+	uint32* BodyUuid = MeshElementToBodyUuidMap.Find(OutMeshElement);
 	if (BodyUuid == nullptr)
 	{
 		return TOptional<FMeshDescription>();
 	}
 
-	FCoretechBody*const* PPBody = BodyUuidToCoretechBodyMap.Find(*BodyUuid);
+	FBody*const* PPBody = BodyUuidToCADBRepMap.Find(*BodyUuid);
 	if(PPBody == nullptr || *PPBody == nullptr)
 	{
 		return TOptional<FMeshDescription>();
 	}
 
-	FCoretechBody& Body = **PPBody;
+	FBody& Body = **PPBody;
 
-	FCTMaterialPartition Material2Partition;
+	TMap<uint32, uint32> MaterialIdToMaterialHash;
 
 	FMeshDescription MeshDescription;
 	DatasmithMeshHelper::PrepareAttributeForStaticMesh(MeshDescription);
 
-	if(!ConvertCTBodySetToMeshDescription(ScaleFactor, OutMeshParameters, Body.GetTriangleCount(), Body.GetTessellationSet(), Material2Partition, MeshDescription))
+	if (!ConvertCTBodySetToMeshDescription(ImportParameters, OutMeshParameters, Body.GetTriangleCount(), Body.GetTessellationSet(), MaterialIdToMaterialHash, MeshDescription))
 	{
 		return TOptional<FMeshDescription>();
 	}
 
 	return MoveTemp(MeshDescription);
 }
-
-
-#endif  // USE_CORETECH_MT_PARSER
\ No newline at end of file
diff --git a/Engine/Plugins/Enterprise/DatasmithCADImporter/Source/DatasmithCADTranslator/Private/DatasmithMeshBuilder.h b/Engine/Plugins/Enterprise/DatasmithCADImporter/Source/DatasmithCADTranslator/Private/DatasmithMeshBuilder.h
--- a/Engine/Plugins/Enterprise/DatasmithCADImporter/Source/DatasmithCADTranslator/Private/DatasmithMeshBuilder.h
+++ b/Engine/Plugins/Enterprise/DatasmithCADImporter/Source/DatasmithCADTranslator/Private/DatasmithMeshBuilder.h
@@ -1,117 +1,48 @@
 // Copyright 1998-2019 Epic Games, Inc. All Rights Reserved.
 #pragma once
 
-#ifdef USE_CORETECH_MT_PARSER
-
 #include "CoreMinimal.h"
 
-#include "CADLibraryOptions.h"
+#include "CADData.h"
 #include "CoretechHelper.h"
 
 
-class FCTRawGeomFile;
 class IDatasmithMeshElement;
 
 struct FMeshParameters;
 
 
-class FCoretechBody
-{
-public:
-	FCoretechBody(uint32 BodyID);
-	void Add(CADLibrary::FCTTessellation* InFaceTessellation);
-	TArray<CADLibrary::FCTTessellation*>& GetTessellationSet()
-	{
-		return FaceTessellationSet;
-	}
-
-	uint32 GetTriangleCount()
-	{
-		return TriangleCount;
-	}
-
-protected:
-
-	/** 
-	 * Array of FCTTessellation
-	 * FCTTessellation is an elementary structure of Mesh
-	 */
-	TArray<CADLibrary::FCTTessellation*> FaceTessellationSet;
-
-	/**
-	 * Structure use to map Material ID to Material Hash. Material hash is an unique value by material 
-	 */
-	CADLibrary::FCTMaterialPartition MaterialToPartition;
-	uint32 TriangleCount;
-	uint32 BodyID;
-};
-
-class FCTRawGeomFile
-{
-public:
-	FCTRawGeomFile(FString& InFileName, TMap< uint32, FCoretechBody* >& InBodyUuidToCTBodyMap);
-
-	FCoretechBody* GetCTNode(int32 nodeId)
-	{
-		uint32* index = CTIdToBodyMap.Find(nodeId);
-		if (index == nullptr)
-		{
-			return nullptr;
-		}
-		return &CTBodySet[*index];
-	}
-
-	FString& GetFileName()
-	{
-		return FileName;
-	}
-
-protected:
-	FString FileName;
-	TArray<uint8> RawData;
-	TArray<FCoretechBody> CTBodySet;
-	TMap<uint32, uint32> CTIdToBodyMap;
-	TMap< uint32, FCoretechBody* >& BodyUuidToCTBodyMap;
-
-	/**
-	 * A body has a set of face. A face mesh is defined by is FCTTessellation
-	 */
-	TArray<CADLibrary::FCTTessellation> FaceTessellationSet;
-};
-
-
-
-
-
-
 class FDatasmithMeshBuilder
 {
 public:
-	FDatasmithMeshBuilder(const FString& InCachePath, TMap<FString, FString>& InCADFileToUE4GeomMap, TMap< TSharedPtr< IDatasmithMeshElement >, uint32 >& InMeshElementToCTBodyUuidMap);
+	FDatasmithMeshBuilder(TMap<FString, FString>& InCADFileToUE4GeomMap, TMap< TSharedPtr< IDatasmithMeshElement >, uint32 >& InMeshElementToCTBodyUuidMap);
+	
+	void Init(const FString& InCachePath);
+
 	TOptional<FMeshDescription> GetMeshDescription(TSharedRef<IDatasmithMeshElement> OutMeshElement, CADLibrary::FMeshParameters& OutMeshParameters);
 	void LoadRawDataGeom();
+	void Clear();
 
-	void SetScaleFactor(float InScaleFactor)
+	void SetImportParameters(const CADLibrary::FImportParameters& InImportParameters)
 	{
-		ScaleFactor = InScaleFactor;
+		ImportParameters = InImportParameters;
 	}
 
 protected:
 	//const FDatasmithSceneSource& Source;
 	FString CachePath;
 
-	TArray<FCTRawGeomFile> RawDataArray;
+	TArray<CADLibrary::FCTRawGeomFile> RawDataArray;
 
 	/** Map linking Cad file to RawGeom file (*.gm) */
 	TMap<FString, FString>& CADFileToUE4GeomMap;
 
 	/** Datasmith mesh elements to BodyUuid */
-	TMap< TSharedPtr< IDatasmithMeshElement >, uint32 >& MeshElementToCTBodyUuidMap;
+	TMap< TSharedPtr< IDatasmithMeshElement >, uint32 >& MeshElementToBodyUuidMap;
 
 	/** BodyUuid to CTBody */
-	TMap< uint32, FCoretechBody* > BodyUuidToCoretechBodyMap;
-	float ScaleFactor;
-};
+	TMap< uint32, CADLibrary::FBody* > BodyUuidToCADBRepMap;
 
+	CADLibrary::FImportParameters ImportParameters;
+};
 
-#endif // USE_CORETECH_MT_PARSER
diff --git a/Engine/Plugins/Enterprise/DatasmithCADImporter/Source/DatasmithCADTranslator/Private/DatasmithSceneGraphBuilder.cpp b/Engine/Plugins/Enterprise/DatasmithCADImporter/Source/DatasmithCADTranslator/Private/DatasmithSceneGraphBuilder.cpp
--- a/Engine/Plugins/Enterprise/DatasmithCADImporter/Source/DatasmithCADTranslator/Private/DatasmithSceneGraphBuilder.cpp
+++ b/Engine/Plugins/Enterprise/DatasmithCADImporter/Source/DatasmithCADTranslator/Private/DatasmithSceneGraphBuilder.cpp
@@ -1,17 +1,51 @@
 // Copyright 1998-2019 Epic Games, Inc. All Rights Reserved.
 #include "DatasmithSceneGraphBuilder.h"
 
-#ifdef USE_CORETECH_MT_PARSER
-
+#include "CADData.h"
+#include "CoreTechFileParser.h"
 #include "CoreTechHelper.h"
-#include "CoreTechParserMt.h"
 #include "DatasmithSceneFactory.h"
 #include "DatasmithSceneSource.h"
 #include "IDatasmithSceneElements.h"
-#include "Misc/FileHelper.h"
 #include "Utility/DatasmithMathUtils.h"
 
-using namespace CADLibrary;
+#include "Misc/FileHelper.h"
+
+
+namespace
+{
+	void GetMainMaterial(const TMap<FString, FString>& InNodeMetaDataMap, FCTNode* InNode, ActorData& OutNodeData, bool bMaterialPropagationIsTopDown)
+	{
+		if (const FString* MaterialIdStr = InNodeMetaDataMap.Find(TEXT("MaterialId")))
+		{
+			if (!bMaterialPropagationIsTopDown || !OutNodeData.MaterialUuid)
+			{
+				uint32 MaterialId = FCString::Atoi64(**MaterialIdStr);
+				CADLibrary::FCADMaterial Material;
+				if (InNode->GetMaterial(MaterialId, Material))
+				{
+					OutNodeData.Material = Material;
+					OutNodeData.MaterialUuid = CADLibrary::BuildMaterialHash(OutNodeData.Material);
+				}
+			}
+		}
+
+		if (const FString* ColorIdStr = InNodeMetaDataMap.Find(TEXT("ColorUEId")))
+		{
+			if (bMaterialPropagationIsTopDown || !OutNodeData.ColorUuid)
+			{
+				uint32 ColorHId = FCString::Atoi64(**ColorIdStr);
+				FColor Color;
+				if (InNode->GetColor(ColorHId, Color))
+				{
+					OutNodeData.Color = Color;
+					OutNodeData.ColorUuid = CADLibrary::BuildColorHash(OutNodeData.Color);
+				}
+			}
+		}
+	}
+}
+
 
 FCTRawDataFile::FCTRawDataFile(const FString& InFileName) 
 	: FileName(InFileName)
@@ -19,6 +53,82 @@ FCTRawDataFile::FCTRawDataFile(const FString& InFileName)
 	ReadFile();
 }
 
+void FCTRawDataFile::GetMaterialDescription(int32 LineNumber, CADLibrary::FCADMaterial& Material) const
+{
+	FString OutMaterialName;
+	uint8 OutDiffuse[3];
+	uint8 OutAmbient[3];
+	uint8 OutSpecular[3];
+	float OutShininess, OutTransparency, OutReflexion;
+	FString OutTextureName;
+
+	OutDiffuse[0] = OutDiffuse[1] = OutDiffuse[2] = 255;
+	OutAmbient[0] = OutAmbient[1] = OutAmbient[2] = 255;
+	OutSpecular[0] = OutSpecular[1] = OutSpecular[2] = 255;
+	OutShininess = OutReflexion = 0;
+	OutTransparency = 255;
+	OutMaterialName = TEXT("");
+	OutTextureName = TEXT("");
+
+	const FString& MaterialLine1 = GetString(LineNumber);
+	FString MaterialIdStr;
+	MaterialLine1.Split(SPACE, &MaterialIdStr, &OutMaterialName);
+
+	const FString& MaterialLine2 = GetString(LineNumber+1);
+	TArray<FString> MaterialParameters;
+	MaterialLine2.ParseIntoArray(MaterialParameters, *SPACE);
+
+	for (int32 index = 0; index < 3; index++)
+	{
+		OutDiffuse[index] = FCString::Atoi(*MaterialParameters[index + 2]);
+		OutAmbient[index] = FCString::Atoi(*MaterialParameters[index + 5]);
+		OutSpecular[index] = FCString::Atoi(*MaterialParameters[index + 8]);
+	}
+	OutShininess = FCString::Atof(*MaterialParameters[11]);
+	OutTransparency = FCString::Atof(*MaterialParameters[12]);
+	OutReflexion = FCString::Atof(*MaterialParameters[13]);
+	if (MaterialParameters.Num() > 14)
+	{
+		OutTextureName = MaterialParameters[14];
+	}
+
+	Material.MaterialId = FCString::Atoi(*MaterialParameters[0]);
+	Material.MaterialName = OutMaterialName;
+	Material.Diffuse= FColor(OutDiffuse[0], OutDiffuse[1], OutDiffuse[2]);
+	Material.Ambient = FColor(OutAmbient[0], OutAmbient[1], OutAmbient[2]);
+	Material.Specular = FColor(OutSpecular[0], OutSpecular[1], OutSpecular[2]);
+	Material.Shininess = OutShininess;
+	Material.Transparency = OutTransparency;
+	Material.Reflexion = OutReflexion;
+	Material.TextureName = OutTextureName;
+}
+
+bool FCTRawDataFile::GetColor(uint32 ColorHId, FColor& Color) const
+{
+	uint32 ColorId;
+	uint8 Alpha;
+	CADLibrary::UnhashFastColorHash(ColorHId, ColorId, Alpha);
+	const FColor* ColorOpac = ColorIdToColor.Find(ColorId);
+	if (ColorOpac)
+	{
+		Color = *ColorOpac;
+		Color.A = Alpha;
+		return true;
+	}
+	return false;
+}
+
+bool FCTRawDataFile::GetMaterial(int32 MaterialId, CADLibrary::FCADMaterial& Material) const
+{
+	const CADLibrary::FCADMaterial* PMaterial = MaterialIdToMaterial.Find(MaterialId);
+	if (PMaterial)
+	{
+		Material = *PMaterial;
+		return true;
+	}
+	return false;
+}
+
 void FCTRawDataFile::ReadFile()
 {
 	FFileHelper::LoadFileToStringArray(RawData, *FileName);
@@ -40,7 +150,6 @@ void FCTRawDataFile::ReadFile()
 		CTIdToRawEntryMap.Emplace(FCString::Atoi(*HeaderCTIdSplit[index]), FCTNode(*this, FCString::Atoi(*HeaderCTIdSplit[index + 1])));
 	}
 
-
 	// Parse color
 	HeaderCTId = RawData[COLORSETLINE];
 	HeaderCTId.Split(SPACE, &Left, &Right);
@@ -50,8 +159,9 @@ void FCTRawDataFile::ReadFile()
 		FString ColorString = RawData[Line];
 		while (!ColorString.IsEmpty())
 		{
-			ColorString.Split(SPACE, &Left, &Right);
-			ColorIdToLineMap.Add(FCString::Atoi(*Left), Line);
+			TArray<FString> ColorData;
+			ColorString.ParseIntoArray(ColorData, TEXT(" "));
+			ColorIdToColor.Emplace(FCString::Atoi(*ColorData[0]), FColor(FCString::Atoi(*ColorData[1]), FCString::Atoi(*ColorData[2]), FCString::Atoi(*ColorData[3])));
 			Line++;
 			ColorString = RawData[Line];
 		}
@@ -66,8 +176,9 @@ void FCTRawDataFile::ReadFile()
 		FString MaterialString = RawData[Line];
 		while (!MaterialString.IsEmpty())
 		{
-			MaterialString.Split(SPACE, &Left, &Right);
-			MaterialIdToLineMap.Add(FCString::Atoi(*Left), Line);
+			CADLibrary::FCADMaterial Material;
+			GetMaterialDescription(Line, Material);
+			MaterialIdToMaterial.Add(Material.MaterialId, Material);
 			Line += 2;
 			MaterialString = RawData[Line];
 		}
@@ -75,6 +186,21 @@ void FCTRawDataFile::ReadFile()
 
 }
 
+void FCTRawDataFile::SetMaterialMaps(TMap< uint32, FColor>& MaterialUuidToColor, TMap< uint32, CADLibrary::FCADMaterial>& MaterialUuidToMaterial)
+{
+	for (const auto& Color : ColorIdToColor)
+	{
+		uint32 ColorHash = CADLibrary::BuildColorHash(Color.Value);
+		MaterialUuidToColor.Add(ColorHash, Color.Value);
+	}
+	for (const auto& Material : MaterialIdToMaterial)
+	{
+		uint32 MaterialHash = CADLibrary::BuildMaterialHash(Material.Value);
+		MaterialUuidToMaterial.Add(MaterialHash, Material.Value);
+	}
+}
+
+
 void FCTNode::GetMetaDatas(TMap<FString, FString>& MetaDataMap)
 {
 	const FString& RawDataString = RawData.GetString(Line+1);
@@ -95,7 +221,7 @@ void FCTNode::GetMetaDatas(TMap<FString, FString>& MetaDataMap)
 	}
 }
 
-void FCTNode::GetMaterialSet(FCTMaterialPartition& MaterialPartition)
+void FCTNode::GetMaterialSet(TMap<uint32, uint32>& MaterialIdToMaterialHashMap)
 {
 	const FString& MaterialString = RawData.GetString(EndMeta);
 	TArray<FString> MaterialIDToHashSet;
@@ -110,7 +236,7 @@ void FCTNode::GetMaterialSet(FCTMaterialPartition& MaterialPartition)
 		uint32 MaterialId = FCString::Atoi64(*MaterialIDToHashSet[index]); 
 		index++;
 		uint32 MaterialHash = FCString::Atoi64(*MaterialIDToHashSet[index]);
-		MaterialPartition.LinkMaterialId2MaterialHash(MaterialId, MaterialHash);
+		MaterialIdToMaterialHashMap.Add(MaterialId, MaterialHash);
 	}
 }
 
@@ -176,7 +302,7 @@ void FCTNode::GetNodeReference(int32& RefId, FString& ExternalFile, NODE_TYPE& N
 	}
 }
 
-void FCTNode::AddTransformToActor(TSharedPtr< IDatasmithActorElement > Actor, CADLibrary::FImportParameters& ImportParameters)
+void FCTNode::AddTransformToActor(TSharedPtr< IDatasmithActorElement > Actor, const CADLibrary::FImportParameters& ImportParameters)
 {
 	if (!Actor.IsValid())
 	{
@@ -205,7 +331,8 @@ void FCTNode::AddTransformToActor(TSharedPtr< IDatasmithActorElement > Actor, CA
 	}
 
 	FTransform LocalTransform(Matrix);
-	FTransform LocalUETransform = FDatasmithUtils::ConvertTransform(FDatasmithUtils::EModelCoordSystem::ZUp_RightHanded, LocalTransform);
+	FTransform LocalUETransform = FDatasmithUtils::ConvertTransform((FDatasmithUtils::EModelCoordSystem) ImportParameters.ModelCoordSys, LocalTransform);
+
 	FQuat Quat;
 	FDatasmithTransformUtils::GetRotation(LocalUETransform, Quat);
 	
@@ -254,14 +381,14 @@ uint32 FCTNode::GetStaticMeshUuid()
 	return BodyUUID;
 }
 
-void FCTNode::GetColorDescription(int32 ColorId, uint8 CtColor[3])
+bool FCTNode::GetColor(int32 ColorHId, FColor& OutColor) const
 {
-	RawData.GetColorDescription(ColorId, CtColor);
+	return RawData.GetColor(ColorHId, OutColor);
 }
 
-void FCTNode::GetMaterialDescription(int32 MaterialId, int32& OutMaterialHash, FString& OutMaterialName, uint8 OutDiffuse[3], uint8 OutAmbient[3], uint8 OutSpecular[3], uint8& OutShininess, uint8& OutTransparency, uint8& OutReflexion, FString& OutTextureName)
+bool FCTNode::GetMaterial(int32 MaterialId, CADLibrary::FCADMaterial& OutMaterial) const
 {
-	RawData.GetMaterialDescription(MaterialId, OutMaterialHash, OutMaterialName, OutDiffuse, OutAmbient, OutSpecular, OutShininess, OutTransparency, OutReflexion, OutTextureName);
+	return RawData.GetMaterial(MaterialId, OutMaterial);
 }
 
 FCTNode* FCTNode::GetCTNode(int32 NodeId)
@@ -280,12 +407,13 @@ FCTNode::FCTNode(FCTRawDataFile& InRawData, int32 InLine)
 	SetNodeType();
 }
 
-FDatasmithSceneGraphBuilder::FDatasmithSceneGraphBuilder(const FString& InCachePath, TSharedRef<IDatasmithScene> InScene, const FDatasmithSceneSource& InSource, TMap<FString, FString>& InCADFileToUE4FileMap, TMap< TSharedPtr< IDatasmithMeshElement >, uint32 >& InMeshElementToCTBodyUuidMap)
-	: DatasmithScene(InScene)
-	, Source(InSource)
-	, CachePath(InCachePath)
-	, CADFileToRawDataFile(InCADFileToUE4FileMap)
+FDatasmithSceneGraphBuilder::FDatasmithSceneGraphBuilder(TMap<FString, FString>& InCADFileToUE4FileMap, TMap< TSharedPtr< IDatasmithMeshElement >, uint32 >& InMeshElementToCTBodyUuidMap, const FString& InCachePath, TSharedRef<IDatasmithScene> InScene, const FDatasmithSceneSource& InSource, const CADLibrary::FImportParameters& InImportParameters)
+	: CADFileToRawDataFile(InCADFileToUE4FileMap)
 	, MeshElementToCTBodyUuidMap(InMeshElementToCTBodyUuidMap)
+	, CachePath(InCachePath)
+	, DatasmithScene(InScene)
+	, Source(InSource)
+	, ImportParameters(InImportParameters)
 {
 
 }
@@ -308,12 +436,20 @@ bool FDatasmithSceneGraphBuilder::Build()
 		return false;
 	}
 
-	TSharedPtr< IDatasmithActorElement > RootActor = BuildNode(*RootNode, TEXT(""));
+	ActorData Data(TEXT(""));
+	TSharedPtr< IDatasmithActorElement > RootActor = BuildNode(*RootNode, Data);
 	DatasmithScene->AddActor(RootActor);
 
 	return true;
 }
 
+void FDatasmithSceneGraphBuilder::Clear()
+{
+	CADFileToRawData.Empty();
+	RawDataArray.Empty();
+	BodyUuidToMeshElementMap.Empty();
+}
+
 void FDatasmithSceneGraphBuilder::LoadRawDataFile()
 {
 	RawDataArray.Reserve(CADFileToRawDataFile.Num());
@@ -323,27 +459,25 @@ void FDatasmithSceneGraphBuilder::LoadRawDataFile()
 		//OutSgFile = FilePair.Value;
 		FString RawDataFile = FPaths::Combine(CachePath, TEXT("scene"), FilePair.Value + TEXT(".sg"));
 		uint32 index = RawDataArray.Emplace(RawDataFile);
+		RawDataArray[index].SetMaterialMaps(MaterialUuidToColor, MaterialUuidToMaterial);
 		CADFileToRawData.Emplace(FilePair.Key, &RawDataArray[index]);
 	}
 }
 
-TSharedPtr< IDatasmithActorElement >  FDatasmithSceneGraphBuilder::BuildNode(FCTNode& Node, const TCHAR* ParentUuid)
+TSharedPtr< IDatasmithActorElement >  FDatasmithSceneGraphBuilder::BuildNode(FCTNode& Node, ActorData& ParentData)
 {
 	NODE_TYPE Type = Node.GetNodeType();
 
 	switch (Type)
 	{
 	case INSTANCE:
-		return BuildInstance(Node, ParentUuid);
+		return BuildInstance(Node, ParentData);
 
 	case COMPONENT:
-		return BuildComponent(Node, ParentUuid);
-
-	//case EXTERNALCOMPONENT:
-		//return BuildExternalComponent(Node);
+		return BuildComponent(Node, ParentData);
 
 	case BODY:
-		return BuildBody(Node, ParentUuid);
+		return BuildBody(Node, ParentData);
 	default:
 		return TSharedPtr< IDatasmithActorElement >();
 	}
@@ -351,7 +485,7 @@ TSharedPtr< IDatasmithActorElement >  FDatasmithSceneGraphBuilder::BuildNode(FCT
 	return TSharedPtr< IDatasmithActorElement >();
 }
 
-TSharedPtr< IDatasmithActorElement >  FDatasmithSceneGraphBuilder::BuildInstance(FCTNode& Node, const TCHAR* ParentUuid)
+TSharedPtr< IDatasmithActorElement >  FDatasmithSceneGraphBuilder::BuildInstance(FCTNode& Node, ActorData& ParentData)
 {
 	TMap<FString, FString> InstanceNodeMetaDataMap;
 	TMap<FString, FString> ReferenceNodeMetaDataMap;
@@ -360,7 +494,7 @@ TSharedPtr< IDatasmithActorElement >  FDatasmithSceneGraphBuilder::BuildInstance
 	FString ActorUUID = TEXT("");
 	FString ActorLabel = TEXT("");
 
-	GetNodeUUIDAndName(InstanceNodeMetaDataMap, ReferenceNodeMetaDataMap, ParentUuid, ActorUUID, ActorLabel);
+	GetNodeUUIDAndName(InstanceNodeMetaDataMap, ReferenceNodeMetaDataMap, ParentData.Uuid, ActorUUID, ActorLabel);
 
 	int32 RefId; 
 	NODE_TYPE NodeType;
@@ -408,77 +542,20 @@ TSharedPtr< IDatasmithActorElement >  FDatasmithSceneGraphBuilder::BuildInstance
 	}
 	AddMetaData(Actor, InstanceNodeMetaDataMap, ReferenceNodeMetaDataMap);
 
+	ActorData InstanceData(*ActorUUID, ParentData);
+
+	GetMainMaterial(InstanceNodeMetaDataMap, ComponentNode, InstanceData, bMaterialPropagationIsTopDown);
+	GetMainMaterial(ReferenceNodeMetaDataMap, ComponentNode, InstanceData, bMaterialPropagationIsTopDown);
+
 	if (ComponentNode)
 	{
-		AddChildren(Actor, *ComponentNode, *ActorUUID);
+		AddChildren(Actor, *ComponentNode, InstanceData);
 	}
 
 	Node.AddTransformToActor(Actor, ImportParameters);
 	return Actor;
 }
 
-//void FDatasmithSceneGraphBuilder::SetNodeParameterFromAttribute(bool bIsBody)
-//{
-//	FString* IName = InstanceNodeAttributeSet.Find(TEXT("CTName"));
-//	FString* IOriginalName = InstanceNodeAttributeSet.Find(TEXT("Name"));
-//	FString* IUUID = InstanceNodeAttributeSet.Find(TEXT("UUID"));
-//
-//	FString* RName = ReferenceNodeAttributeSet.Find(TEXT("CTName"));
-//	FString* ROriginalName = ReferenceNodeAttributeSet.Find(TEXT("Name"));
-//	FString* RUUID = ReferenceNodeAttributeSet.Find(TEXT("UUID"));
-//
-//	// Reference Name
-//	if (ROriginalName)
-//	{
-//		ReferenceName = *ROriginalName;
-//	}
-//	else if (RName)
-//	{
-//		ReferenceName = *RName;
-//	}
-//	else
-//	{
-//		ReferenceName = "NoName";
-//	}
-//
-//	//ReferenceInstanceName = ReferenceName + TEXT("(") + (IOriginalName ? *IOriginalName : ReferenceName) + TEXT(")");
-//	ReferenceInstanceName = IOriginalName ? *IOriginalName : IName ? *IName : ReferenceName;
-//
-//	//// UUID
-//	if (bIsBody)
-//	{
-//		BuildMeshActorUUID();
-//	}
-//	else
-//	{
-//		UEUUID = 0;
-//		if (Parent.IsValid())
-//		{
-//			UEUUID = Parent->GetUUID();
-//		}
-//
-//		if (IUUID)
-//		{
-//			UEUUID = HashCombine(UEUUID, GetTypeHash(*IUUID));
-//		}
-//		else if (IOriginalName)
-//		{
-//			UEUUID = HashCombine(UEUUID, GetTypeHash(*IOriginalName));
-//		}
-//
-//		if (RUUID)
-//		{
-//			UEUUID = HashCombine(UEUUID, GetTypeHash(*RUUID));
-//		}
-//		else if (RName)
-//		{
-//			UEUUID = HashCombine(UEUUID, GetTypeHash(*RName));
-//		}
-//
-//		UEUUIDStr = FString::Printf(TEXT("0x%08x"), UEUUID);
-//	}
-//}
-
 TSharedPtr< IDatasmithActorElement >  FDatasmithSceneGraphBuilder::CreateActor(const TCHAR* InEUUID, const TCHAR* InLabel)
 {
 	TSharedPtr< IDatasmithActorElement > Actor = FDatasmithSceneFactory::CreateActor(InEUUID);
@@ -536,11 +613,19 @@ void FDatasmithSceneGraphBuilder::GetNodeUUIDAndName(
 	{
 		UEUUID = HashCombine(UEUUID, GetTypeHash(*IOriginalName));
 	}
+	else if (IName)
+	{
+		UEUUID = HashCombine(UEUUID, GetTypeHash(*IName));
+	}
 
 	if (RUUID)
 	{
 		UEUUID = HashCombine(UEUUID, GetTypeHash(*RUUID));
 	}
+	else if (ROriginalName)
+	{
+		UEUUID = HashCombine(UEUUID, GetTypeHash(*ROriginalName));
+	}
 	else if (RName)
 	{
 		UEUUID = HashCombine(UEUUID, GetTypeHash(*RName));
@@ -549,8 +634,7 @@ void FDatasmithSceneGraphBuilder::GetNodeUUIDAndName(
 	OutUEUUID = FString::Printf(TEXT("0x%08x"), UEUUID);
 }
 
-
-TSharedPtr< IDatasmithActorElement > FDatasmithSceneGraphBuilder::BuildComponent(FCTNode& Node, const TCHAR* ParentUuid)
+TSharedPtr< IDatasmithActorElement > FDatasmithSceneGraphBuilder::BuildComponent(FCTNode& Node, ActorData& ParentData)
 {
 	TMap<FString, FString> InstanceNodeMetaDataMap;
 	TMap<FString, FString> ReferenceNodeMetaDataMap;
@@ -559,7 +643,7 @@ TSharedPtr< IDatasmithActorElement > FDatasmithSceneGraphBuilder::BuildComponent
 	FString ActorUUID = TEXT("");
 	FString ActorLabel = TEXT("");
 	
-	GetNodeUUIDAndName(InstanceNodeMetaDataMap, ReferenceNodeMetaDataMap, ParentUuid, ActorUUID, ActorLabel);
+	GetNodeUUIDAndName(InstanceNodeMetaDataMap, ReferenceNodeMetaDataMap, ParentData.Uuid, ActorUUID, ActorLabel);
 
 	TSharedPtr< IDatasmithActorElement > Actor = CreateActor(*ActorUUID, *ActorLabel);
 	if (!Actor.IsValid())
@@ -569,12 +653,15 @@ TSharedPtr< IDatasmithActorElement > FDatasmithSceneGraphBuilder::BuildComponent
 
 	AddMetaData(Actor, InstanceNodeMetaDataMap, ReferenceNodeMetaDataMap);
 
-	AddChildren(Actor, Node, *ActorUUID);
+	ActorData ComponentData(*ActorUUID, ParentData);
+	GetMainMaterial(ReferenceNodeMetaDataMap, &Node, ComponentData, bMaterialPropagationIsTopDown);
+
+	AddChildren(Actor, Node, ComponentData);
 
 	return Actor;
 }
 
-TSharedPtr< IDatasmithActorElement > FDatasmithSceneGraphBuilder::BuildBody(FCTNode& Node, const TCHAR* ParentUuid)
+TSharedPtr< IDatasmithActorElement > FDatasmithSceneGraphBuilder::BuildBody(FCTNode& Node, ActorData& ParentData)
 {
 	TMap<FString, FString> InstanceNodeMetaDataMap;
 	TMap<FString, FString> BodyNodeMetaDataMap;
@@ -582,7 +669,7 @@ TSharedPtr< IDatasmithActorElement > FDatasmithSceneGraphBuilder::BuildBody(FCTN
 	
 	FString BodyUUID = TEXT("");
 	FString BodyLabel = TEXT("");
-	GetNodeUUIDAndName(InstanceNodeMetaDataMap, BodyNodeMetaDataMap, ParentUuid, BodyUUID, BodyLabel);
+	GetNodeUUIDAndName(InstanceNodeMetaDataMap, BodyNodeMetaDataMap, ParentData.Uuid, BodyUUID, BodyLabel);
 
 	TSharedPtr< IDatasmithMeshElement > MeshElement = FindOrAddMeshElement(Node, BodyLabel);
 	if (!MeshElement.IsValid())
@@ -600,12 +687,29 @@ TSharedPtr< IDatasmithActorElement > FDatasmithSceneGraphBuilder::BuildBody(FCTN
 	ActorElement->SetLabel(*BodyLabel);
 	ActorElement->SetStaticMeshPathName(MeshElement->GetName());
 
+	// Apply materials on the current part
+	uint32 MaterialUuid = 0;
+	MaterialUuid = ParentData.MaterialUuid ? ParentData.MaterialUuid : ParentData.ColorUuid;
+
+	if (MaterialUuid)
+	{
+		TSharedPtr< IDatasmithMaterialIDElement > PartMaterialIDElement = FindOrAddMaterial(MaterialUuid);
+		const TCHAR* MaterialIDElementName = PartMaterialIDElement->GetName();
+
+		for (int32 Index = 0; Index < MeshElement->GetMaterialSlotCount(); ++Index)
+		{
+			TSharedPtr< IDatasmithMaterialIDElement > MaterialIDElement = FDatasmithSceneFactory::CreateMaterialId(MaterialIDElementName);
+			MaterialIDElement->SetId(MeshElement->GetMaterialSlotAt(Index)->GetId());
+			ActorElement->AddMaterialOverride(MaterialIDElement);
+		}
+	}
+
 	return ActorElement;
 }
 
 TSharedPtr< IDatasmithMeshElement > FDatasmithSceneGraphBuilder::FindOrAddMeshElement(FCTNode& Node, FString& BodyName)
 {
-	uint32 ShellUuid = Node.GetStaticMeshUuid();
+	const uint32 ShellUuid = Node.GetStaticMeshUuid();
 	FString ShellUuidName = FString::Printf(TEXT("0x%08x"), ShellUuid);
 
 	// Look if geometry has not been already processed, return it if found
@@ -619,6 +723,9 @@ TSharedPtr< IDatasmithMeshElement > FDatasmithSceneGraphBuilder::FindOrAddMeshEl
 	MeshElement->SetLabel(*BodyName);
 	MeshElement->SetLightmapSourceUV(-1);
 
+	FString BodyFile = FString::Printf(TEXT("UEx%08x"), ShellUuid);
+	MeshElement->SetFile(*FPaths::Combine(CachePath, TEXT("body"), BodyFile + TEXT(".ct")));
+
 	// TODO: Set bounding box 
 	//float BoundingBox[6];
 	//FString Buffer = GetStringAttribute(GeomID, TEXT("UE_MESH_BBOX"));
@@ -627,13 +734,13 @@ TSharedPtr< IDatasmithMeshElement > FDatasmithSceneGraphBuilder::FindOrAddMeshEl
 	//	MeshElement->SetDimensions(BoundingBox[3] - BoundingBox[0], BoundingBox[4] - BoundingBox[1], BoundingBox[5] - BoundingBox[2], 0.0f);
 	//}
 
-	FCTMaterialPartition MaterialPartition;
-	Node.GetMaterialSet(MaterialPartition);
+	TMap<uint32, uint32> MaterialIdToMaterialHashMap;
+	Node.GetMaterialSet(MaterialIdToMaterialHashMap);
 
-	for (auto MaterialId2Hash : MaterialPartition.GetMaterialIdToHashSet())
+	for (auto MaterialId2Hash : MaterialIdToMaterialHashMap)
 	{
 		TSharedPtr< IDatasmithMaterialIDElement > PartMaterialIDElement;
-		PartMaterialIDElement = FindOrAddMaterial(Node, MaterialId2Hash.Key);
+		PartMaterialIDElement = FindOrAddMaterial(MaterialId2Hash.Value);
 
 		const TCHAR* MaterialIDElementName = PartMaterialIDElement->GetName();
 
@@ -654,39 +761,39 @@ TSharedPtr< IDatasmithUEPbrMaterialElement > FDatasmithSceneGraphBuilder::GetDef
 {
 	if (!DefaultMaterial.IsValid())
 	{
-		DefaultMaterial = CreateDefaultUEPbrMaterial();
+		DefaultMaterial = CADLibrary::CreateDefaultUEPbrMaterial();
 		DatasmithScene->AddMaterial(DefaultMaterial);
 	}
 
 	return DefaultMaterial;
 }
 
-TSharedPtr<IDatasmithMaterialIDElement> FDatasmithSceneGraphBuilder::FindOrAddMaterial(FCTNode& Node, uint32 MaterialID)
+TSharedPtr<IDatasmithMaterialIDElement> FDatasmithSceneGraphBuilder::FindOrAddMaterial(uint32 MaterialUuid)
 {
 	TSharedPtr< IDatasmithUEPbrMaterialElement > MaterialElement;
 
-	TSharedPtr< IDatasmithUEPbrMaterialElement >* MaterialPtr = MaterialMap.Find(MaterialID);
+	TSharedPtr< IDatasmithUEPbrMaterialElement >* MaterialPtr = MaterialUuidMap.Find(MaterialUuid);
 	if (MaterialPtr != nullptr)
 	{
 		MaterialElement = *MaterialPtr;
 	}
-	else if (MaterialID > 0)
+	else if (MaterialUuid > 0)
 	{
-		if (MaterialID > LAST_CT_MATERIAL_ID)
+		if (FColor* Color = MaterialUuidToColor.Find(MaterialUuid))
 		{
-			MaterialElement = CreateUEPbrMaterialFromColorId(Node, MaterialID);
+			MaterialElement = CADLibrary::CreateUEPbrMaterialFromColor(*Color);
 		}
-		else
+		else if (CADLibrary::FCADMaterial* Material = MaterialUuidToMaterial.Find(MaterialUuid))
 		{
-			MaterialElement = CreateUEPbrMaterialFromMaterialId(Node, MaterialID, DatasmithScene);
+			MaterialElement = CADLibrary::CreateUEPbrMaterialFromMaterial(*Material, DatasmithScene);
 		}
 
 		if (MaterialElement.IsValid())
 		{
 			DatasmithScene->AddMaterial(MaterialElement);
-			MaterialMap.Add(MaterialID, MaterialElement);
+			MaterialUuidMap.Add(MaterialUuid, MaterialElement);
+			DatasmithScene->AddMaterial(MaterialElement);
 		}
-		DatasmithScene->AddMaterial(MaterialElement);
 	}
 
 	if (!MaterialElement.IsValid())
@@ -699,118 +806,18 @@ TSharedPtr<IDatasmithMaterialIDElement> FDatasmithSceneGraphBuilder::FindOrAddMa
 	return MaterialIDElement;
 }
 
-TSharedPtr<IDatasmithUEPbrMaterialElement> FDatasmithSceneGraphBuilder::CreateUEPbrMaterialFromColorId(FCTNode& Node, uint32 InColorHashId)
+TSharedPtr<IDatasmithUEPbrMaterialElement> CreateDefaultUEPbrMaterial()
 {
-	uint32 ColorId;
-	uint8 Alpha;
-	UnhashFastColorHash(InColorHashId, ColorId, Alpha);
-
-	uint8 Color[3];
-	Node.GetColorDescription(ColorId, Color);
-
-	FString Name = FString::FromInt(BuildColorHash(ColorId, Alpha));
-	FString Label = FString::Printf(TEXT("%02x%02x%02x%02x"), Color[0], Color[1], Color[2], Alpha);
-
 	// Take the Material diffuse color and connect it to the BaseColor of a UEPbrMaterial
-	TSharedRef<IDatasmithUEPbrMaterialElement> MaterialElement = FDatasmithSceneFactory::CreateUEPbrMaterial(*Name);
-	MaterialElement->SetLabel(*Label);
-
-	FLinearColor LinearColor = FLinearColor::FromPow22Color(FColor(Color[0], Color[1], Color[2], Alpha));
+	TSharedRef<IDatasmithUEPbrMaterialElement> MaterialElement = FDatasmithSceneFactory::CreateUEPbrMaterial(TEXT("0"));
+	MaterialElement->SetLabel(TEXT("DefaultCADImportMaterial"));
 
+	FLinearColor LinearColor = FLinearColor::FromPow22Color(FColor(200, 200, 200, 255));
 	IDatasmithMaterialExpressionColor* ColorExpression = MaterialElement->AddMaterialExpression<IDatasmithMaterialExpressionColor>();
 	ColorExpression->SetName(TEXT("Diffuse Color"));
 	ColorExpression->GetColor() = LinearColor;
-
 	MaterialElement->GetBaseColor().SetExpression(ColorExpression);
 
-	if (LinearColor.A < 1.0f)
-	{
-		MaterialElement->SetBlendMode(/*EBlendMode::BLEND_Translucent*/2);
-
-		IDatasmithMaterialExpressionScalar* Scalar = MaterialElement->AddMaterialExpression<IDatasmithMaterialExpressionScalar>();
-		Scalar->GetScalar() = LinearColor.A;
-
-		MaterialElement->GetOpacity().SetExpression(Scalar);
-	}
-	return MaterialElement;
-}
-
-TSharedPtr<IDatasmithUEPbrMaterialElement> FDatasmithSceneGraphBuilder::CreateUEPbrMaterialFromMaterialId(FCTNode& Node, uint32 InMaterialID, TSharedRef<IDatasmithScene> Scene)
-{
-	FString MaterialLabel, TextureName;
-	uint8   Diffuse[3];
-	uint8   Ambient[3];
-	uint8   Specular[3];
-	uint8   Shininess, Transparency, Reflexion;
-	int32  Hash;
-
-	FString MaterialName, MaterialStr;
-	Node.GetMaterialDescription(InMaterialID, Hash, MaterialLabel, Diffuse, Ambient, Specular, Shininess, Transparency, Reflexion, TextureName);
-
-	FString Name = FString::FromInt(Hash);
-	// Take the Material diffuse color and connect it to the BaseColor of a UEPbrMaterial
-	TSharedRef<IDatasmithUEPbrMaterialElement> MaterialElement = FDatasmithSceneFactory::CreateUEPbrMaterial(*Name);
-	if (MaterialLabel.IsEmpty())
-	{
-		MaterialLabel = TEXT("Material");
-	}
-	MaterialElement->SetLabel(*MaterialLabel);
-
-	// Set a diffuse color if there's nothing in the BaseColor
-	if (MaterialElement->GetBaseColor().GetExpression() == nullptr)
-	{
-		FColor Color((uint8)Diffuse[0], (uint8)Diffuse[1], (uint8)Diffuse[2], Transparency);
-		FLinearColor LinearColor = FLinearColor::FromPow22Color(Color);
-
-		IDatasmithMaterialExpressionColor* ColorExpression = MaterialElement->AddMaterialExpression<IDatasmithMaterialExpressionColor>();
-		ColorExpression->SetName(TEXT("Diffuse Color"));
-		ColorExpression->GetColor() = LinearColor;
-
-		MaterialElement->GetBaseColor().SetExpression(ColorExpression);
-	}
-
-	if (Transparency < 255)
-	{
-		MaterialElement->SetBlendMode(/*EBlendMode::BLEND_Translucent*/2);
-		IDatasmithMaterialExpressionScalar* Scalar = MaterialElement->AddMaterialExpression<IDatasmithMaterialExpressionScalar>();
-		Scalar->GetScalar() = (float) Transparency / 255.f;
-		MaterialElement->GetOpacity().SetExpression(Scalar);
-	}
-
-	// Set a Emissive color 
-	if (MaterialElement->GetEmissiveColor().GetExpression() == nullptr)
-	{
-		// Doc CT => TODO
-		//GLfloat Specular[4] = { specular.rgb[0] / 255., specular.rgb[1] / 255., specular.rgb[2] / 255., 1. - transparency };
-		//GLfloat Shininess[1] = { (float)(128 * shininess) };
-		//glMaterialfv(GL_FRONT_AND_BACK, GL_SPECULAR, Specular);
-		//glMaterialfv(GL_FRONT_AND_BACK, GL_SHININESS, Shininess);
-
-		FColor Color((uint8)Specular[0], (uint8)Specular[1], (uint8)Specular[2], Transparency);
-		FLinearColor LinearColor = FLinearColor::FromPow22Color(Color);
-
-		IDatasmithMaterialExpressionColor* ColorExpression = MaterialElement->AddMaterialExpression<IDatasmithMaterialExpressionColor>();
-		ColorExpression->SetName(TEXT("Specular Color"));
-		ColorExpression->GetColor() = LinearColor;
-
-		MaterialElement->GetEmissiveColor().SetExpression(ColorExpression);
-	}
-
-	// Simple conversion of shininess and reflectivity to PBR roughness and metallic values; model could be improved to properly blend the values
-	if (Shininess != 0)
-	{
-		IDatasmithMaterialExpressionScalar* Scalar = static_cast<IDatasmithMaterialExpressionScalar*>(MaterialElement->AddMaterialExpression(EDatasmithMaterialExpressionType::ConstantScalar));
-		Scalar->GetScalar() = 1.f - Shininess/255.;
-		MaterialElement->GetRoughness().SetExpression(Scalar);
-	}
-
-	if (Reflexion != 0)
-	{
-		IDatasmithMaterialExpressionScalar* Scalar = static_cast<IDatasmithMaterialExpressionScalar*>(MaterialElement->AddMaterialExpression(EDatasmithMaterialExpressionType::ConstantScalar));
-		Scalar->GetScalar() = Reflexion;
-		MaterialElement->GetMetallic().SetExpression(Scalar);
-	}
-
 	return MaterialElement;
 }
 
@@ -923,7 +930,7 @@ void FDatasmithSceneGraphBuilder::LinkActor(TSharedPtr< IDatasmithActorElement >
 	}
 }
 
-void FDatasmithSceneGraphBuilder::AddChildren(TSharedPtr< IDatasmithActorElement > Actor, FCTNode& ComponentNode, const TCHAR* ParentUuid)
+void FDatasmithSceneGraphBuilder::AddChildren(TSharedPtr< IDatasmithActorElement > Actor, FCTNode& ComponentNode, ActorData& ParentData)
 {
 	TArray<int32> ChildrenIds;
 	ComponentNode.GetChildren(ChildrenIds);
@@ -934,13 +941,10 @@ void FDatasmithSceneGraphBuilder::AddChildren(TSharedPtr< IDatasmithActorElement
 		{
 			continue;
 		}
-		TSharedPtr< IDatasmithActorElement > ChildActor = BuildNode(*ChildNode, ParentUuid);
+		TSharedPtr< IDatasmithActorElement > ChildActor = BuildNode(*ChildNode, ParentData);
 		if (ChildActor.IsValid())
 		{
 			Actor->AddChild(ChildActor);
 		}
 	}
 }
-
-
-#endif  // USE_CORETECH_MT_PARSER
\ No newline at end of file
diff --git a/Engine/Plugins/Enterprise/DatasmithCADImporter/Source/DatasmithCADTranslator/Private/DatasmithSceneGraphBuilder.h b/Engine/Plugins/Enterprise/DatasmithCADImporter/Source/DatasmithCADTranslator/Private/DatasmithSceneGraphBuilder.h
--- a/Engine/Plugins/Enterprise/DatasmithCADImporter/Source/DatasmithCADTranslator/Private/DatasmithSceneGraphBuilder.h
+++ b/Engine/Plugins/Enterprise/DatasmithCADImporter/Source/DatasmithCADTranslator/Private/DatasmithSceneGraphBuilder.h
@@ -1,17 +1,16 @@
 // Copyright 1998-2019 Epic Games, Inc. All Rights Reserved.
 #pragma once
 
-#ifdef USE_CORETECH_MT_PARSER
-
 #include "CoreMinimal.h"
 
-#include "CADLibraryOptions.h"
-#include "Containers/Map.h"
-#include "Containers/Queue.h"
+#include "CADData.h"
 #include "CoreTechHelper.h"
 #include "CoreTechTypes.h"
 #include "CTSession.h"
 #include "DatasmithImportOptions.h"
+
+#include "Containers/Map.h"
+#include "Containers/Queue.h"
 #include "Misc/Paths.h"
 
 static const FString SPACE(" ");
@@ -23,7 +22,6 @@ class IDatasmithMeshElement;
 class IDatasmithScene;
 class IDatasmithUEPbrMaterialElement;
 
-
 enum NODE_TYPE
 {
 	INSTANCE,
@@ -46,15 +44,15 @@ public:
 	void GetChildren(TArray<int32>& Children);
 	
 	void GetNodeReference(int32& OutRefId, FString& OutExternalFile, NODE_TYPE& OutType);
-	void AddTransformToActor(TSharedPtr< IDatasmithActorElement > Actor, CADLibrary::FImportParameters& InImportParameters);
+	void AddTransformToActor(TSharedPtr< IDatasmithActorElement > Actor, const CADLibrary::FImportParameters& InImportParameters);
 
 	FCTNode* GetCTNode(int32 NodeId);
 
 	uint32 GetStaticMeshUuid();
-	void GetMaterialSet(CADLibrary::FCTMaterialPartition& MaterialPartition);
+	void GetMaterialSet(TMap<uint32, uint32>& MaterialIdToMaterialHashMap);
 
-	void GetColorDescription(int32 ColorId, uint8 CtColor[3]);
-	void GetMaterialDescription(int32 MaterialId, int32& OutMaterialHash, FString& OutMaterialName, uint8 OutDiffuse[3], uint8 OutAmbient[3], uint8 OutSpecular[3], uint8& OutShininess, uint8& OutTransparency, uint8& OutReflexion, FString& OutTextureName);
+	bool GetColor(int32 ColorHId, FColor& Color) const;
+	bool GetMaterial(int32 MaterialId, CADLibrary::FCADMaterial& Material) const;
 
 protected:
 	void SetNodeType();
@@ -69,13 +67,42 @@ protected:
 	
 };
 
-class FCTRawDataFile
+
+class ActorData  //#ueent_CAD
 {
 public:
-	FCTRawDataFile(const FString& InFileName);
-	~FCTRawDataFile()
+	ActorData(const TCHAR* NodeUuid, const ActorData& ParentData)
+		: Uuid(NodeUuid)
+		, Material(ParentData.Material)
+		, MaterialUuid(ParentData.MaterialUuid)
+		, Color(ParentData.Color)
+		, ColorUuid(ParentData.ColorUuid)
+	{
+	}
+
+	ActorData(const TCHAR* NodeUuid)
+		: Uuid(NodeUuid)
+		, MaterialUuid(0)
+		, ColorUuid(0)
 	{
 	}
+
+	const TCHAR* Uuid;
+
+	CADLibrary::FCADMaterial Material;
+	uint32 MaterialUuid;
+
+	FColor Color;
+	uint32 ColorUuid;
+};
+
+
+
+class FCTRawDataFile //#ueent_CAD
+{
+public:
+	FCTRawDataFile(const FString& InFileName);
+
 	void ReadFile();
 
 	FCTNode* GetCTNode(int32 nodeId)
@@ -83,83 +110,30 @@ public:
 		return CTIdToRawEntryMap.Find(nodeId);
 	}
 
-	const FString& GetString(int32 line)
+	const FString& GetString(int32 line) const
 	{
 		return RawData[line];
 	}
 
-	const FString& GetFileName()
+	const FString& GetFileName() const
 	{
 		return FileName;
 	}
 
-	void GetColorDescription(int32 ColorId, uint8 Color[3])
-	{
-		Color[0] = Color[1] = Color[2] = 255;
-
-		int32* line = ColorIdToLineMap.Find(ColorId);
-		if (line == nullptr)
-		{
-			return;
-		}
-		const FString& StrColor = GetString(*line);
-		TArray<FString> ColorParameters;
-		StrColor.ParseIntoArray(ColorParameters, *SPACE);
-		for (int32 index = 0; index < 3; index++)
-		{
-			Color[index] = FCString::Atoi(*ColorParameters[index + 1]);
-		}
-	}
+	bool GetColor(uint32 ColorHId, FColor& Color) const;
 
-	void GetMaterialDescription(int32 MaterialId, int32& OutMaterialHash, FString& OutMaterialName, uint8 OutDiffuse[3], uint8 OutAmbient[3], uint8 OutSpecular[3], uint8& OutShininess, uint8& OutTransparency, uint8& OutReflexion, FString& OutTextureName)
-	{
-		OutDiffuse[0] = OutDiffuse[1] = OutDiffuse[2] = 255;
-		OutAmbient[0] = OutAmbient[1] = OutAmbient[2] = 255;
-		OutSpecular[0] = OutSpecular[1] = OutSpecular[2] = 255;
-		OutShininess = OutReflexion = 0;
-		OutTransparency = 255;
-		OutMaterialName = TEXT("");
-		OutTextureName = TEXT("");
-
-		int32* line = MaterialIdToLineMap.Find(MaterialId);
-		if (line == nullptr)
-		{
-			return;
-		}
-
-		const FString& MaterialLine1 = GetString(*line);
-		FString MaterialIdStr;
-		MaterialLine1.Split(SPACE, &MaterialIdStr, &OutMaterialName);
-
-		const FString& MaterialLine2 = GetString(*line + 1);
-		TArray<FString> MaterialParameters;
-		MaterialLine2.ParseIntoArray(MaterialParameters, *SPACE);
-
-		OutMaterialHash = FCString::Atoi(*MaterialParameters[1]);
-		for (int32 index = 0; index < 3; index++)
-		{
-			OutDiffuse[index] = FCString::Atoi(*MaterialParameters[index + 2]);
-			OutAmbient[index] = FCString::Atoi(*MaterialParameters[index + 5]);
-			OutSpecular[index] = FCString::Atoi(*MaterialParameters[index + 8]);
-		}
-		OutShininess = FCString::Atoi(*MaterialParameters[11]);
-		OutTransparency = FCString::Atoi(*MaterialParameters[12]);
-		OutReflexion = FCString::Atoi(*MaterialParameters[13]);
-		int32 TextureId = FCString::Atoi(*MaterialParameters[14]);
-
-		if (MaterialParameters[14] != TEXT("0"))
-		{
-			const FString& MaterialLine3 = GetString(*line + 2);
-			MaterialLine3.Split(SPACE, &MaterialIdStr, &OutTextureName);
-		}
-	}
+	bool GetMaterial(int32 MaterialId, CADLibrary::FCADMaterial& material) const;
+
+	void GetMaterialDescription(int32 LineNumber, CADLibrary::FCADMaterial& Material) const;
+
+	void SetMaterialMaps(TMap< uint32, FColor>& MaterialUuidToColor, TMap< uint32, CADLibrary::FCADMaterial>& MaterialUuidToMaterial);
 
 protected:
 	FString FileName;
 	TArray<FString> RawData;
 	TMap<int32, FCTNode> CTIdToRawEntryMap;
-	TMap<int32, int32> ColorIdToLineMap;
-	TMap<int32, int32> MaterialIdToLineMap;
+	TMap<int32, FColor> ColorIdToColor;
+	TMap<int32, CADLibrary::FCADMaterial> MaterialIdToMaterial;
 };
 
 
@@ -167,60 +141,55 @@ protected:
 class FDatasmithSceneGraphBuilder
 {
 public:
-	FDatasmithSceneGraphBuilder(const FString& InCachePath, TSharedRef<IDatasmithScene> InScene, const FDatasmithSceneSource& InSource, TMap<FString, FString>& InCADFileToUE4FileMap, TMap< TSharedPtr< IDatasmithMeshElement >, uint32 >& InMeshElementToCTBodyUuidMap);
+	FDatasmithSceneGraphBuilder(TMap<FString, FString>& InCADFileToUE4FileMap, TMap< TSharedPtr< IDatasmithMeshElement >, uint32 >& InMeshElementToCTBodyUuidMap, const FString& InCachePath, 
+		TSharedRef<IDatasmithScene> InScene, const FDatasmithSceneSource& InSource, const CADLibrary::FImportParameters& InImportParameters);
 
 	bool Build();
 
-	void SetImportParameters(const CADLibrary::FImportParameters& InImportParameters)
-	{
-		ImportParameters = InImportParameters;
-	};
-
+	void Clear();
 
 protected:
 	void LoadRawDataFile();
 
-	TSharedPtr< IDatasmithActorElement > BuildNode(FCTNode& Node, const TCHAR* ParentUuid);
-	TSharedPtr< IDatasmithActorElement > BuildInstance(FCTNode& Node, const TCHAR* ParentUuid);
-	TSharedPtr< IDatasmithActorElement > BuildComponent(FCTNode& Node, const TCHAR* ParentUuid);
-	TSharedPtr< IDatasmithActorElement > BuildBody(FCTNode& NodeId, const TCHAR* ParentUuid);
+	TSharedPtr< IDatasmithActorElement > BuildNode(FCTNode& Node, ActorData& ParentData);
+	TSharedPtr< IDatasmithActorElement > BuildInstance(FCTNode& Node, ActorData& ParentData);
+	TSharedPtr< IDatasmithActorElement > BuildComponent(FCTNode& Node, ActorData& ParentData);
+	TSharedPtr< IDatasmithActorElement > BuildBody(FCTNode& NodeId, ActorData& ParentData);
 
 	void AddMetaData(TSharedPtr< IDatasmithActorElement > ActorElement, TMap<FString, FString>& InstanceNodeAttributeSetMap, TMap<FString, FString>& ReferenceNodeAttributeSetMap);
-	void AddChildren(TSharedPtr< IDatasmithActorElement > Actor, FCTNode& Node, const TCHAR* ParentUuid);
+	void AddChildren(TSharedPtr< IDatasmithActorElement > Actor, FCTNode& Node, ActorData& ParentData);
 	void LinkActor(TSharedPtr< IDatasmithActorElement > ParentActor, TSharedPtr< IDatasmithActorElement > Actor);
 
 	TSharedPtr< IDatasmithUEPbrMaterialElement > GetDefaultMaterial();
-	TSharedPtr<IDatasmithMaterialIDElement> FindOrAddMaterial(FCTNode& Node, uint32 MaterialID);
+	TSharedPtr<IDatasmithMaterialIDElement> FindOrAddMaterial(uint32 MaterialUuid);
 
 	TSharedPtr< IDatasmithActorElement > CreateActor(const TCHAR* ActorUUID, const TCHAR* ActorLabel);
 	TSharedPtr< IDatasmithMeshElement > FindOrAddMeshElement(FCTNode& Node, FString& InLabel);
-	//TSharedPtr< IDatasmithActorElement > CreateMeshActor(CTNode& Node);
 
 	void GetNodeUUIDAndName(TMap<FString, FString>& InInstanceNodeMetaDataMap, TMap<FString, FString>& InReferenceNodeMetaDataMap, const TCHAR* InParentUEUUID, FString& OutUEUUID, FString& OutName);
 
-	TSharedPtr<IDatasmithUEPbrMaterialElement> CreateUEPbrMaterialFromColorId(FCTNode& Node, uint32 InColorHashId);
-	TSharedPtr<IDatasmithUEPbrMaterialElement> CreateUEPbrMaterialFromMaterialId(FCTNode& Node, uint32 InMaterialID, TSharedRef<IDatasmithScene> Scene);
-
 
 protected:
+	TMap<FString, FString>& CADFileToRawDataFile;
+	TMap< TSharedPtr< IDatasmithMeshElement >, uint32 >& MeshElementToCTBodyUuidMap;
+	const FString& CachePath;
 	TSharedRef<IDatasmithScene> DatasmithScene;
 	const FDatasmithSceneSource& Source;
-	FString CachePath;
-
-	CADLibrary::FImportParameters ImportParameters;
+	const CADLibrary::FImportParameters& ImportParameters;
 
-	TMap<FString, FString>& CADFileToRawDataFile;
 	TMap<FString, FCTRawDataFile*> CADFileToRawData;
-	TArray<FCTRawDataFile> RawDataArray;
 
+	TArray<FCTRawDataFile> RawDataArray; //#ueent_CAD
 
 	TMap< uint32, TSharedPtr< IDatasmithMeshElement > > BodyUuidToMeshElementMap;
-	TMap< TSharedPtr< IDatasmithMeshElement >, uint32 >& MeshElementToCTBodyUuidMap;
 
-	TMap< uint32, TSharedPtr< IDatasmithUEPbrMaterialElement > > MaterialMap;
-	TSharedPtr< IDatasmithUEPbrMaterialElement > DefaultMaterial;
+	TMap< uint32, TSharedPtr< IDatasmithUEPbrMaterialElement > > MaterialUuidMap;
+	TSharedPtr<IDatasmithUEPbrMaterialElement > DefaultMaterial;
 
-};
+	TMap< uint32, FColor> MaterialUuidToColor;
+	TMap< uint32, CADLibrary::FCADMaterial> MaterialUuidToMaterial;
 
 
-#endif  // USE_CORETECH_MT_PARSER
+	bool bPreferMaterial;
+	bool bMaterialPropagationIsTopDown;
+};
diff --git a/Engine/Plugins/Enterprise/DatasmithCADImporter/Source/DatasmithCoreTechExtension/Private/CoreTechParametricSurfaceExtension.cpp b/Engine/Plugins/Enterprise/DatasmithCADImporter/Source/DatasmithCoreTechExtension/Private/CoreTechParametricSurfaceExtension.cpp
--- a/Engine/Plugins/Enterprise/DatasmithCADImporter/Source/DatasmithCoreTechExtension/Private/CoreTechParametricSurfaceExtension.cpp
+++ b/Engine/Plugins/Enterprise/DatasmithCADImporter/Source/DatasmithCoreTechExtension/Private/CoreTechParametricSurfaceExtension.cpp
@@ -5,7 +5,9 @@
 #ifdef CAD_LIBRARY
 #include "CoreTechHelper.h"
 #include "CoreTechMeshLoader.h"
+#include "CoreTechTypes.h"
 #endif
+
 #include "DatasmithAdditionalData.h"
 #include "DatasmithStaticMeshImporter.h" // Call to BuildStaticMesh
 #include "DatasmithUtils.h"
@@ -176,8 +178,8 @@ bool FCoreTechRetessellate_Impl::ApplyOnOneAsset(UStaticMesh& StaticMesh, UCoreT
 	CADLibrary::FImportParameters ImportParameters;
 	ImportParameters.MetricUnit = CoreTechData.SceneParameters.MetricUnit;
 	ImportParameters.ScaleFactor = CoreTechData.SceneParameters.ScaleFactor;
-	ImportParameters.ChordTolerance = RetessellateOptions.ChordTolerance / CoreTechData.SceneParameters.ScaleFactor;
-	ImportParameters.MaxEdgeLength = RetessellateOptions.MaxEdgeLength / CoreTechData.SceneParameters.ScaleFactor;
+	ImportParameters.ChordTolerance = RetessellateOptions.ChordTolerance;
+	ImportParameters.MaxEdgeLength = RetessellateOptions.MaxEdgeLength;
 	ImportParameters.MaxNormalAngle = RetessellateOptions.NormalTolerance;
 	ImportParameters.StitchingTechnique = CADLibrary::EStitchingTechnique(RetessellateOptions.StitchingTechnique);
 
diff --git a/Engine/Plugins/Enterprise/DatasmithCADImporter/Source/DatasmithCoreTechExtension/Public/DatasmithCoreTechTranslator.h b/Engine/Plugins/Enterprise/DatasmithCADImporter/Source/DatasmithCoreTechExtension/Public/DatasmithCoreTechTranslator.h
--- a/Engine/Plugins/Enterprise/DatasmithCADImporter/Source/DatasmithCoreTechExtension/Public/DatasmithCoreTechTranslator.h
+++ b/Engine/Plugins/Enterprise/DatasmithCADImporter/Source/DatasmithCoreTechExtension/Public/DatasmithCoreTechTranslator.h
@@ -4,7 +4,6 @@
 
 #include "CoreMinimal.h"
 
-#include "CADLibraryOptions.h"
 #include "DatasmithImportOptions.h"
 #include "Translators/DatasmithTranslator.h"
 #include "UObject/ObjectMacros.h"
diff --git a/Engine/Plugins/Enterprise/DatasmithCADImporter/Source/DatasmithDispatcher/Private/DatasmithDispatcher.cpp b/Engine/Plugins/Enterprise/DatasmithCADImporter/Source/DatasmithDispatcher/Private/DatasmithDispatcher.cpp
--- a/Engine/Plugins/Enterprise/DatasmithCADImporter/Source/DatasmithDispatcher/Private/DatasmithDispatcher.cpp
+++ b/Engine/Plugins/Enterprise/DatasmithCADImporter/Source/DatasmithDispatcher/Private/DatasmithDispatcher.cpp
@@ -1,69 +1,81 @@
 // Copyright 1998-2019 Epic Games, Inc. All Rights Reserved.
 #include "DatasmithDispatcher.h"
 
-#include "Misc/ScopeLock.h"
-#include "HAL/Thread.h"
 #include "Containers/Array.h"
+#include "CoreTechFileParser.h"
+#include "HAL/FileManager.h"
+#include "HAL/Thread.h"
 #include "Misc/Optional.h"
+#include "Misc/Paths.h"
+#include "Misc/ScopeLock.h"
+
 
 namespace DatasmithDispatcher
 {
-	FDatasmithDispatcher::FDatasmithDispatcher(const TCHAR* InCacheDir, int32 InNumberOfWorkers, TMap<FString, FString>& OutCADFileToUnrealFileMap, TMap<FString, FString>& OutCADFileToUnrealGeomMap)
-		: ProcessCacheFolder(InCacheDir)
-		, CADFileToUnrealFileMap(OutCADFileToUnrealFileMap)
-		, CADFileToUnrealGeomMap(OutCADFileToUnrealGeomMap)
-		, NumberOfWorkers(InNumberOfWorkers)
-		, LastLaunchTask(0)
-		, NumOfFinishedTask(0)
- 	{
-		WorkerHandlerSet.Reserve(NumberOfWorkers);
-		for (int32 Index = 0; Index < NumberOfWorkers; Index++)
-		{
-			WorkerHandlerSet.Emplace(*this, ProcessCacheFolder);
-		}
-	}
 
-	FDatasmithDispatcher::~FDatasmithDispatcher()
+FDatasmithDispatcher::FDatasmithDispatcher(const FString& InCacheDir, int32 InNumberOfWorkers, TMap<FString, FString>& OutCADFileToUnrealFileMap, TMap<FString, FString>& OutCADFileToUnrealGeomMap)
+	: CADFileToUnrealFileMap(OutCADFileToUnrealFileMap)
+	, CADFileToUnrealGeomMap(OutCADFileToUnrealGeomMap)
+	, ProcessCacheFolder(InCacheDir)
+	, NumberOfWorkers(InNumberOfWorkers)
+	, LastLaunchTask(0)
+	, NumOfFinishedTask(0)
+{
+	WorkerHandlerSet.Reserve(NumberOfWorkers);
+	for (int32 Index = 0; Index < NumberOfWorkers; Index++)
 	{
+		WorkerHandlerSet.Emplace(*this, ProcessCacheFolder);
 	}
+}
 
+void FDatasmithDispatcher::Init(const CADLibrary::FImportParameters& InImportParameters)
+{
+	ImportParameters = InImportParameters;
 
-	void FDatasmithDispatcher::RunHandlers()
+	IFileManager::Get().MakeDirectory(*FPaths::Combine(ProcessCacheFolder, TEXT("scene")), true);
+	IFileManager::Get().MakeDirectory(*FPaths::Combine(ProcessCacheFolder, TEXT("cad")), true);
+	IFileManager::Get().MakeDirectory(*FPaths::Combine(ProcessCacheFolder, TEXT("mesh")), true);
+	IFileManager::Get().MakeDirectory(*FPaths::Combine(ProcessCacheFolder, TEXT("body")), true);
+}
+
+void FDatasmithDispatcher::RunHandlers()
+{
+	WorkersRunning.Reserve(NumberOfWorkers);
+	int32 WorkerIndex = 0;
+	for (FDatasmithWorkerHandler& Handler : WorkerHandlerSet)
 	{
-		WorkersRunning.Reserve(NumberOfWorkers);
-		int32 WorkerIndex = 0;
-		for (FDatasmithWorkerHandler& Handler : WorkerHandlerSet)
-		{
-			FString ProcessName = TEXT("DatasmithWorker ") + FString::FromInt(WorkerIndex);
-			WorkerIndex++;
-			WorkersRunning.Emplace(*ProcessName, [&]() { Handler.Run(); });
-		}
+		FString ProcessName = TEXT("DatasmithWorker ") + FString::FromInt(WorkerIndex);
+		WorkerIndex++;
+		WorkersRunning.Emplace(*ProcessName, [&]() { Handler.Run(); });
 	}
+}
 
-	void FDatasmithDispatcher::Close()
+void FDatasmithDispatcher::Close()
+{
+	for (FDatasmithWorkerHandler& Handler : WorkerHandlerSet)
 	{
-		for (FDatasmithWorkerHandler& Handler : WorkerHandlerSet)
-		{
-			Handler.NotifyKill();
-		}
-
-		for (FThread& Process : WorkersRunning)
-		{
-			Process.Join();
-		}
+		Handler.NotifyKill();
 	}
 
-	void FDatasmithDispatcher::Clear()
+	for (FThread& Process : WorkersRunning)
 	{
-		TaskPool.Empty();
-		LastLaunchTask = 0;
-		NumOfFinishedTask = 0;
+		Process.Join();
 	}
+}
+
+void FDatasmithDispatcher::Clear()
+{
+	TaskPool.Empty();
+	LastLaunchTask = 0;
+	NumOfFinishedTask = 0;
+}
 
 
-	void FDatasmithDispatcher::Process()
+void FDatasmithDispatcher::Process(bool bWithProcessor)
+{
+	if (bWithProcessor)
 	{
-		RunHandlers();
+		RunHandlers();  
 
 		while (!IsOver())
 		{
@@ -72,70 +84,113 @@ namespace DatasmithDispatcher
 
 		Close();
 	}
+	else
+	{
+		ProcessLocal();
+	}
+}
 
-	void FDatasmithDispatcher::AddTask(const FString& InFile)
+void FDatasmithDispatcher::ProcessLocal()
+{
+	TOptional<FTask> Task = GetTask();
+	while (Task.IsSet())
 	{
-		FScopeLock Lock(&TaskPoolCriticalSection);
-		for (const FTask& Task : TaskPool)
+		FString FullPath = Task->FileName;
+
+		CADLibrary::FCoreTechFileParser FileParser(FullPath, ProcessCacheFolder, ImportParameters, true, true);
+		EProcessState TaskState = FileParser.ProcessFile();
+		SetTaskState(Task->Index, TaskState);
+
+		if (TaskState == ProcessOk)
 		{
-			if (Task.FileName == InFile)
+			FString CurrentPath = FPaths::GetPath(FullPath);
+			const TSet<FString>& ExternalRefSet = FileParser.GetExternalRefSet();
+			if (ExternalRefSet.Num() > 0)
 			{
-				return;
+				for (const FString& ExternalFile : ExternalRefSet)
+				{
+					AddTask(FPaths::Combine(CurrentPath, ExternalFile));
+				}
 			}
+			LinkCTFileToUnrealCacheFile(FileParser.GetFileName(), FileParser.GetSceneGraphFile(), FileParser.GetMeshFile());
 		}
 
-		int32 TaskIndex = TaskPool.Emplace(InFile);
-		TaskPool[TaskIndex].Index = TaskIndex;
+		Task = GetTask();
 	}
+}
 
-	TOptional<FTask> FDatasmithDispatcher::GetTask()
+void FDatasmithDispatcher::AddTask(const FString& InFile)
+{
+	FScopeLock Lock(&TaskPoolCriticalSection);
+	for (const FTask& Task : TaskPool)
 	{
-		FScopeLock Lock(&TaskPoolCriticalSection);
-
-		if (!TaskPool.IsValidIndex(LastLaunchTask))
+		if (Task.FileName == InFile)
 		{
-			return TOptional<FTask>();
+			return;
 		}
-
-		TaskPool[LastLaunchTask].State = Running;
-		return TaskPool[LastLaunchTask++];
 	}
 
-	bool FDatasmithDispatcher::IsOver() 
+	int32 TaskIndex = TaskPool.Emplace(InFile);
+	TaskPool[TaskIndex].Index = TaskIndex;
+}
+
+TOptional<FTask> FDatasmithDispatcher::GetTask()
+{
+	FScopeLock Lock(&TaskPoolCriticalSection);
+
+	if (!TaskPool.IsValidIndex(LastLaunchTask))
 	{
-		FScopeLock Lock(&TaskPoolCriticalSection);
-		return NumOfFinishedTask == TaskPool.Num();
+		return TOptional<FTask>();
 	}
 
-	void FDatasmithDispatcher::SetTaskState(int32 TaskIndex, EProcessState TaskState) 
-	{
-		FScopeLock Lock(&TaskPoolCriticalSection);
+	TaskPool[LastLaunchTask].State = Running;
+	return TaskPool[LastLaunchTask++];
+}
+
+bool FDatasmithDispatcher::IsOver() 
+{
+	FScopeLock Lock(&TaskPoolCriticalSection);
+	return NumOfFinishedTask == TaskPool.Num();
+}
+
+void FDatasmithDispatcher::SetTessellationOptions(float InChordTolerance, float InMaxEdgeLength, float InNormalTolerance, CADLibrary::EStitchingTechnique InStitchingTechnique)
+{
+	ImportParameters.ChordTolerance = InChordTolerance;
+	ImportParameters.MaxEdgeLength = InMaxEdgeLength;
+	ImportParameters.MaxNormalAngle = InNormalTolerance;
+	ImportParameters.StitchingTechnique = InStitchingTechnique;
+}
 
-		if (TaskIndex < 0)
-		{
-			return;
-		}
-		if (TaskIndex < TaskPool.Num())
-		{
-			TaskPool[TaskIndex].State = TaskState;
-			if (TaskState & (ProcessOk | ProcessFailed | FileNotFound))
-			{
-				NumOfFinishedTask++;
-			}
-		}
-	}
 
-	void FDatasmithDispatcher::LinkCTFileToUnrealCacheFile(const FString& CTFile, const FString& UnrealSceneGraphFile, const FString& UnrealMeshFile)
+void FDatasmithDispatcher::SetTaskState(int32 TaskIndex, EProcessState TaskState)
+{
+	FScopeLock Lock(&TaskPoolCriticalSection);
+
+	if (TaskIndex < 0)
 	{
-		FScopeLock Lock(&TaskPoolCriticalSection);
-		if (!UnrealSceneGraphFile.IsEmpty())
-		{
-			CADFileToUnrealFileMap.Add(CTFile, UnrealSceneGraphFile);
-		}
-		if (!UnrealMeshFile.IsEmpty())
+		return;
+	}
+	if (TaskIndex < TaskPool.Num())
+	{
+		TaskPool[TaskIndex].State = TaskState;
+		if (TaskState & (ProcessOk | ProcessFailed | FileNotFound))
 		{
-			CADFileToUnrealGeomMap.Add(CTFile, UnrealMeshFile);
+			NumOfFinishedTask++;
 		}
+	}
+}
 
+void FDatasmithDispatcher::LinkCTFileToUnrealCacheFile(const FString& CTFile, const FString& UnrealSceneGraphFile, const FString& UnrealMeshFile)
+{
+	FScopeLock Lock(&TaskPoolCriticalSection);
+	if (!UnrealSceneGraphFile.IsEmpty())
+	{
+		CADFileToUnrealFileMap.Add(CTFile, UnrealSceneGraphFile);
+	}
+	if (!UnrealMeshFile.IsEmpty())
+	{
+		CADFileToUnrealGeomMap.Add(CTFile, UnrealMeshFile);
 	}
 }
+
+}
\ No newline at end of file
diff --git a/Engine/Plugins/Enterprise/DatasmithCADImporter/Source/DatasmithDispatcher/Private/DatasmithDispatcherSocket.cpp b/Engine/Plugins/Enterprise/DatasmithCADImporter/Source/DatasmithDispatcher/Private/DatasmithDispatcherSocket.cpp
--- a/Engine/Plugins/Enterprise/DatasmithCADImporter/Source/DatasmithDispatcher/Private/DatasmithDispatcherSocket.cpp
+++ b/Engine/Plugins/Enterprise/DatasmithCADImporter/Source/DatasmithDispatcher/Private/DatasmithDispatcherSocket.cpp
@@ -1,16 +1,19 @@
 // Copyright 1998-2019 Epic Games, Inc. All Rights Reserved.
 #include "DatasmithDispatcherSocket.h"
 
-#include "CoreMinimal.h"
+#include "DatasmithDispatcherLog.h"  //#ueent_CAD
+#include "DatasmithDispatcherModule.h"
 
+#include "CoreMinimal.h"
+#include "Logging/LogMacros.h"
 #include "SocketSubsystem.h"
 #include "Sockets.h"
 
-using namespace DatasmithDispatcher;
-
-DEFINE_LOG_CATEGORY_STATIC(LogDatasmithCoretechWorker, Log, All);
+DEFINE_LOG_CATEGORY(LogDatasmithDispatcher);
 
 
+namespace DatasmithDispatcher
+{
 
 FDatasmithDispatcherSocket::FDatasmithDispatcherSocket(const TCHAR* ServerAddress)
 	: SocketAdress(ServerAddress)
@@ -50,8 +53,6 @@ FSocket* FDatasmithDispatcherSocket::Accept()
 	return Socket->Accept(Message);
 }
 
-#include "CoreGlobals.h"
-
 void FDatasmithDispatcherSocket::Bind()
 {
 	ISocketSubsystem* SocketSubsystem = ISocketSubsystem::Get();
@@ -112,6 +113,11 @@ void FDatasmithDispatcherSocket::Connect(const int32& InServerPort)
 	}
 
 	TSharedPtr<FInternetAddr> InternetAddress = SocketSubsystem->CreateInternetAddr();
+	if(!InternetAddress.IsValid())
+	{
+		return;
+	}
+
 	bool bIsValid = true;
 	InternetAddress->SetIp(*SocketAdress, bIsValid);
 	if (!bIsValid)
@@ -131,12 +137,12 @@ void FDatasmithDispatcherSocket::Connect(const int32& InServerPort)
 	{
 		return;
 	}
-	UE_LOG(LogDatasmithCoretechWorker, Display, TEXT("Is connected"));
+	UE_LOG(LogDatasmithDispatcher, Display, TEXT("Is connected"));
 
 	bool bHasPendingConnection = true;
 	Socket->WaitForPendingConnection(bHasPendingConnection, FTimespan(0, 0, 5));
 
-	UE_LOG(LogDatasmithCoretechWorker, Display, TEXT("WaitForPendingConnection return = %d"), bHasPendingConnection);
+	UE_LOG(LogDatasmithDispatcher, Display, TEXT("WaitForPendingConnection return = %d"), bHasPendingConnection);
 	bOpen = !bHasPendingConnection;
 }
 
@@ -165,7 +171,7 @@ void FDatasmithDispatcherSocket::Read(int32 DataSize, uint8* Data)
 		bool bStatus = Socket->Recv(Data + Total, DataSize - Total, BytesRead);
 		if (!bStatus)
 		{
-			UE_LOG(LogDatasmithCoretechWorker, Display, TEXT("Close socket Read %d %d "), DataSize, Total);
+			UE_LOG(LogDatasmithDispatcher, Display, TEXT("Close socket Read %d %d "), DataSize, Total);
 
 			Socket->Close();
 			ErrorCode = UnableToReadOnSocket;
@@ -204,7 +210,7 @@ void FDatasmithDispatcherSocket::SendData()
 			Socket->Close();
 			ErrorCode = UnableToSendData;
 			bOpen = false;
-			UE_LOG(LogDatasmithCoretechWorker, Display, TEXT("Close socket Write %d %d "), BufferSize, Total);
+			UE_LOG(LogDatasmithDispatcher, Display, TEXT("Close socket Write %d %d "), BufferSize, Total);
 		}
 		Total += (uint32) BytesSend;
 	}
@@ -234,3 +240,4 @@ void FDatasmithDispatcherSocket::operator <<(const FString& Data)
 		*this << Data[k];
 	}
 }
+}
\ No newline at end of file
diff --git a/Engine/Plugins/Enterprise/DatasmithCADImporter/Source/DatasmithDispatcher/Private/DatasmithWorkerHandler.cpp b/Engine/Plugins/Enterprise/DatasmithCADImporter/Source/DatasmithDispatcher/Private/DatasmithWorkerHandler.cpp
--- a/Engine/Plugins/Enterprise/DatasmithCADImporter/Source/DatasmithDispatcher/Private/DatasmithWorkerHandler.cpp
+++ b/Engine/Plugins/Enterprise/DatasmithCADImporter/Source/DatasmithDispatcher/Private/DatasmithWorkerHandler.cpp
@@ -3,232 +3,219 @@
 
 #include "HAL/FileManager.h"
 #include "HAL/PlatformProcess.h"
-#include "Logging/LogMacros.h"
 #include "Misc/Paths.h"
 #include "SocketSubsystem.h"
 #include "Sockets.h"
 #include "DatasmithDispatcher.h"
+#include "DatasmithDispatcherLog.h"
 #include "DatasmithCommands.h"
 
-//DEFINE_LOG_CATEGORY_STATIC(LogDatasmithDispatcher, Log, All);
-
 
 namespace DatasmithDispatcher
 {
-	FString FDatasmithWorkerHandler::GetExePath()
-	{
-		FString ProcessorPath;
+FString FDatasmithWorkerHandler::GetExePath()
+{
+	FString ProcessorPath;
 #if PLATFORM_MAC
-		ProcessorPath = FPaths::Combine(FPaths::EngineDir(), TEXT("Binaries/Mac/NotForLicensees/DatasmithCoretechWorker"));
+	ProcessorPath = FPaths::Combine(FPaths::EngineDir(), TEXT("Binaries/Mac/NotForLicensees/DatasmithCADWorker"));
 #elif PLATFORM_LINUX
-		ProcessorPath = FPaths::Combine(FPaths::EngineDir(), TEXT("Binaries/Linux/NotForLicensees/DatasmithCoretechWorker"));
+	ProcessorPath = FPaths::Combine(FPaths::EngineDir(), TEXT("Binaries/Linux/NotForLicensees/DatasmithCADWorker"));
 #else
-		ProcessorPath = FPaths::Combine(FPaths::EngineDir(), TEXT("Binaries/Win64/NotForLicensees/DatasmithCoretechWorker.exe"));
+	ProcessorPath = FPaths::Combine(FPaths::EngineDir(), TEXT("Binaries/Win64/NotForLicensees/DatasmithCADWorker.exe"));
 #endif
 
-		if (FPaths::FileExists(ProcessorPath))
-		{
-			return ProcessorPath;
-		}
-		return FString();
+	if (FPaths::FileExists(ProcessorPath))
+	{
+		return ProcessorPath;
 	}
+	return FString();
+}
+
+FDatasmithWorkerHandler::FDatasmithWorkerHandler(FDatasmithDispatcher& InDispatcher, const FString& InCachePath)
+	: Dispatcher(InDispatcher)
+	, WorkerProcessId(0)
+	, ServerSocket(TEXT("127.0.0.1"))
+	, ClientListener()
+	, CachePath(InCachePath)
+	, bShouldTerminate(false)
+{
+}
 
-	FDatasmithWorkerHandler::FDatasmithWorkerHandler(FDatasmithDispatcher& InDispatcher, const FString& InCachePath)
-		: Dispatcher(InDispatcher)
-		, WorkerProcessId(0)
-		, ServerSocket(TEXT("127.0.0.1"))
-		, ClientListener()
-		, CachePath(InCachePath)
-		, bShouldTerminate(false)
-		//, Stopwatch(0)
-		//, StartTime(0)
+bool FDatasmithWorkerHandler::StartWorker()
+{
+	FString ProcessorPath = GetExePath();
+	if (!ProcessorPath.Len())
 	{
+		return false;
 	}
 
-	bool FDatasmithWorkerHandler::StartWorker()
-	{
-		FString ProcessorPath = GetExePath();
-		if (!ProcessorPath.Len())
-		{
-			return false;
-		}
+	FString WorkerAbsolutePath = IFileManager::Get().ConvertToAbsolutePathForExternalAppForWrite(*ProcessorPath);
+	FPaths::NormalizeDirectoryName(WorkerAbsolutePath);
 
-		//StartTime = Stopwatch.GetDuration();
+	FString CommandToProcess = TEXT(" -Listen ") + FString::FromInt(ServerSocket.GetPort()) + TEXT(" -CacheDir \"") + GetCachePath() + TEXT("\"");
 
-		FString WorkerAbsolutePath = IFileManager::Get().ConvertToAbsolutePathForExternalAppForWrite(*ProcessorPath);
-		FPaths::NormalizeDirectoryName(WorkerAbsolutePath);
+	WorkerHandle = FPlatformProcess::CreateProc(*WorkerAbsolutePath, *CommandToProcess, true, false, false, &WorkerProcessId, 0, nullptr, nullptr);
+	//WorkerHandle = FPlatformProcess::CreateProc(*WorkerAbsolutePath, *CommandToProcess, false, false, false, &WorkerProcessId, 0, nullptr, nullptr);
 
-		FString CommandToProcess = TEXT(" -Listen ") + FString::FromInt(ServerSocket.GetPort()) + TEXT(" -CacheDir \"") + GetCachePath() + TEXT("\"");
+	if (!WorkerHandle.IsValid())
+	{
+		return false;
+	}
 
-		WorkerHandle = FPlatformProcess::CreateProc(*WorkerAbsolutePath, *CommandToProcess, true, false, false, &WorkerProcessId, 0, nullptr, nullptr);
-		//WorkerHandle = FPlatformProcess::CreateProc(*WorkerAbsolutePath, *CommandToProcess, false, false, false, &WorkerProcessId, 0, nullptr, nullptr);
+	FSocket *ClientSocket = ServerSocket.Accept();
+	if (!ClientSocket)
+	{
+		return false;
+	}
 
-		if (!WorkerHandle.IsValid())
-		{
-			return false;
-		}
+	ClientListener.SetSocket(ClientSocket);
+	return true;
+}
 
-		FSocket *ClientSocket = ServerSocket.Accept();
-		if (!ClientSocket)
-		{
-			return false;
-		}
 
-		ClientListener.SetSocket(ClientSocket);
-		return true;
-	}
+void FDatasmithWorkerHandler::RestartProcessor()
+{
+	FPlatformProcess::TerminateProc(WorkerHandle, true);
+	StartWorker();
+}
 
+bool FDatasmithWorkerHandler::InitializeSocket()
+{
+	ServerSocket.Bind();
+	return ServerSocket.IsOpen();
+}
+
+bool FDatasmithWorkerHandler::Run()
+{
 
-	void FDatasmithWorkerHandler::RestartProcessor()
+	if (!InitializeSocket())
 	{
-		FPlatformProcess::TerminateProc(WorkerHandle, true);
-		StartWorker();
+		return false;
 	}
 
-	bool FDatasmithWorkerHandler::InitializeSocket()
+	if(!ServerSocket.Listen())
 	{
-		ServerSocket.Bind();
-		if(!ServerSocket.IsOpen())
-		{
-			return false;
-		}
-		return true; 
+		return false;
 	}
 
-	bool FDatasmithWorkerHandler::Run()
+	if (!StartWorker())
 	{
+		return false;
+	}
 
-		if (!InitializeSocket())
-		{
-			return false;
-		}
+	DatasmithDispatcher::FDatasmithCommandManager CommandManager(ClientListener);
 
-		if(!ServerSocket.Listen())
-		{
-			return false;
-		}
+	int test = 0;
+	bool bTaskIsRunning = true;
 
-		if (!StartWorker())
+	while (NeedToRun())
+	{
+
+		if (!ClientListener.IsConnected())
 		{
-			return false;
+			if (CurrentTask.IsSet())
+			{
+				Dispatcher.SetTaskState(CurrentTask->Index, EProcessState::ProcessFailed);
+				CurrentTask.Reset();
+			}
+			RestartProcessor();
+			bTaskIsRunning = true;
+			continue;
 		}
 
-		DatasmithDispatcher::FDatasmithCommandManager CommandManager(ClientListener);
-
-		int test = 0;
-		bool bTaskIsRunning = true;
-
-		while (NeedToRun())
+		// If No task running, launch one
+		uint32 DataSize;
+		ClientListener.HasPendingData(DataSize);
+		if (!bTaskIsRunning && !DataSize)
 		{
-
-			if (!ClientListener.IsConnected())
+			CurrentTask = Dispatcher.GetTask();
+			if (CurrentTask.IsSet())
 			{
-				if (CurrentTask.IsSet())
-				{
-					Dispatcher.SetTaskState(CurrentTask->Index, ProcessFailed);
-					CurrentTask.Reset();
-				}
-				RestartProcessor();
+				FDatasmithRunTaskCommand AddJob(CurrentTask.GetValue().FileName, CurrentTask.GetValue().Index);
+				AddJob.Write(ClientListener);
 				bTaskIsRunning = true;
 				continue;
 			}
+		}
 
-			// If No task running, launch one
-			uint32 DataSize;
-			ClientListener.HasPendingData(DataSize);
-			if (!bTaskIsRunning && !DataSize)
+		if (DataSize)
+		{
+			if (ICommand* Command = CommandManager.GetNextCommand())
 			{
-				CurrentTask = Dispatcher.GetTask();
-				if (CurrentTask.IsSet())
+				switch (Command->GetType())
 				{
-					FDatasmithRunTaskCommand AddJob(CurrentTask.GetValue().FileName, CurrentTask.GetValue().Index);
-					AddJob.Write(ClientListener);
-					bTaskIsRunning = true;
-					continue;
+				case Ping:
+					PingProcess(StaticCast<FDatasmithPingCommand*>(Command));
+					bTaskIsRunning = false;
+					break;
+				case BackPing:
+					BackPingProcess(StaticCast<FDatasmithBackPingCommand*>(Command));
+					break;
+				case NotifyEndTask:
+				{
+					EndTaskProcess(StaticCast<FDatasmithNotifyEndTaskCommand*>(Command));
+					bTaskIsRunning = false;
+					break;
 				}
-			}
 
-			if (DataSize)
-			{
-				if (ICommand* Command = CommandManager.GetNextCommand())
-				{
-					switch (Command->GetType())
-					{
-					case Ping:
-						PingProcess(StaticCast<FDatasmithPingCommand*>(Command));
-						bTaskIsRunning = false;
-						break;
-					case BackPing:
-						BackPingProcess(StaticCast<FDatasmithBackPingCommand*>(Command));
-						break;
-					case NotifyEndTask:
-					{
-						EndTaskProcess(StaticCast<FDatasmithNotifyEndTaskCommand*>(Command));
-						bTaskIsRunning = false;
-						break;
-					}
-
-					case ImportParams:
-					case RunTask:
-					default:
-						break;
-					}
+				case ImportParams:
+				case RunTask:
+				default:
+					break;
 				}
 			}
-
-			FWindowsPlatformProcess::Sleep(0.001);
 		}
 
-		Kill();
-		return true;
+		FWindowsPlatformProcess::Sleep(0.001);
 	}
 
-	void FDatasmithWorkerHandler::PingProcess(FDatasmithPingCommand* PingCommand)
-	{
-		DatasmithDispatcher::FDatasmithBackPingCommand NewCommand;
-		NewCommand.Write(ClientListener);
-	}
+	Kill();
+	return true;
+}
 
-	void FDatasmithWorkerHandler::BackPingProcess(FDatasmithBackPingCommand* PingCommand)
-	{
-		//FTimespan PingTime = Stopwatch.GetDuration() - StartTime;
-		//UE_LOG(LogDatasmithDispatcher, Display, TEXT("Connected in %d s"), PingTime.GetTotalSeconds());
+void FDatasmithWorkerHandler::PingProcess(FDatasmithPingCommand* PingCommand)
+{
+	DatasmithDispatcher::FDatasmithBackPingCommand NewCommand;
+	NewCommand.Write(ClientListener);
+}
 
-		return ;
-	}
+void FDatasmithWorkerHandler::BackPingProcess(FDatasmithBackPingCommand* PingCommand)
+{
+	return ;
+}
 
-	void FDatasmithWorkerHandler::EndTaskProcess(FDatasmithNotifyEndTaskCommand* RunTaskCommand)
+void FDatasmithWorkerHandler::EndTaskProcess(FDatasmithNotifyEndTaskCommand* RunTaskCommand)
+{
+	if (!CurrentTask.IsSet())
 	{
-		if (!CurrentTask.IsSet())
-		{
-			return;
-		}
-
-		FString CurrentPath = FPaths::GetPath(CurrentTask->FileName);
-		for (const FString& ExternalReferenceFile : RunTaskCommand->GetExternalReferences())
-		{
-			Dispatcher.AddTask(FPaths::Combine(CurrentPath, ExternalReferenceFile));
-		}
-		FString CurrentFileName = FPaths::GetCleanFilename(CurrentTask->FileName);
-		Dispatcher.SetTaskState(CurrentTask->Index, RunTaskCommand->GetProcessResult());
-		Dispatcher.LinkCTFileToUnrealCacheFile(CurrentFileName, RunTaskCommand->GetSceneGraphFile(), RunTaskCommand->GetGeomFile());
-		CurrentTask.Reset();
+		return;
 	}
 
-	void FDatasmithWorkerHandler::NotifyKill()
+	FString CurrentPath = FPaths::GetPath(CurrentTask->FileName);
+	for (const FString& ExternalReferenceFile : RunTaskCommand->GetExternalReferences())
 	{
-		bShouldTerminate = true;
+		Dispatcher.AddTask(FPaths::Combine(CurrentPath, ExternalReferenceFile));
 	}
+	FString CurrentFileName = FPaths::GetCleanFilename(CurrentTask->FileName);
+	Dispatcher.SetTaskState(CurrentTask->Index, RunTaskCommand->GetProcessResult());
+	Dispatcher.LinkCTFileToUnrealCacheFile(CurrentFileName, RunTaskCommand->GetSceneGraphFile(), RunTaskCommand->GetGeomFile());
+	CurrentTask.Reset();
+}
 
-	bool FDatasmithWorkerHandler::NeedToRun()
-	{
-		return !bShouldTerminate;
-	}
+void FDatasmithWorkerHandler::NotifyKill()
+{
+	bShouldTerminate = true;
+}
 
-	void FDatasmithWorkerHandler::Kill()
-	{
-		ClientListener.Close();
-		ServerSocket.Close();
-		FPlatformProcess::TerminateProc(WorkerHandle, true);
-	}
+bool FDatasmithWorkerHandler::NeedToRun()
+{
+	return !bShouldTerminate;
+}
+
+void FDatasmithWorkerHandler::Kill()
+{
+	ClientListener.Close();
+	ServerSocket.Close();
+	FPlatformProcess::TerminateProc(WorkerHandle, true);
+}
 }
diff --git a/Engine/Plugins/Enterprise/DatasmithCADImporter/Source/DatasmithDispatcher/Public/DatasmithCommands.h b/Engine/Plugins/Enterprise/DatasmithCADImporter/Source/DatasmithDispatcher/Public/DatasmithCommands.h
--- a/Engine/Plugins/Enterprise/DatasmithCADImporter/Source/DatasmithDispatcher/Public/DatasmithCommands.h
+++ b/Engine/Plugins/Enterprise/DatasmithCADImporter/Source/DatasmithDispatcher/Public/DatasmithCommands.h
@@ -6,7 +6,7 @@
 #include "CADOptions.h"
 #include "HAL/Runnable.h"
 #include "Misc/DateTime.h"
-
+#include "CoreTechFileParser.h"
 
 class FCoreTechFileParser;
 
@@ -17,284 +17,244 @@ class FCoreTechFileParser;
 
 namespace DatasmithDispatcher
 {
-
-	enum EProcessState
-	{
-		Unknown = 0x00,
-		UnTreated = 0x01,
-		Running = 0x02,
-		ProcessOk = 0x04,
-		ProcessFailed = 0x08,
-		FileNotFound = 0x10
-	};
-
-	struct FTask
+struct FTask
+{
+	FTask(const FString& InFile)
 	{
-		FTask(const FString& InFile)
-		{
-			FileName = InFile;
-			State = UnTreated;
-		}
+		FileName = InFile;
+		State = UnTreated;
+	}
 
-		FTask()
-		{
-		}
-
-		FString FileName;
-		int32 Index = -1;
-		EProcessState State = Unknown;
-	};
-
-	class FDatasmithDispatcherSocket;
-	class FDatasmithDispatcher;
-
-	/**
-	 * Enum with all command types
-	 */
-	enum DatasmithCommandType
-	{
-		Ping,
-		BackPing,
-		RunTask,
-		NotifyEndTask,
-		ImportParams,
-		Last
-	};
-
-	/**
-	 * Command interface
-	 * Serializable for Socket
-	 */
-	class ICommand
+	FTask()
 	{
-	public:
-		//Initialize command
-		virtual void Initialize() = 0;
+	}
+
+	FString FileName;
+	int32 Index = -1;
+	EProcessState State = Unknown;
+};
 
-		//Get command type
-		virtual DatasmithCommandType GetType() = 0;
+class FDatasmithDispatcherSocket;
+class FDatasmithDispatcher;
 
-		//Write to stream
-		virtual void Write(FDatasmithDispatcherSocket&) = 0;
+/**
+ * Enum with all command types
+ */
+enum DatasmithCommandType
+{
+	Ping,
+	BackPing,
+	RunTask,
+	NotifyEndTask,
+	ImportParams,
+	Last
+};
 
-		//Read from stream
-		virtual void Read(FDatasmithDispatcherSocket&) = 0;
-	};
+/**
+ * Command interface
+ * Serializable for Socket
+ */
+class ICommand
+{
+public:
+	virtual ~ICommand() = default;
 
-	/**
-	 * Process commands received by client listener (here GPCTClientListener)
-	 */
-	class FDatasmithCommandManager
+	virtual void Initialize()
 	{
-	public:
-		//Initialize command manager (add default commands)
-		FDatasmithCommandManager(FDatasmithDispatcherSocket& CurrentSocket);
+	}
 
-		//Destructor
-		~FDatasmithCommandManager();
+	virtual DatasmithCommandType GetType() = 0;
 
-		void Initialize();
-		void Terminate();
+	virtual void Write(FDatasmithDispatcherSocket&) = 0;
+	virtual void Read(FDatasmithDispatcherSocket&) = 0;
+};
+	
+/**
+ * Process commands received by client listener (here GPCTClientListener)
+ */
+class FDatasmithCommandManager
+{
+public:
+	FDatasmithCommandManager(FDatasmithDispatcherSocket& CurrentSocket);
 
-		ICommand *GetNextCommand();
+	~FDatasmithCommandManager();
 
-	private:
-		//Command map (map command type with command implementation -> for Reading)
-		TMap<DatasmithCommandType, ICommand*> CommandMap;
+	void Initialize();
+	void Terminate();
 
-		FDatasmithDispatcherSocket& Socket;
-		uint32 CurrentCommandeSize;
-	};
+	ICommand *GetNextCommand();
 
-	/**
-	 * "Template" command. Use this command if you want to implement new ones -> this is one of the simplest possible commands
-	 */
-	class FDatasmithPingCommand : public ICommand
-	{
-	public:
-		//Initialize command
-		FDatasmithPingCommand() {}
+private:
+	// Command map (map command type with command implementation -> for Reading)
+	TMap<DatasmithCommandType, ICommand*> CommandMap;
 
-		virtual void Initialize() override
-		{
-		}
+	FDatasmithDispatcherSocket& Socket;
+	uint32 CurrentCommandeSize;
+};
 
-		// Inherited via ICommand
-		virtual DatasmithCommandType GetType() override
-		{
-			return Ping;
-		}
-		virtual void Write(FDatasmithDispatcherSocket&) override;
-		virtual void Read(FDatasmithDispatcherSocket&) override;
+/**
+ * "Template" command. Use this command if you want to implement new ones -> this is one of the simplest possible commands
+ */
+class FDatasmithPingCommand : public ICommand
+{
+public:
+	virtual DatasmithCommandType GetType() override
+	{
+		return Ping;
+	}
+	virtual void Write(FDatasmithDispatcherSocket&) override;
+	virtual void Read(FDatasmithDispatcherSocket&) override;
 
-	private:
+private:
 
-	};
+};
 
-	class FDatasmithBackPingCommand : public ICommand
+class FDatasmithBackPingCommand : public ICommand
+{
+public:
+	virtual DatasmithCommandType GetType() override
 	{
-	public:
-		//Initialize command
-		FDatasmithBackPingCommand() {}
+		return BackPing;
+	}
+	virtual void Write(FDatasmithDispatcherSocket&) override;
+	virtual void Read(FDatasmithDispatcherSocket&) override;
 
-		virtual void Initialize() override
-		{
-		}
-
-		// Inherited via ICommand
-		virtual DatasmithCommandType GetType() override
-		{
-			return BackPing;
-		}
-		virtual void Write(FDatasmithDispatcherSocket&) override;
-		virtual void Read(FDatasmithDispatcherSocket&) override;
+private:
 
-	private:
+};
 
-	};
+class FDatasmithRunTaskCommand : public ICommand
+{
+public:
+	FDatasmithRunTaskCommand()
+	{
+	}
 
-	class FDatasmithRunTaskCommand : public ICommand
+	FDatasmithRunTaskCommand(const FString& InOutputFile, const int32 InJobIndex) 
+		: JobFilePath(InOutputFile)
+		, JobIndex(InJobIndex)
 	{
-	public:
-		//Constructors
-		FDatasmithRunTaskCommand()
-		{
-		}
-		FDatasmithRunTaskCommand(const FString& InOutputFile, const int32 InJobIndex) 
-			: JobFilePath(InOutputFile)
-			, JobIndex(InJobIndex)
-		{
-		}
+	}
 
-		virtual ~FDatasmithRunTaskCommand()
-		{
-		}
+	virtual ~FDatasmithRunTaskCommand()
+	{
+	}
 
-		const FString& GetFileToProcess() const
-		{
-			return JobFilePath;
-		}
+	const FString& GetFileToProcess() const
+	{
+		return JobFilePath;
+	}
 
-		virtual void Initialize() override
-		{
-			JobFilePath.Empty();
-			JobIndex = 0;
-		}
+	virtual void Initialize() override
+	{
+		JobFilePath.Empty();
+		JobIndex = 0;
+	}
 
-		// Inherited via ICommand
-		virtual DatasmithCommandType GetType() override
-		{
-			return RunTask;
-		}
+	virtual DatasmithCommandType GetType() override
+	{
+		return RunTask;
+	}
 
-		virtual void Write(FDatasmithDispatcherSocket&) override;
-		virtual void Read(FDatasmithDispatcherSocket&) override;
+	virtual void Write(FDatasmithDispatcherSocket&) override;
+	virtual void Read(FDatasmithDispatcherSocket&) override;
 
-	private:
-		FString JobFilePath;
-		int32 JobIndex;
-	};
+private:
+	FString JobFilePath;
+	int32 JobIndex;
+};
 
-	class FDatasmithNotifyEndTaskCommand : public ICommand
+class FDatasmithNotifyEndTaskCommand : public ICommand
+{
+public:
+	void SetExternalReferences(const TSet<FString>& ExternalRefSet)
 	{
-	public:
-		virtual ~FDatasmithNotifyEndTaskCommand()
+		ExternalReferenceSet.Reserve(ExternalRefSet.Num());
+		for (const FString& ExternalFile : ExternalRefSet)
 		{
+			ExternalReferenceSet.Add(ExternalFile);
 		}
-
-		void SetExternalReferences(const TArray<FString>& ExternalRef)
-		{
-			ExternalReferenceSet = ExternalRef;
-		}
+	}
 	
-		const TArray<FString>& GetExternalReferences()
-		{
-			return ExternalReferenceSet;
-		}
+	const TArray<FString>& GetExternalReferences()
+	{
+		return ExternalReferenceSet;
+	}
 
-		void SetProcessResult(EProcessState InProcessResult)
-		{
-			ProcessResult = InProcessResult;
-		}
+	void SetProcessResult(EProcessState InProcessResult)
+	{
+		ProcessResult = InProcessResult;
+	}
 
-		EProcessState GetProcessResult()
-		{
-			return ProcessResult;
-		}
+	EProcessState GetProcessResult()
+	{
+		return ProcessResult;
+	}
 
-		void SetSceneGraphFile(const FString& InSceneGraphFileName)
-		{
-			SceneGraphFileName = InSceneGraphFileName;
-		}
+	void SetSceneGraphFile(const FString& InSceneGraphFileName)
+	{
+		SceneGraphFileName = InSceneGraphFileName;
+	}
 
-		const FString& GetSceneGraphFile()
-		{
-			return SceneGraphFileName;
-		}
+	const FString& GetSceneGraphFile()
+	{
+		return SceneGraphFileName;
+	}
 
-		void SetGeomFile(const FString& InGeomFile)
-		{
-			GeomFileName = InGeomFile;
-		}
+	void SetGeomFile(const FString& InGeomFile)
+	{
+		GeomFileName = InGeomFile;
+	}
 
-		const FString& GetGeomFile()
-		{
-			return GeomFileName;
-		}
+	const FString& GetGeomFile()
+	{
+		return GeomFileName;
+	}
 
-		virtual void Initialize() override
-		{
-			ExternalReferenceSet.Empty();
-			ProcessResult = UnTreated;
-		}
+	virtual void Initialize()
+	{
+		ExternalReferenceSet.Empty();
+		ProcessResult = EProcessState::UnTreated;
+	}
 
-		virtual DatasmithCommandType GetType() override
-		{
-			return NotifyEndTask;
-		}
+	virtual DatasmithCommandType GetType() override
+	{
+		return NotifyEndTask;
+	}
 
-		virtual void Write(FDatasmithDispatcherSocket&) override;
-		virtual void Read(FDatasmithDispatcherSocket&) override;
+	virtual void Write(FDatasmithDispatcherSocket&) override;
+	virtual void Read(FDatasmithDispatcherSocket&) override;
 
-	private:
-		TArray<FString> ExternalReferenceSet;
-		EProcessState ProcessResult;
-		FString SceneGraphFileName;
-		FString GeomFileName;
-	};
+private:
+	TArray<FString> ExternalReferenceSet;
+	EProcessState ProcessResult;
+	FString SceneGraphFileName;
+	FString GeomFileName;
+};
 
-	class FDatasmithImportParametersCommand : public ICommand
+class FDatasmithImportParametersCommand : public ICommand
+{
+public:
+	virtual DatasmithCommandType GetType() override
 	{
-	public:
-		virtual ~FDatasmithImportParametersCommand()
-		{
-		}
-
-		virtual DatasmithCommandType GetType() override
-		{
-			return ImportParams;
-		}
+		return ImportParams;
+	}
 		
-		virtual void Initialize() override
-		{
-			ImportParameters.init();
-		}
-
-		virtual void Set(const CAD::FImportParameters& InImportParameters)
-		{
-			ImportParameters = InImportParameters;
-		}
+	virtual void Set(const CADLibrary::FImportParameters& InImportParameters)
+	{
+		ImportParameters = InImportParameters;
+	}
 
-		virtual void Get(CAD::FImportParameters& OutImportParameters)
-		{
-			OutImportParameters = ImportParameters;
-		}
+	virtual void Get(CADLibrary::FImportParameters& OutImportParameters)
+	{
+		OutImportParameters = ImportParameters;
+	}
 
-		virtual void Write(FDatasmithDispatcherSocket&) override;
-		virtual void Read(FDatasmithDispatcherSocket&) override;
+	virtual void Write(FDatasmithDispatcherSocket&) override;
+	virtual void Read(FDatasmithDispatcherSocket&) override;
 
-	private:
-		CAD::FImportParameters ImportParameters;
-	};
+private:
+	CADLibrary::FImportParameters ImportParameters;
+};
 
 } // NS DatasmithDispatcher
diff --git a/Engine/Plugins/Enterprise/DatasmithCADImporter/Source/DatasmithDispatcher/Public/DatasmithDispatcher.h b/Engine/Plugins/Enterprise/DatasmithCADImporter/Source/DatasmithDispatcher/Public/DatasmithDispatcher.h
--- a/Engine/Plugins/Enterprise/DatasmithCADImporter/Source/DatasmithDispatcher/Public/DatasmithDispatcher.h
+++ b/Engine/Plugins/Enterprise/DatasmithCADImporter/Source/DatasmithDispatcher/Public/DatasmithDispatcher.h
@@ -11,53 +11,59 @@
 namespace DatasmithDispatcher
 {
 
+//Wrapper class to launch processors
+class DATASMITHDISPATCHER_API FDatasmithDispatcher
+{
+
+public:
+	//Constructor
+	FDatasmithDispatcher(const FString& InCacheDir, int32 InNumberOfWorkers, TMap<FString, FString>& CADFileToUnrealFileMap, TMap<FString, FString>& CADFileToUnrealGeomMap);
+
+	void Init(const CADLibrary::FImportParameters& ImportParameters);
+
+	/**
+		* Spawns worker handlers
+		*/
+	void RunHandlers();
+	void Close();
 
-	//Wrapper class to launch processors
-	class DATASMITHDISPATCHER_API FDatasmithDispatcher
-	{
+	void Process(bool bWithProcessor);
+	void ProcessLocal();
 
-	public:
-		//Constructor
-		FDatasmithDispatcher(const TCHAR* InCacheDir, int32 InNumberOfWorkers, TMap<FString, FString>& CADFileToUnrealFileMap, TMap<FString, FString>& CADFileToUnrealGeomMap);
-		~FDatasmithDispatcher();
+	void Clear();
 
-		/**
-		 * Spawns worker handlers
-		 */
-		void RunHandlers();
-		void Close();
+	void AddTask(const FString& FileName); // Thread safe with FScopeLock Lock(&TaskPoolCriticalSection);
 
-		void Process();
-		void Clear();
+	TOptional<FTask> GetTask();   // Thread safe with FScopeLock Lock(&TaskPoolCriticalSection);
+	bool IsOver(); // Thread safe with FScopeLock Lock(&TaskPoolCriticalSection);
 
-		void AddTask(const FString& FileName); // Thread safe with FScopeLock Lock(&TaskPoolCriticalSection);
+	void SetTaskState(int32 TaskIndex, EProcessState TaskState);   // Thread safe
+	void LinkCTFileToUnrealCacheFile(const FString& CTFile, const FString& UnrealSceneGraphFile, const FString& UnrealGeomFile);   // Thread safe
 
-		TOptional<FTask> GetTask();   // Thread safe with FScopeLock Lock(&TaskPoolCriticalSection);
-		bool IsOver(); // Thread safe with FScopeLock Lock(&TaskPoolCriticalSection);
+	void SetTessellationOptions(float InChordTolerance = 0.2f, float InMaxEdgeLength = 0.0f, float InNormalTolerance = 20.0f, CADLibrary::EStitchingTechnique InStitchingTechnique = CADLibrary::EStitchingTechnique::StitchingSew);
 
-		void SetTaskState(int32 TaskIndex, EProcessState TaskState);   // Thread safe
-		void LinkCTFileToUnrealCacheFile(const FString& CTFile, const FString& UnrealSceneGraphFile, const FString& UnrealGeomFile);   // Thread safe
 
+private:
+	FCriticalSection TaskPoolCriticalSection;
 
-	private:
-		FCriticalSection TaskPoolCriticalSection;
+	TMap<FString, FString>& CADFileToUnrealFileMap;
+	TMap<FString, FString>& CADFileToUnrealGeomMap;
 
-		FString ProcessCacheFolder;
-		TMap<FString, FString>& CADFileToUnrealFileMap;
-		TMap<FString, FString>& CADFileToUnrealGeomMap;
+	FString ProcessCacheFolder;
+	CADLibrary::FImportParameters ImportParameters;
 
-		//Started
-		bool bStarted;
+	//Started
+	bool bStarted;
 
-		int32 NumberOfWorkers;
-		TArray<FDatasmithWorkerHandler> WorkerHandlerSet;
-		TArray<FThread> WorkersRunning;
+	int32 NumberOfWorkers;
+	TArray<FDatasmithWorkerHandler> WorkerHandlerSet;
+	TArray<FThread> WorkersRunning;
 
-		TArray<FTask> TaskPool;
-		int32 LastLaunchTask;
-		int32 NumOfFinishedTask;
+	TArray<FTask> TaskPool;
+	int32 LastLaunchTask;
+	int32 NumOfFinishedTask;
 
 
-	}; // FDatasmithDispatcher
+}; // FDatasmithDispatcher
 
 } // NS DatasmithDispatcher
diff --git a/Engine/Plugins/Enterprise/DatasmithCADImporter/Source/DatasmithDispatcher/Public/DatasmithDispatcherLog.h b/Engine/Plugins/Enterprise/DatasmithCADImporter/Source/DatasmithDispatcher/Public/DatasmithDispatcherLog.h
new file mode 100644
--- /dev/null
+++ b/Engine/Plugins/Enterprise/DatasmithCADImporter/Source/DatasmithDispatcher/Public/DatasmithDispatcherLog.h
@@ -0,0 +1,7 @@
+// Copyright 1998-2019 Epic Games, Inc. All Rights Reserved.
+#pragma once
+
+#include "Logging/LogMacros.h"
+
+DECLARE_LOG_CATEGORY_EXTERN(LogDatasmithDispatcher, Log, All)
+
diff --git a/Engine/Plugins/Enterprise/DatasmithCADImporter/Source/DatasmithDispatcher/Public/DatasmithDispatcherSocket.h b/Engine/Plugins/Enterprise/DatasmithCADImporter/Source/DatasmithDispatcher/Public/DatasmithDispatcherSocket.h
--- a/Engine/Plugins/Enterprise/DatasmithCADImporter/Source/DatasmithDispatcher/Public/DatasmithDispatcherSocket.h
+++ b/Engine/Plugins/Enterprise/DatasmithCADImporter/Source/DatasmithDispatcher/Public/DatasmithDispatcherSocket.h
@@ -12,106 +12,106 @@
 namespace DatasmithDispatcher
 {
 
-	enum SocketErrorCode
-	{
-		NoError = 0,
-		UnableToReadOnSocket = 1,
-		UnableToSendData = 2,
-		CouldNotStartWSA = 3,
-		UnableToGetLocalAddress = 4,
-		ConnectionToServerFailed = 5
-	};
-
-	class FDatasmithDispatcherSocket 
-	{
+enum SocketErrorCode
+{
+	NoError = 0,
+	UnableToReadOnSocket = 1,
+	UnableToSendData = 2,
+	CouldNotStartWSA = 3,
+	UnableToGetLocalAddress = 4,
+	ConnectionToServerFailed = 5
+};
+
+class FDatasmithDispatcherSocket 
+{
 
-	public:
-		FDatasmithDispatcherSocket(const TCHAR* ServerAddress);
-		FDatasmithDispatcherSocket();
-		FDatasmithDispatcherSocket& operator=(const FDatasmithDispatcherSocket&) = delete;
+public:
+	FDatasmithDispatcherSocket(const TCHAR* ServerAddress);
+	FDatasmithDispatcherSocket();
+	FDatasmithDispatcherSocket& operator=(const FDatasmithDispatcherSocket&) = delete;
 			
-		void SetSocket(FSocket* InSocket);
+	void SetSocket(FSocket* InSocket);
 
-		int32 GetErrorCode() const
-		{
-			return ErrorCode;
-		}
+	int32 GetErrorCode() const
+	{
+		return ErrorCode;
+	}
 
-		void Connect(const int32& ServerPort);
+	void Connect(const int32& ServerPort);
 
-		bool IsConnected();
+	bool IsConnected();
 
-		void Bind();
+	void Bind();
 
-		bool Listen();
-		FSocket* Accept();
+	bool Listen();
+	FSocket* Accept();
 
-		int32 GetPort() const
+	int32 GetPort() const
+	{
+		if (!Socket)
 		{
-			if (!Socket)
-			{
-				return -1;
-			}
-			return Socket->GetPortNo();
+			return -1;
 		}
+		return Socket->GetPortNo();
+	}
 
-		void Close();
+	void Close();
 
-		template <typename T>
-		void operator >>(T &val)
-		{
-			Read(sizeof(val), (uint8*)&val);
-		}
+	template <typename T>
+	void operator >>(T &val)
+	{
+		Read(sizeof(val), (uint8*)&val);
+	}
 
-		template <typename T>
-		void operator <<(T val)
-		{
-			Write(sizeof(val), (uint8*)&val);
-		}
+	template <typename T>
+	void operator <<(T val)
+	{
+		Write(sizeof(val), (uint8*)&val);
+	}
 
-		void operator >>(FString& str);
-		void operator <<(const FString& str);
+	void operator >>(FString& str);
+	void operator <<(const FString& str);
 
-		bool IsClosed() const 
-		{
-			return !bOpen;
-		}
+	bool IsClosed() const 
+	{
+		return !bOpen;
+	}
 
-		bool IsOpen() const
-		{
-			return bOpen;
-		}
+	bool IsOpen() const
+	{
+		return bOpen;
+	}
 
-		FSocket* GetSocket()
-		{
-			return Socket;
-		}
+	FSocket* GetSocket()
+	{
+		return Socket;
+	}
 
-		bool HasPendingData(uint32& DataSize);
+	bool HasPendingData(uint32& DataSize);
 
-		/**
-		 * Actually send the buffered output cache
-		 */
-		void SendData();
+	/**
+		* Actually send the buffered output cache
+		*/
+	void SendData();
 
-	private:
-		// low-level reading function
-		void Read(int32 Num, uint8 *Buffer);
+private:
+	// low-level reading function
+	void Read(int32 Num, uint8 *Buffer);
 
-		// low-level writing function
-		void Write(int32 Num, uint8 *Buffer);
+	// low-level writing function
+	void Write(int32 Num, uint8 *Buffer);
 
-	private:
-		TArray<uint8> Cache;
+private:
+	TArray<uint8> Cache;
 
-		FString SocketAdress;
-		int32 SocketId;
-		FSocket* Socket;
+	FString SocketAdress;
+	int32 SocketId;
+	FSocket* Socket;
 
-		bool bServerSide;
-		bool bOpen;
+	bool bServerSide;
+	bool bOpen;
 
-		int32 ErrorCode;
-	};
+	int32 ErrorCode;
+};
 
 } //namespace DatasmithDispatcher
diff --git a/Engine/Plugins/Enterprise/DatasmithCADImporter/Source/DatasmithOpenNurbsTranslator/DatasmithOpenNurbsTranslator.Build.cs b/Engine/Plugins/Enterprise/DatasmithCADImporter/Source/DatasmithOpenNurbsTranslator/DatasmithOpenNurbsTranslator.Build.cs
--- a/Engine/Plugins/Enterprise/DatasmithCADImporter/Source/DatasmithOpenNurbsTranslator/DatasmithOpenNurbsTranslator.Build.cs
+++ b/Engine/Plugins/Enterprise/DatasmithCADImporter/Source/DatasmithOpenNurbsTranslator/DatasmithOpenNurbsTranslator.Build.cs
@@ -20,7 +20,8 @@ namespace UnrealBuildTool.Rules
 			PrivateDependencyModuleNames.AddRange(
 				new string[]
 				{
-					"CADLibrary",
+                    "CADInterfaces",
+                    "CADLibrary",
 					"Engine",
 					"DatasmithCore",
                     "DatasmithCoreTechExtension",
diff --git a/Engine/Plugins/Enterprise/DatasmithCADImporter/Source/DatasmithOpenNurbsTranslator/Private/DatasmithOpenNurbsTranslator.cpp b/Engine/Plugins/Enterprise/DatasmithCADImporter/Source/DatasmithOpenNurbsTranslator/Private/DatasmithOpenNurbsTranslator.cpp
--- a/Engine/Plugins/Enterprise/DatasmithCADImporter/Source/DatasmithOpenNurbsTranslator/Private/DatasmithOpenNurbsTranslator.cpp
+++ b/Engine/Plugins/Enterprise/DatasmithCADImporter/Source/DatasmithOpenNurbsTranslator/Private/DatasmithOpenNurbsTranslator.cpp
@@ -4,7 +4,6 @@
 #include "DatasmithOpenNurbsTranslatorModule.h"
 
 #ifdef CAD_LIBRARY
-#include "CADLibraryOptions.h"
 #include "CoreTechParametricSurfaceExtension.h"
 #endif
 #include "DatasmithMaterialElements.h"
@@ -612,7 +611,7 @@ public:
 		}
 
 #ifdef CAD_LIBRARY
-		LocalSession = CADLibrary::FRhinoCoretechWrapper::GetSharedSession(MetricUnit, ScalingFactor);
+		LocalSession = FRhinoCoretechWrapper::GetSharedSession(MetricUnit, ScalingFactor);
 #endif
 	}
 
@@ -635,6 +634,7 @@ public:
 	void SetTessellationOptions(const FDatasmithTessellationOptions& Options);
 	void SetOutputPath(const FString& Path) { OutputPath = Path; }
 	double GetScalingFactor () const { return ScalingFactor; }
+	double GetMetricUnit() const { return MetricUnit; }
 
 private:
 	void TranslateTextureMappingTable(const ON_ObjectArray<ON_TextureMapping>& TextureMappingTable);
@@ -677,7 +677,7 @@ private:
 	uint32 TessellationOptionsHash;
 
 #ifdef CAD_LIBRARY
-	TSharedPtr<CADLibrary::FRhinoCoretechWrapper> LocalSession;
+	TSharedPtr<FRhinoCoretechWrapper> LocalSession;
 #endif
 
 private:
@@ -3147,7 +3147,8 @@ bool FDatasmithOpenNurbsTranslator::LoadStaticMesh(const TSharedRef<IDatasmithMe
 				CoreTechData->SourceFile = CoretechFile;
 				CoreTechData->RawData = MoveTemp(ByteArray);
 				CoreTechData->SceneParameters.ModelCoordSys = uint8(FDatasmithUtils::EModelCoordSystem::ZUp_RightHanded);
-				CoreTechData->SceneParameters.ScaleFactor = Translator->GetScalingFactor(); // Scale is set according to the Rhino file unit
+				CoreTechData->SceneParameters.ScaleFactor = Translator->GetScalingFactor();
+				CoreTechData->SceneParameters.MetricUnit = Translator->GetMetricUnit(); 
 				CoreTechData->LastTessellationOptions = GetCommonTessellationOptions();
 				OutMeshPayload.AdditionalData.Add(CoreTechData);
 			}
diff --git a/Engine/Plugins/Enterprise/DatasmithCADImporter/Source/DatasmithOpenNurbsTranslator/Private/RhinoCoretechWrapper.cpp b/Engine/Plugins/Enterprise/DatasmithCADImporter/Source/DatasmithOpenNurbsTranslator/Private/RhinoCoretechWrapper.cpp
--- a/Engine/Plugins/Enterprise/DatasmithCADImporter/Source/DatasmithOpenNurbsTranslator/Private/RhinoCoretechWrapper.cpp
+++ b/Engine/Plugins/Enterprise/DatasmithCADImporter/Source/DatasmithOpenNurbsTranslator/Private/RhinoCoretechWrapper.cpp
@@ -20,419 +20,422 @@
 // current fix will split the uv plane in two parts.
 #define FIX_HOLE_IN_WHOLE_FACE 1
 
-namespace CADLibrary
-{
-	TWeakPtr<FRhinoCoretechWrapper> FRhinoCoretechWrapper::SharedSession;
+TWeakPtr<FRhinoCoretechWrapper> FRhinoCoretechWrapper::SharedSession;
 
-	template<class T>
-	void MarkUninitializedMemory(std::vector<T>& VectorOfT)
-	{
+template<class T>
+void MarkUninitializedMemory(std::vector<T>& VectorOfT)
+{
 #if MARK_UNINITIALIZED_MEMORY
-		for (auto& Value : VectorOfT)
-			Value = T(-1);
+	for (auto& Value : VectorOfT)
+		Value = T(-1);
 #endif // MARK_UNINITIALIZED_MEMORY
-	};
+};
 
-	// Handle surface state without u/v duplication
-	struct SurfaceInfo
+// Handle surface state without u/v duplication
+struct SurfaceInfo
+{
+	ON_NurbsSurface& surf;
+
+	enum Axis { U, V };
+	struct PerAxisInfo
 	{
-		ON_NurbsSurface& surf;
+		int Order; // Degree + 1
+		int CtrlVertCount; // number of control points
+		int KnotSize;  // ON knots
+		int KnotCount; // CT knots
+
+		std::vector<int> KnotsMult; // from ON, not relevant as we send n-plicated knots to CT (dbg only)
+		std::vector<CT_UINT32> KnotMul; // array of '1' for CT...
+		std::vector<double> Knots; // t values with superflux values
+
+		PerAxisInfo(Axis InA, ON_NurbsSurface& InSurf)
+			: surf(InSurf)
+			, a(InA)
+		{
+			Populate();
+		}
 
-		enum Axis { U, V };
-		struct PerAxisInfo
+		// increasing a nurbs degree doesn't change the shape, but this operation generates a new hull with new weights.
+		// This can fix exotic nurbs
+		void IncreaseDegree()
 		{
-			int Order; // Degree + 1
-			int CtrlVertCount; // number of control points
-			int KnotSize;  // ON knots
-			int KnotCount; // CT knots
-
-			std::vector<int> KnotsMult; // from ON, not relevant as we send n-plicated knots to CT (dbg only)
-			std::vector<CT_UINT32> KnotMul; // array of '1' for CT...
-			std::vector<double> Knots; // t values with superflux values
-
-			PerAxisInfo(Axis InA, ON_NurbsSurface& InSurf)
-				: surf(InSurf)
-				, a(InA)
-			{
-				Populate();
-			}
+			surf.IncreaseDegree(a, surf.Degree(a) + 1);
+			Populate();
+		}
 
-			// increasing a nurbs degree doesn't change the shape, but this operation generates a new hull with new weights.
-			// This can fix exotic nurbs
-			void IncreaseDegree()
+		// detect cases not handled by CT, that is knot vectors with multiplicity < order on either end
+		bool FixMultiplicity()
+		{
+			if (KnotsMult.front() + 1 < Order || KnotsMult.back() + 1 < Order)
 			{
-				surf.IncreaseDegree(a, surf.Degree(a) + 1);
-				Populate();
-			}
-
-			// detect cases not handled by CT, that is knot vectors with multiplicity < order on either end
-			bool FixMultiplicity()
-			{
-				if (KnotsMult.front() + 1 < Order || KnotsMult.back() + 1 < Order)
-				{
-					IncreaseDegree();
-					return true;
-				}
-				return false;
+				IncreaseDegree();
+				return true;
 			}
+			return false;
+		}
 
-		private:
-			void Populate()
-			{
-				Order = surf.Order(a);
-				CtrlVertCount = surf.CVCount(a);
-				KnotSize = Order + CtrlVertCount;
-				KnotCount = surf.KnotCount(a);
-
-				KnotsMult.resize(KnotSize - 2);
-				MarkUninitializedMemory(KnotsMult);
-				for (int i = 0; i < KnotSize - 2; ++i)
-					KnotsMult[i] = surf.KnotMultiplicity(a, i); // 0 and < Order + CV_count - 2
-
-				Knots.resize(KnotSize);
-				MarkUninitializedMemory(Knots);
-				Knots.front() = surf.SuperfluousKnot(a, 0);
-				for (int i = 0; i < KnotCount; ++i)
-					Knots[i + 1] = surf.Knot(a, i);
-				Knots.back() = surf.SuperfluousKnot(a, 1);
-
-				KnotMul.resize(KnotSize, 1);
-			}
+	private:
+		void Populate()
+		{
+			Order = surf.Order(a);
+			CtrlVertCount = surf.CVCount(a);
+			KnotSize = Order + CtrlVertCount;
+			KnotCount = surf.KnotCount(a);
+
+			KnotsMult.resize(KnotSize - 2);
+			MarkUninitializedMemory(KnotsMult);
+			for (int i = 0; i < KnotSize - 2; ++i)
+				KnotsMult[i] = surf.KnotMultiplicity(a, i); // 0 and < Order + CV_count - 2
+
+			Knots.resize(KnotSize);
+			MarkUninitializedMemory(Knots);
+			Knots.front() = surf.SuperfluousKnot(a, 0);
+			for (int i = 0; i < KnotCount; ++i)
+				Knots[i + 1] = surf.Knot(a, i);
+			Knots.back() = surf.SuperfluousKnot(a, 1);
+
+			KnotMul.resize(KnotSize, 1);
+		}
 
-			ON_NurbsSurface& surf;
-			Axis a;
-		};
+		ON_NurbsSurface& surf;
+		Axis a;
+	};
 
-		PerAxisInfo u;
-		PerAxisInfo v;
-		int CtrlVertDim; // Number of doubles per ctrl vertex
-		std::vector<double> CtrlHull; // [xyzw...]
+	PerAxisInfo u;
+	PerAxisInfo v;
+	int CtrlVertDim; // Number of doubles per ctrl vertex
+	std::vector<double> CtrlHull; // [xyzw...]
 
-		SurfaceInfo(ON_NurbsSurface& inSurf)
-			: surf(inSurf)
-			, u(U, surf)
-			, v(V, surf)
-			, CtrlVertDim(surf.CVSize())
-		{
-			BuildHull();
-		}
+	SurfaceInfo(ON_NurbsSurface& inSurf)
+		: surf(inSurf)
+		, u(U, surf)
+		, v(V, surf)
+		, CtrlVertDim(surf.CVSize())
+	{
+		BuildHull();
+	}
 
-		void BuildHull()
+	void BuildHull()
+	{
+		CtrlHull.resize(u.CtrlVertCount * v.CtrlVertCount * CtrlVertDim);
+		MarkUninitializedMemory(CtrlHull);
+		double* CtrlHullPtr = CtrlHull.data();
+		ON::point_style pt_style = surf.IsRational() ? ON::point_style::euclidean_rational : ON::point_style::not_rational;
+		for (int UIndex = 0; UIndex < u.CtrlVertCount; ++UIndex)
 		{
-			CtrlHull.resize(u.CtrlVertCount * v.CtrlVertCount * CtrlVertDim);
-			MarkUninitializedMemory(CtrlHull);
-			double* CtrlHullPtr = CtrlHull.data();
-			ON::point_style pt_style = surf.IsRational() ? ON::point_style::euclidean_rational : ON::point_style::not_rational;
-			for (int UIndex = 0; UIndex < u.CtrlVertCount; ++UIndex)
+			for (int VIndex = 0; VIndex < v.CtrlVertCount; ++VIndex, CtrlHullPtr += CtrlVertDim)
 			{
-				for (int VIndex = 0; VIndex < v.CtrlVertCount; ++VIndex, CtrlHullPtr += CtrlVertDim)
-				{
-					surf.GetCV(UIndex, VIndex, pt_style, CtrlHullPtr);
-				}
+				surf.GetCV(UIndex, VIndex, pt_style, CtrlHullPtr);
 			}
 		}
+	}
 
-		// CT doesn't allow weights < 0
-		void FixNegativeWeights()
+	// CT doesn't allow weights < 0
+	void FixNegativeWeights()
+	{
+		if (surf.IsRational())
 		{
-			if (surf.IsRational())
+			bool hasNegativeWeight = false;
+			for (int cvPointIndex = 0; cvPointIndex < u.CtrlVertCount * v.CtrlVertCount; ++cvPointIndex)
 			{
-				bool hasNegativeWeight = false;
-				for (int cvPointIndex = 0; cvPointIndex < u.CtrlVertCount * v.CtrlVertCount; ++cvPointIndex)
+				double pointDataOffset = cvPointIndex * CtrlVertDim;
+				double weight = CtrlHull[pointDataOffset + CtrlVertDim - 1];
+				if (weight < 0.)
 				{
-					double pointDataOffset = cvPointIndex * CtrlVertDim;
-					double weight = CtrlHull[pointDataOffset + CtrlVertDim - 1];
-					if (weight < 0.)
-					{
-						hasNegativeWeight = true;
-						break;
-					}
-				}
-
-				if (hasNegativeWeight)
-				{
-					u.IncreaseDegree();
-					v.IncreaseDegree();
-					BuildHull();
+					hasNegativeWeight = true;
+					break;
 				}
 			}
-		}
 
-		// CT doesn't allow multiplicity < order
-		void FixUnsupportedMultiplicity()
-		{
-			bool uOrderIncreased = u.FixMultiplicity();
-			bool vOrderIncreased = v.FixMultiplicity();
-			if (uOrderIncreased || vOrderIncreased)
+			if (hasNegativeWeight)
+			{
+				u.IncreaseDegree();
+				v.IncreaseDegree();
 				BuildHull();
+			}
 		}
+	}
 
-		void FixUnsupportedParameters()
-		{
-			FixNegativeWeights();
-			FixUnsupportedMultiplicity();
-		}
-	};
+	// CT doesn't allow multiplicity < order
+	void FixUnsupportedMultiplicity()
+	{
+		bool uOrderIncreased = u.FixMultiplicity();
+		bool vOrderIncreased = v.FixMultiplicity();
+		if (uOrderIncreased || vOrderIncreased)
+			BuildHull();
+	}
 
-	CT_OBJECT_ID CreateCTSurface(ON_NurbsSurface& Surface)
+	void FixUnsupportedParameters()
 	{
-		if (Surface.Dimension() < 3)
-			return 0;
-
-		SurfaceInfo si(Surface);
-		si.FixUnsupportedParameters();
-
-		CT_OBJECT_ID CTSurfaceID = 0;
-		CheckedCTError result = CT_SNURBS_IO::Create(CTSurfaceID,
-			si.u.Order, si.v.Order,
-			si.u.KnotSize, si.v.KnotSize,
-			si.u.CtrlVertCount, si.v.CtrlVertCount,
-			si.CtrlVertDim, si.CtrlHull.data(),
-			si.u.Knots.data(), si.v.Knots.data(),
-			si.u.KnotMul.data(), si.v.KnotMul.data()
-		);
-
-		return CTSurfaceID;
+		FixNegativeWeights();
+		FixUnsupportedMultiplicity();
 	}
+};
+
+CT_OBJECT_ID CreateCTSurface(ON_NurbsSurface& Surface)
+{
+	if (Surface.Dimension() < 3)
+		return 0;
+
+	SurfaceInfo si(Surface);
+	si.FixUnsupportedParameters();
+
+	CT_OBJECT_ID CTSurfaceID = 0;
+	CADLibrary::CheckedCTError result = CT_SNURBS_IO::Create(CTSurfaceID,
+		si.u.Order, si.v.Order,
+		si.u.KnotSize, si.v.KnotSize,
+		si.u.CtrlVertCount, si.v.CtrlVertCount,
+		si.CtrlVertDim, si.CtrlHull.data(),
+		si.u.Knots.data(), si.v.Knots.data(),
+		si.u.KnotMul.data(), si.v.KnotMul.data()
+	);
+
+	return CTSurfaceID;
+}
+
+void CreateCTFace_internal(const ON_BrepFace& Face, CT_LIST_IO& dest, ON_BoundingBox& outerBBox, ON_NurbsSurface& Surface, bool ignoreInner)
+{
+	CT_OBJECT_ID SurfaceID = CreateCTSurface(Surface);
+	if (SurfaceID == 0)
+		return;
 
-	void CreateCTFace_internal(const ON_BrepFace& Face, CT_LIST_IO& dest, ON_BoundingBox& outerBBox, ON_NurbsSurface& Surface, bool ignoreInner)
+	ON_BoundingBox realOuterBBox = Face.OuterLoop()->BoundingBox();
+	bool outerRedefined = realOuterBBox != outerBBox;
+
+	int LoopCount = ignoreInner ? 1 : Face.LoopCount();
+	CT_LIST_IO Loops;
+	for (int LoopIndex = 0; LoopIndex < LoopCount; ++LoopIndex)
 	{
-		CT_OBJECT_ID SurfaceID = CreateCTSurface(Surface);
-		if (SurfaceID == 0)
-			return;
+		const ON_BrepLoop& on_Loop = *Face.Loop(LoopIndex);
+		if (!on_Loop.IsValid())
+			continue;
 
-		ON_BoundingBox realOuterBBox = Face.OuterLoop()->BoundingBox();
-		bool outerRedefined = realOuterBBox != outerBBox;
+		ON_BrepLoop::TYPE onLoopType = on_Loop.m_type;
+		bool bIsOuter = (onLoopType == ON_BrepLoop::TYPE::outer);
+		//assert(bIsOuter == (LoopIndex == 0));
 
-		int LoopCount = ignoreInner ? 1 : Face.LoopCount();
-		CT_LIST_IO Loops;
-		for (int LoopIndex = 0; LoopIndex < LoopCount; ++LoopIndex)
+		CT_LIST_IO Coedges;
+		int TrimCount = on_Loop.TrimCount();
+		for (int i = 0; i < TrimCount; ++i)
 		{
-			const ON_BrepLoop& on_Loop = *Face.Loop(LoopIndex);
-			if (!on_Loop.IsValid())
+			CADLibrary::CheckedCTError err = IO_OK;
+			ON_BrepTrim& Trim = *on_Loop.Trim(i);
+
+			ON_BrepEdge* on_edge = Trim.Edge();
+			if (on_edge == nullptr)
+				continue;
+
+			ON_NurbsCurve nurbs_curve;
+			int nurbFormSuccess = Trim.GetNurbForm(nurbs_curve); // 0:Nok 1:Ok 2:OkBut
+			if (nurbFormSuccess == 0)
 				continue;
 
-			ON_BrepLoop::TYPE onLoopType = on_Loop.m_type;
-			bool bIsOuter = (onLoopType == ON_BrepLoop::TYPE::outer);
-			//assert(bIsOuter == (LoopIndex == 0));
+			CT_OBJECT_ID Coedge;
+			err = CT_COEDGE_IO::Create(Coedge, Trim.m_bRev3d ? CT_ORIENTATION::CT_REVERSE : CT_ORIENTATION::CT_FORWARD);
+			if (err != IO_OK)
+				continue;
 
-			CT_LIST_IO Coedges;
-			int TrimCount = on_Loop.TrimCount();
-			for (int i = 0; i < TrimCount; ++i)
+			// fill edge data
+			CT_UINT32 order = nurbs_curve.Order();
+
+			int KnotCount = nurbs_curve.KnotCount();
+			int CtrlVertCount = nurbs_curve.CVCount();
+			CT_UINT32 KnotSize = order + CtrlVertCount; // cvCount + degree - 1 for OpenNurb, cvCount + degree + 1 for OpenNurb,
+
+			// knot data
+			std::vector<double> Knots;
+			Knots.resize(KnotSize);
+			Knots.front() = nurbs_curve.SuperfluousKnot(0);
+			for (int j = 0; j < KnotCount; ++j)
 			{
-				CheckedCTError err = IO_OK;
-				ON_BrepTrim& Trim = *on_Loop.Trim(i);
-
-				ON_BrepEdge* on_edge = Trim.Edge();
-				if (on_edge == nullptr)
-					continue;
-
-				ON_NurbsCurve nurbs_curve;
-				int nurbFormSuccess = Trim.GetNurbForm(nurbs_curve); // 0:Nok 1:Ok 2:OkBut
-				if (nurbFormSuccess == 0)
-					continue;
-
-				CT_OBJECT_ID Coedge;
-				err = CT_COEDGE_IO::Create(Coedge, Trim.m_bRev3d ? CT_ORIENTATION::CT_REVERSE : CT_ORIENTATION::CT_FORWARD);
-				if (err != IO_OK)
-					continue;
-
-				// fill edge data
-				CT_UINT32 order = nurbs_curve.Order();
-
-				int KnotCount = nurbs_curve.KnotCount();
-				int CtrlVertCount = nurbs_curve.CVCount();
-				CT_UINT32 KnotSize = order + CtrlVertCount; // cvCount + degree - 1 for OpenNurb, cvCount + degree + 1 for OpenNurb,
-
-				// knot data
-				std::vector<double> Knots;
-				Knots.resize(KnotSize);
-				Knots.front() = nurbs_curve.SuperfluousKnot(0);
-				for (int j = 0; j < KnotCount; ++j)
-				{
-					Knots[j + 1] = nurbs_curve.Knot(j);
-				}
-				Knots.back() = nurbs_curve.SuperfluousKnot(1);
-
-				// Control hull
-				int ctrl_hull_dim = nurbs_curve.CVSize(); // = IsRational() ? Dim()+1 : Dim()
-				int ctrl_hull_size = nurbs_curve.CVCount();
-				std::vector<double> cv_data;
-				cv_data.resize(ctrl_hull_size * ctrl_hull_dim);
-				double* wp = cv_data.data();
-				for (int j = 0; j < ctrl_hull_size; ++j, wp += ctrl_hull_dim)
-				{
-					nurbs_curve.GetCV(j, nurbs_curve.IsRational() ? ON::point_style::euclidean_rational : ON::point_style::not_rational, wp);
-				}
-				if (outerRedefined && bIsOuter)
+				Knots[j + 1] = nurbs_curve.Knot(j);
+			}
+			Knots.back() = nurbs_curve.SuperfluousKnot(1);
+
+			// Control hull
+			int ctrl_hull_dim = nurbs_curve.CVSize(); // = IsRational() ? Dim()+1 : Dim()
+			int ctrl_hull_size = nurbs_curve.CVCount();
+			std::vector<double> cv_data;
+			cv_data.resize(ctrl_hull_size * ctrl_hull_dim);
+			double* wp = cv_data.data();
+			for (int j = 0; j < ctrl_hull_size; ++j, wp += ctrl_hull_dim)
+			{
+				nurbs_curve.GetCV(j, nurbs_curve.IsRational() ? ON::point_style::euclidean_rational : ON::point_style::not_rational, wp);
+			}
+			if (outerRedefined && bIsOuter)
+			{
+				for (int j = 0; j < cv_data.size(); j += ctrl_hull_dim)
 				{
-					for (int j = 0; j < cv_data.size(); j += ctrl_hull_dim)
-					{
-						cv_data[j] = min(max(cv_data[j], outerBBox.m_min.x), outerBBox.m_max.x);
-					}
+					cv_data[j] = min(max(cv_data[j], outerBBox.m_min.x), outerBBox.m_max.x);
 				}
-
-				// knot multiplicity (ignored as knots are stored multiple times already)
-				std::vector<CT_UINT32> knotMult;
-				knotMult.resize(KnotSize, 1);
-
-				ON_Interval dom = nurbs_curve.Domain();
-				CheckedCTError setUvCurveError = CT_COEDGE_IO::SetUVCurve(
-					Coedge,          /*!< [out] Id of created coedge */
-					order,           /*!< [in] Order of curve */
-					KnotSize,        /*!< [in] Knot vector size */
-					ctrl_hull_size,  /*!< [in] Control Hull Size */
-					ctrl_hull_dim,   /*!< [in] Control Hull Dimension (2 for non-rational or 3 for rational) */
-					cv_data.data(),  /*!< [in] Control hull array */
-					Knots.data(),    /*!< [in] Value knot array */
-					knotMult.data(), /*!< [in] Multiplicity Knot array */
-					dom.m_t[0],      /*!< [in] start parameter of coedge on the uv curve (t range=[knot[0], knot[knot_size-1]]) */
-					dom.m_t[1]       /*!< [in] end parameter of coedge on the uv curve (t range=[knot[0], knot[knot_size-1]]) */
-				);
-
-				if (!setUvCurveError)
-					continue;
-
-				Coedges.PushBack(Coedge);
 			}
 
-			CT_OBJECT_ID Loop;
-			CheckedCTError LoopCreationError = CT_LOOP_IO::Create(Loop, Coedges);
-			if (!LoopCreationError)
+			// knot multiplicity (ignored as knots are stored multiple times already)
+			std::vector<CT_UINT32> knotMult;
+			knotMult.resize(KnotSize, 1);
+
+			ON_Interval dom = nurbs_curve.Domain();
+			CADLibrary::CheckedCTError setUvCurveError = CT_COEDGE_IO::SetUVCurve(
+				Coedge,          /*!< [out] Id of created coedge */
+				order,           /*!< [in] Order of curve */
+				KnotSize,        /*!< [in] Knot vector size */
+				ctrl_hull_size,  /*!< [in] Control Hull Size */
+				ctrl_hull_dim,   /*!< [in] Control Hull Dimension (2 for non-rational or 3 for rational) */
+				cv_data.data(),  /*!< [in] Control hull array */
+				Knots.data(),    /*!< [in] Value knot array */
+				knotMult.data(), /*!< [in] Multiplicity Knot array */
+				dom.m_t[0],      /*!< [in] start parameter of coedge on the uv curve (t range=[knot[0], knot[knot_size-1]]) */
+				dom.m_t[1]       /*!< [in] end parameter of coedge on the uv curve (t range=[knot[0], knot[knot_size-1]]) */
+			);
+
+			if (!setUvCurveError)
 				continue;
 
-			Loops.PushBack(Loop);
+			Coedges.PushBack(Coedge);
 		}
 
-		CT_OBJECT_ID FaceID;
-		CT_ORIENTATION faceOrient = CT_ORIENTATION::CT_FORWARD;
-		CheckedCTError result = CT_FACE_IO::Create(FaceID, SurfaceID, faceOrient, Loops);
-		if (!result)
-			return;
-		dest.PushBack(FaceID);
+		CT_OBJECT_ID Loop;
+		CADLibrary::CheckedCTError LoopCreationError = CT_LOOP_IO::Create(Loop, Coedges);
+		if (!LoopCreationError)
+			continue;
+
+		Loops.PushBack(Loop);
 	}
 
-	void CreateCTFace(const ON_Brep& brep, const ON_BrepFace& Face, CT_LIST_IO& dest)
+	CT_OBJECT_ID FaceID;
+	CT_ORIENTATION faceOrient = CT_ORIENTATION::CT_FORWARD;
+	CADLibrary::CheckedCTError result = CT_FACE_IO::Create(FaceID, SurfaceID, faceOrient, Loops);
+	if (!result)
+		return;
+	dest.PushBack(FaceID);
+}
+
+void CreateCTFace(const ON_Brep& brep, const ON_BrepFace& Face, CT_LIST_IO& dest)
+{
+	const ON_BrepLoop* outerLoop = Face.OuterLoop();
+	if (outerLoop == nullptr)
 	{
-		const ON_BrepLoop* outerLoop = Face.OuterLoop();
-		if (outerLoop == nullptr)
-		{
-			return;
-		}
+		return;
+	}
 
-		ON_NurbsSurface Surface;
-		Face.NurbsSurface(&Surface);
+	ON_NurbsSurface Surface;
+	Face.NurbsSurface(&Surface);
 
-		ON_BoundingBox outerBBox = Face.OuterLoop()->BoundingBox();
+	ON_BoundingBox outerBBox = Face.OuterLoop()->BoundingBox();
 
 #if FIX_HOLE_IN_WHOLE_FACE
-		int LoopCount = Face.LoopCount();
-		bool bBadLoopHack = brep.LoopIsSurfaceBoundary(outerLoop->m_loop_index) && LoopCount >= 2;
-		if (bBadLoopHack)
+	int LoopCount = Face.LoopCount();
+	bool bBadLoopHack = brep.LoopIsSurfaceBoundary(outerLoop->m_loop_index) && LoopCount >= 2;
+	if (bBadLoopHack)
+	{
+		CT_LIST_IO Coedges;
+		int TrimCount = outerLoop->TrimCount();
+		bool hasSingularTrim = false;
+		for (int i = 0; i < TrimCount; ++i)
 		{
-			CT_LIST_IO Coedges;
-			int TrimCount = outerLoop->TrimCount();
-			bool hasSingularTrim = false;
-			for (int i = 0; i < TrimCount; ++i)
-			{
-				ON_BrepTrim& Trim = *outerLoop->Trim(i);
-				hasSingularTrim |= (Trim.m_type == ON_BrepTrim::TYPE::singular);
-			}
-			bBadLoopHack &= hasSingularTrim;
+			ON_BrepTrim& Trim = *outerLoop->Trim(i);
+			hasSingularTrim |= (Trim.m_type == ON_BrepTrim::TYPE::singular);
 		}
-
-		if (bBadLoopHack)
-		{
-			// we try to split in 2 faces: one west one east.
-			double span = outerBBox.Diagonal().x;
-
-			ON_BoundingBox innerBBox = Face.Loop(1)->BoundingBox();
-			for (int LoopIndex = 2; LoopIndex < LoopCount; ++LoopIndex)
-			{
-				const ON_BrepLoop& innerLoop = *Face.Loop(LoopIndex);
-				innerBBox.Union(innerLoop.BoundingBox());
-			}
-			double spaceWest = innerBBox.Min().x - outerBBox.Min().x;
-			double spaceEast = outerBBox.Max().x - innerBBox.Max().x;
-			double e = 0.01;
-			if (spaceEast < e*span && spaceWest < e*span)
-			{
-				// can't split... We could ignore inner loops:
-				CreateCTFace_internal(Face, dest, outerBBox, Surface, true);
-				return;
-			}
-
-			// split the outer bbox in a normal part (with outer) and a 'rest' (with no iner)
-			ON_BoundingBox outerTrimedPartBBox = outerBBox;
-			if (spaceEast > spaceWest)
-			{
-				outerBBox.m_max.x -= 0.5 * spaceEast;
-				outerTrimedPartBBox.m_min.x = outerBBox.m_max.x;
-			}
-			else
-			{
-				outerBBox.m_min.x += 0.5 * spaceWest;
-				outerTrimedPartBBox.m_max.x = outerBBox.m_min.x;
-			}
-
-			// We now need to create a bonus face to handle the hole we just created
-			CreateCTFace_internal(Face, dest, outerTrimedPartBBox, Surface, true);
-		}
-#endif // FIX_HOLE_IN_WHOLE_FACE
-
-		CreateCTFace_internal(Face, dest, outerBBox, Surface, false);
+		bBadLoopHack &= hasSingularTrim;
 	}
 
-	CheckedCTError FRhinoCoretechWrapper::AddBRep(ON_Brep& Brep)
+	if (bBadLoopHack)
 	{
-		CheckedCTError Result;
-		if (!IsSessionValid())
+		// we try to split in 2 faces: one west one east.
+		double span = outerBBox.Diagonal().x;
+
+		ON_BoundingBox innerBBox = Face.Loop(1)->BoundingBox();
+		for (int LoopIndex = 2; LoopIndex < LoopCount; ++LoopIndex)
 		{
-			Result.RaiseOtherError("bad session init");
-			return Result;
+			const ON_BrepLoop& innerLoop = *Face.Loop(LoopIndex);
+			innerBBox.Union(innerLoop.BoundingBox());
 		}
-
-		// Create ct faces
-		Brep.FlipReversedSurfaces();
-		CT_LIST_IO FaceList;
-		int FaceCount = Brep.m_F.Count();
-		for (int index = 0; index < FaceCount; index++)
+		double spaceWest = innerBBox.Min().x - outerBBox.Min().x;
+		double spaceEast = outerBBox.Max().x - innerBBox.Max().x;
+		double e = 0.01;
+		if (spaceEast < e*span && spaceWest < e*span)
 		{
-			const ON_BrepFace& on_face = Brep.m_F[index];
-			CreateCTFace(Brep, on_face, FaceList);
+			// can't split... We could ignore inner loops:
+			CreateCTFace_internal(Face, dest, outerBBox, Surface, true);
+			return;
 		}
 
-		if (FaceList.IsEmpty())
+		// split the outer bbox in a normal part (with outer) and a 'rest' (with no iner)
+		ON_BoundingBox outerTrimedPartBBox = outerBBox;
+		if (spaceEast > spaceWest)
 		{
-			return Result;
+			outerBBox.m_max.x -= 0.5 * spaceEast;
+			outerTrimedPartBBox.m_min.x = outerBBox.m_max.x;
 		}
+		else
+		{
+			outerBBox.m_min.x += 0.5 * spaceWest;
+			outerTrimedPartBBox.m_max.x = outerBBox.m_min.x;
+		}
+
+		// We now need to create a bonus face to handle the hole we just created
+		CreateCTFace_internal(Face, dest, outerTrimedPartBBox, Surface, true);
+	}
+#endif // FIX_HOLE_IN_WHOLE_FACE
+
+	CreateCTFace_internal(Face, dest, outerBBox, Surface, false);
+}
+
+
+CT_IO_ERROR FRhinoCoretechWrapper::Tessellate(FMeshDescription& Mesh, CADLibrary::FMeshParameters& MeshParameters)
+{
+	return CADLibrary::Tessellate(MainObjectId, ImportParams, Mesh, MeshParameters);
+}
 
-		// Create body from faces
-		CT_OBJECT_ID BodyID;
-		Result = CT_BODY_IO::CreateFromFaces(BodyID, CT_BODY_PROP::CT_BODY_PROP_EXACT | CT_BODY_PROP::CT_BODY_PROP_CLOSE, FaceList);
-		if (!Result)
-			return Result;
+CADLibrary::CheckedCTError FRhinoCoretechWrapper::AddBRep(ON_Brep& Brep)
+{
+	CADLibrary::CheckedCTError Result;
+	if (!IsSessionValid())
+	{
+		Result.RaiseOtherError("bad session init");
+		return Result;
+	}
 
-		CT_LIST_IO Bodies;
-		Bodies.PushBack(BodyID);
+	// Create ct faces
+	Brep.FlipReversedSurfaces();
+	CT_LIST_IO FaceList;
+	int FaceCount = Brep.m_F.Count();
+	for (int index = 0; index < FaceCount; index++)
+	{
+		const ON_BrepFace& on_face = Brep.m_F[index];
+		CreateCTFace(Brep, on_face, FaceList);
+	}
 
-		// Setup parenting
-		Result = CT_COMPONENT_IO::AddChildren(MainObjectId, Bodies);
+	if (FaceList.IsEmpty())
+	{
 		return Result;
 	}
 
+	// Create body from faces
+	CT_OBJECT_ID BodyID;
+	Result = CT_BODY_IO::CreateFromFaces(BodyID, CT_BODY_PROP::CT_BODY_PROP_EXACT | CT_BODY_PROP::CT_BODY_PROP_CLOSE, FaceList);
+	if (!Result)
+		return Result;
+
+	CT_LIST_IO Bodies;
+	Bodies.PushBack(BodyID);
+
+	// Setup parenting
+	Result = CT_COMPONENT_IO::AddChildren(MainObjectId, Bodies);
+	return Result;
+}
+
 
 
-	TSharedPtr<FRhinoCoretechWrapper> FRhinoCoretechWrapper::GetSharedSession(double SceneUnit, double ScaleFactor)
+TSharedPtr<FRhinoCoretechWrapper> FRhinoCoretechWrapper::GetSharedSession(double SceneUnit, double ScaleFactor)
+{
+	TSharedPtr<FRhinoCoretechWrapper> Session = SharedSession.Pin();
+	if (!Session.IsValid())
 	{
-		TSharedPtr<FRhinoCoretechWrapper> Session = SharedSession.Pin();
-		if (!Session.IsValid())
-		{
-			Session = MakeShared<FRhinoCoretechWrapper>(TEXT("Rh2CTSharedSession"), SceneUnit, ScaleFactor);
-			SharedSession = Session;
-		}
-		return Session;
+		Session = MakeShared<FRhinoCoretechWrapper>(TEXT("Rh2CTSharedSession"), SceneUnit, ScaleFactor);
+		SharedSession = Session;
 	}
+	return Session;
+}
 
-} // namespace CADLibrary
 #endif
diff --git a/Engine/Plugins/Enterprise/DatasmithCADImporter/Source/DatasmithOpenNurbsTranslator/Private/RhinoCoretechWrapper.h b/Engine/Plugins/Enterprise/DatasmithCADImporter/Source/DatasmithOpenNurbsTranslator/Private/RhinoCoretechWrapper.h
--- a/Engine/Plugins/Enterprise/DatasmithCADImporter/Source/DatasmithOpenNurbsTranslator/Private/RhinoCoretechWrapper.h
+++ b/Engine/Plugins/Enterprise/DatasmithCADImporter/Source/DatasmithOpenNurbsTranslator/Private/RhinoCoretechWrapper.h
@@ -9,30 +9,29 @@
 
 class ON_Brep;
 
-namespace CADLibrary
+class FRhinoCoretechWrapper : public CADLibrary::CTSession
 {
-	class FRhinoCoretechWrapper : public CTSession
+public:
+	/**
+	 * Make sure CT is initialized, and a main object is ready.
+	 * Handle input file unit and an output unit
+	 * @param InOwner
+	 * @param FileMetricUnit number of meters per file unit.
+	 * @param ScaleFactor scale factor to apply to the mesh to be in UE4 unit (cm).
+	 * eg. For a file in inches, arg should be 0.0254
+	 */
+	FRhinoCoretechWrapper(const TCHAR* InOwner, double FileMetricUnit, double ScaleFactor)
+		: CTSession(InOwner, FileMetricUnit, ScaleFactor)
 	{
-	public:
-		/**
-		 * Make sure CT is initialized, and a main object is ready.
-		 * Handle input file unit and an output unit
-		 * @param InOwner 
-		 * @param FileMetricUnit number of meters per file unit.
-		 * @param ScaleFactor scale factor to apply to the mesh to be in UE4 unit (cm).
-		 * eg. For a file in inches, arg should be 0.0254
-		 */
-		FRhinoCoretechWrapper(const TCHAR* InOwner, double FileMetricUnit, double ScaleFactor)
-			: CTSession(InOwner, FileMetricUnit, ScaleFactor)
-		{
-		}
-
-		CheckedCTError AddBRep(ON_Brep& brep);
-		static TSharedPtr<FRhinoCoretechWrapper> GetSharedSession(double SceneUnit, double ScaleFactor);
-
-	protected:
-		static TWeakPtr<FRhinoCoretechWrapper> SharedSession;
-	};
-
-} // namespace CADLibrary
+	}
+
+	CADLibrary::CheckedCTError AddBRep(ON_Brep& brep);
+	static TSharedPtr<FRhinoCoretechWrapper> GetSharedSession(double SceneUnit, double ScaleFactor);
+
+	CT_IO_ERROR Tessellate(FMeshDescription& Mesh, CADLibrary::FMeshParameters& MeshParameters);
+
+protected:
+	static TWeakPtr<FRhinoCoretechWrapper> SharedSession;
+};
+
 #endif
\ No newline at end of file
diff --git a/Engine/Plugins/Enterprise/DatasmithCADImporter/Source/DatasmithWireTranslator/DatasmithWireTranslator.Build.cs b/Engine/Plugins/Enterprise/DatasmithCADImporter/Source/DatasmithWireTranslator/DatasmithWireTranslator.Build.cs
--- a/Engine/Plugins/Enterprise/DatasmithCADImporter/Source/DatasmithWireTranslator/DatasmithWireTranslator.Build.cs
+++ b/Engine/Plugins/Enterprise/DatasmithCADImporter/Source/DatasmithWireTranslator/DatasmithWireTranslator.Build.cs
@@ -7,7 +7,7 @@ public class DatasmithWireTranslator : ModuleRules
 {
 	public DatasmithWireTranslator(ReadOnlyTargetRules Target) : base(Target)
 	{
-		PublicDependencyModuleNames.AddRange(
+        PublicDependencyModuleNames.AddRange(
 			new string[]
 			{
 				"Core",
@@ -18,8 +18,10 @@ public class DatasmithWireTranslator : ModuleRules
 		PrivateDependencyModuleNames.AddRange(
 			new string[]
 			{
-				"CADLibrary",
-				"DatasmithContent",
+                "CADInterfaces",
+                "CADLibrary",
+                "CADTools",
+                "DatasmithContent",
 				"DatasmithCore",
 				"DatasmithCoreTechExtension",
 				"DatasmithImporter",
diff --git a/Engine/Plugins/Enterprise/DatasmithCADImporter/Source/DatasmithWireTranslator/Private/AliasCoretechWrapper.cpp b/Engine/Plugins/Enterprise/DatasmithCADImporter/Source/DatasmithWireTranslator/Private/AliasCoretechWrapper.cpp
--- a/Engine/Plugins/Enterprise/DatasmithCADImporter/Source/DatasmithWireTranslator/Private/AliasCoretechWrapper.cpp
+++ b/Engine/Plugins/Enterprise/DatasmithCADImporter/Source/DatasmithWireTranslator/Private/AliasCoretechWrapper.cpp
@@ -87,7 +87,7 @@ namespace AliasToCoreTechUtils
 		}
 
 		CT_OBJECT_ID CTSurfaceID = 0;
-		CheckedCTError result = CT_SNURBS_IO::Create(CTSurfaceID,
+		CT_IO_ERROR Result = CT_SNURBS_IO::Create(CTSurfaceID,
 			OrderU, OrderV,
 			KnotSizeU, KnotSizeV,
 			ControlPointSizeU, ControlPointSizeV,
@@ -95,11 +95,19 @@ namespace AliasToCoreTechUtils
 			KnotValuesU.data(), KnotValuesV.data(),
 			KnotMultiplicityU.data(), KnotMultiplicityV.data()
 		);
-
+		if (Result != IO_OK)
+		{
+			return 0;
+		}
 		return CTSurfaceID;
 	}
 }
 
+CT_IO_ERROR FAliasCoretechWrapper::Tessellate(FMeshDescription& Mesh, FMeshParameters& MeshParameters)
+{
+	return CADLibrary::Tessellate(MainObjectId, ImportParams, Mesh, MeshParameters);
+}
+
 CT_OBJECT_ID FAliasCoretechWrapper::Add3DCurve(AlCurve& Curve)
 {
 	curveFormType Form = Curve.form();
@@ -133,7 +141,7 @@ CT_OBJECT_ID FAliasCoretechWrapper::Add3DCurve(AlCurve& Curve)
 	Knots[RealKnotSize - 1] = Knots[RealKnotSize - 2];
 
 	CT_OBJECT_ID NurbsID;
-	CheckedCTError setUvCurveError = CT_CNURBS_IO::Create(
+	CT_IO_ERROR setUvCurveError = CT_CNURBS_IO::Create(
 		NurbsID,               /*!< [out] Id of created coedge */
 		Order,                  /*!< [in] Order of curve */
 		RealKnotSize,           /*!< [in] Knot vector size */
@@ -198,8 +206,13 @@ CT_OBJECT_ID FAliasCoretechWrapper::AddTrimCurve(AlTrimCurve& TrimCurve)
 	Knots[RealKnotSize - 1] = Knots[RealKnotSize - 2];
 
 	CT_OBJECT_ID CoedgeID;
-	CheckedCTError setUvCurveError = CT_COEDGE_IO::Create(CoedgeID, bIsReversed ? CT_ORIENTATION::CT_REVERSE : CT_ORIENTATION::CT_FORWARD, D3CurveId);
-	setUvCurveError = CT_COEDGE_IO::SetUVCurve(
+	CT_IO_ERROR Result = CT_COEDGE_IO::Create(CoedgeID, bIsReversed ? CT_ORIENTATION::CT_REVERSE : CT_ORIENTATION::CT_FORWARD, D3CurveId);
+	if (Result != IO_OK)
+	{
+		return 0;
+	}
+
+	Result = CT_COEDGE_IO::SetUVCurve(
 		CoedgeID,               /*!< [out] Id of created coedge */
 		Order,                  /*!< [in] Order of curve */
 		RealKnotSize,           /*!< [in] Knot vector size */
@@ -211,6 +224,10 @@ CT_OBJECT_ID FAliasCoretechWrapper::AddTrimCurve(AlTrimCurve& TrimCurve)
 		Knots[0],               /*!< [in] start parameter of coedge on the uv curve (t range=[knot[0], knot[knot_size-1]]) */
 		Knots[RealKnotSize - 1]   /*!< [in] end parameter of coedge on the uv curve (t range=[knot[0], knot[knot_size-1]]) */
 	);
+	if (Result != IO_OK)
+	{
+		return 0;
+	}
 
 	// Build topo
 	AlTrimCurve *TwinCurve = TrimCurve.getTwinCurve();
@@ -223,7 +240,7 @@ CT_OBJECT_ID FAliasCoretechWrapper::AddTrimCurve(AlTrimCurve& TrimCurve)
 		AlEdge2CTEdge.Add(&TrimCurve, CoedgeID);
 	}
 
-	if (!setUvCurveError)
+	if (Result != IO_OK)
 	{
 		return 0;
 	}
@@ -247,20 +264,22 @@ CT_OBJECT_ID FAliasCoretechWrapper::AddTrimBoundary(AlTrimBoundary& TrimBoundary
 
 	CT_OBJECT_ID LoopId;
 
-	CheckedCTError LoopCreationError = CT_LOOP_IO::Create(LoopId, Coedges);
-	if (!LoopCreationError)
+	CT_IO_ERROR Result = CT_LOOP_IO::Create(LoopId, Coedges);
+	if (Result != IO_OK)
 	{
 		return 0;
 	}
 	return LoopId;
 }
 
-
-
-
 CT_OBJECT_ID FAliasCoretechWrapper::AddTrimRegion(AlTrimRegion& TrimRegion, bool bWorldPosition, bool bOrientation)
 {
 	CT_OBJECT_ID NurbsId = AliasToCoreTechUtils::CreateCTNurbs(TrimRegion, bWorldPosition);
+	if (NurbsId == 0)
+	{
+		return 0;
+	}
+
 	CT_LIST_IO Boundaries;
 	CT_OBJECT_ID FaceID;
 	CT_ORIENTATION FaceOrient = bOrientation ? CT_ORIENTATION::CT_FORWARD : CT_ORIENTATION::CT_REVERSE;
@@ -275,8 +294,8 @@ CT_OBJECT_ID FAliasCoretechWrapper::AddTrimRegion(AlTrimRegion& TrimRegion, bool
 		}
 		TrimBoundary = TrimBoundary->nextBoundary();
 	}
-	CheckedCTError Result = CT_FACE_IO::Create(FaceID, NurbsId, FaceOrient, Boundaries);
-	if (Result)
+	CT_IO_ERROR Result = CT_FACE_IO::Create(FaceID, NurbsId, FaceOrient, Boundaries);
+	if (Result == IO_OK)
 	{
 		return FaceID;
 	}
@@ -304,8 +323,8 @@ void FAliasCoretechWrapper::AddFace(AlSurface& Surface, CT_LIST_IO& FaceList, bo
 	CT_LIST_IO Boundaries;
 	CT_OBJECT_ID FaceID;
 	CT_ORIENTATION FaceOrient = bOrientation ? CT_ORIENTATION::CT_FORWARD : CT_ORIENTATION::CT_REVERSE;
-	CheckedCTError Result = CT_FACE_IO::Create(FaceID, NurbsId, FaceOrient, Boundaries);
-	if (Result)
+	CT_IO_ERROR Result = CT_FACE_IO::Create(FaceID, NurbsId, FaceOrient, Boundaries);
+	if (Result != IO_OK)
 	{
 		FaceList.PushBack(FaceID);
 	}
@@ -325,13 +344,12 @@ void FAliasCoretechWrapper::AddShell(AlShell& Shell, CT_LIST_IO& FaceList, bool
 	}
 }
 
-CheckedCTError FAliasCoretechWrapper::AddBRep(TArray<AlDagNode*>& DagNodeSet, bool bWorldPosition)
+CT_IO_ERROR FAliasCoretechWrapper::AddBRep(TArray<AlDagNode*>& DagNodeSet, bool bWorldPosition)
 {
-	CheckedCTError Result;
+	CT_IO_ERROR Result = IO_OK;
 	if (!IsSessionValid())
 	{
-		Result.RaiseOtherError("bad session init");
-		return Result;
+		return IO_ERROR;
 	}
 
 	if (ImportParams.StitchingTechnique == StitchingSew)
@@ -389,14 +407,14 @@ CheckedCTError FAliasCoretechWrapper::AddBRep(TArray<AlDagNode*>& DagNodeSet, bo
 		}
 		if (FaceList.IsEmpty())
 		{
-			return Result;
+			return IO_ERROR;
 		}
 	}
 
 	// Create body from faces
 	CT_OBJECT_ID BodyID;
 	Result = CT_BODY_IO::CreateFromFaces(BodyID, CT_BODY_PROP::CT_BODY_PROP_EXACT | CT_BODY_PROP::CT_BODY_PROP_CLOSE, FaceList);
-	if (!Result)
+	if (Result != IO_OK)
 	{
 		return Result;
 	}
diff --git a/Engine/Plugins/Enterprise/DatasmithCADImporter/Source/DatasmithWireTranslator/Private/AliasCoretechWrapper.h b/Engine/Plugins/Enterprise/DatasmithCADImporter/Source/DatasmithWireTranslator/Private/AliasCoretechWrapper.h
--- a/Engine/Plugins/Enterprise/DatasmithCADImporter/Source/DatasmithWireTranslator/Private/AliasCoretechWrapper.h
+++ b/Engine/Plugins/Enterprise/DatasmithCADImporter/Source/DatasmithWireTranslator/Private/AliasCoretechWrapper.h
@@ -40,10 +40,12 @@ public:
 	{
 	}
 
-	CheckedCTError AddBRep(TArray<AlDagNode*>& DagNodeSet, bool bIsSymmetricBody);
+	CT_IO_ERROR AddBRep(TArray<AlDagNode*>& DagNodeSet, bool bIsSymmetricBody);
 
 	static TSharedPtr<FAliasCoretechWrapper> GetSharedSession();
 
+	CT_IO_ERROR Tessellate(FMeshDescription& Mesh, FMeshParameters& MeshParameters);
+
 protected:
 	/**
 	* Create a CT coedge (represent the use of an edge by a face).
diff --git a/Engine/Plugins/Enterprise/DatasmithCADImporter/Source/DatasmithWireTranslator/Private/DatasmithWireTranslator.cpp b/Engine/Plugins/Enterprise/DatasmithCADImporter/Source/DatasmithWireTranslator/Private/DatasmithWireTranslator.cpp
--- a/Engine/Plugins/Enterprise/DatasmithCADImporter/Source/DatasmithWireTranslator/Private/DatasmithWireTranslator.cpp
+++ b/Engine/Plugins/Enterprise/DatasmithCADImporter/Source/DatasmithWireTranslator/Private/DatasmithWireTranslator.cpp
@@ -98,14 +98,14 @@ public:
 	}
 
 	bool Read();
-	TOptional<FMeshDescription> GetMeshDescription(TSharedRef<IDatasmithMeshElement> MeshElement, FMeshParameters& MeshParameters);
-	TOptional<FMeshDescription> GetMeshDescription(TSharedRef<IDatasmithMeshElement> MeshElement, FMeshParameters& MeshParameters, TSharedRef<BodyData> BodyTemp);
+	TOptional<FMeshDescription> GetMeshDescription(TSharedRef<IDatasmithMeshElement> MeshElement, CADLibrary::FMeshParameters& MeshParameters);
+	TOptional<FMeshDescription> GetMeshDescription(TSharedRef<IDatasmithMeshElement> MeshElement, CADLibrary::FMeshParameters& MeshParameters, TSharedRef<BodyData> BodyTemp);
 
 	void SetTessellationOptions(const FDatasmithTessellationOptions& Options);
 	void SetOutputPath(const FString& Path) { OutputPath = Path; }
 
 	//double GetMetricUnit() const { return LocalSession->GetImportParameters().MetricUnit; }
-	FImportParameters& GetImportParameters()
+	CADLibrary::FImportParameters& GetImportParameters()
 	{
 		return LocalSession->GetImportParameters();
 	}
@@ -133,19 +133,19 @@ private:
 	void GetDagNodeInfo(TSharedRef<BodyData> CurrentNode, const FDagNodeInfo& ParentInfo, FDagNodeInfo& CurrentNodeInfo);
 	void GetDagNodeMeta(AlDagNode& CurrentNode, TSharedPtr< IDatasmithActorElement > ActorElement);
 
-	TOptional<FMeshDescription> GetMeshOfShellNode(AlDagNode& DagNode, TSharedRef<IDatasmithMeshElement> MeshElement, FMeshParameters& MeshParameters);
-	TOptional<FMeshDescription> GetMeshOfNodeMesh(AlMeshNode& DagNode, TSharedRef<IDatasmithMeshElement> MeshElement, FMeshParameters& MeshParameters, AlMatrix4x4* AlMeshInvGlobalMatrix = nullptr);
-	TOptional<FMeshDescription> GetMeshOfShellBody(TSharedRef<BodyData> DagNode, TSharedRef<IDatasmithMeshElement> MeshElement, FMeshParameters& MeshParameters);
-	TOptional<FMeshDescription> GetMeshOfMeshBody(TSharedRef<BodyData> DagNode, TSharedRef<IDatasmithMeshElement> MeshElement, FMeshParameters& MeshParameters);
+	TOptional<FMeshDescription> GetMeshOfShellNode(AlDagNode& DagNode, TSharedRef<IDatasmithMeshElement> MeshElement, CADLibrary::FMeshParameters& MeshParameters);
+	TOptional<FMeshDescription> GetMeshOfNodeMesh(AlMeshNode& DagNode, TSharedRef<IDatasmithMeshElement> MeshElement, CADLibrary::FMeshParameters& MeshParameters, AlMatrix4x4* AlMeshInvGlobalMatrix = nullptr);
+	TOptional<FMeshDescription> GetMeshOfShellBody(TSharedRef<BodyData> DagNode, TSharedRef<IDatasmithMeshElement> MeshElement, CADLibrary::FMeshParameters& MeshParameters);
+	TOptional<FMeshDescription> GetMeshOfMeshBody(TSharedRef<BodyData> DagNode, TSharedRef<IDatasmithMeshElement> MeshElement, CADLibrary::FMeshParameters& MeshParameters);
 
 	void AddNodeInBodySet(AlDagNode& DagNode, const char* ShaderName, TMap<uint32, TSharedPtr<BodyData>>& ShellToProcess, bool bIsAPatch, uint32 MaxSize);
 
 #ifdef CAD_LIBRARY
-	TOptional<FMeshDescription> MeshDagNodeWithExternalMesher(AlDagNode& DagNode, TSharedRef<IDatasmithMeshElement> MeshElement, FMeshParameters& MeshParameters);
-	TOptional<FMeshDescription> MeshDagNodeWithExternalMesher(TSharedRef<BodyData> DagNode, TSharedRef<IDatasmithMeshElement> MeshElement, FMeshParameters& MeshParameters);
+	TOptional<FMeshDescription> MeshDagNodeWithExternalMesher(AlDagNode& DagNode, TSharedRef<IDatasmithMeshElement> MeshElement, CADLibrary::FMeshParameters& MeshParameters);
+	TOptional<FMeshDescription> MeshDagNodeWithExternalMesher(TSharedRef<BodyData> DagNode, TSharedRef<IDatasmithMeshElement> MeshElement, CADLibrary::FMeshParameters& MeshParameters);
 #endif
 
- 	TOptional< FMeshDescription > ImportMesh(AlMesh& Mesh, FMeshParameters& MeshParameters);
+ 	TOptional< FMeshDescription > ImportMesh(AlMesh& Mesh, CADLibrary::FMeshParameters& MeshParameters);
 
 	void CreateAlCommonMaterial(AlShader *Shader, TSharedRef<IDatasmithUEPbrMaterialElement> MaterialElement);
 	void AddAlBlinnParameters(AlShader *Shader, TSharedRef<IDatasmithUEPbrMaterialElement> MaterialElement);
@@ -1768,7 +1768,7 @@ bool FWireTranslatorImpl::RecurseDagForLeavesNoMerge(AlDagNode* FirstDagNode, co
 
 #ifdef CAD_LIBRARY
 
-TOptional<FMeshDescription> FWireTranslatorImpl::MeshDagNodeWithExternalMesher(AlDagNode& DagNode, TSharedRef<IDatasmithMeshElement> MeshElement, FMeshParameters& MeshParameters)
+TOptional<FMeshDescription> FWireTranslatorImpl::MeshDagNodeWithExternalMesher(AlDagNode& DagNode, TSharedRef<IDatasmithMeshElement> MeshElement, CADLibrary::FMeshParameters& MeshParameters)
 {
 	CADLibrary::CheckedCTError Result;
 
@@ -1799,7 +1799,7 @@ TOptional<FMeshDescription> FWireTranslatorImpl::MeshDagNodeWithExternalMesher(A
 	return MoveTemp(MeshDescription);
 }
 
-TOptional<FMeshDescription> FWireTranslatorImpl::MeshDagNodeWithExternalMesher(TSharedRef<BodyData> Body, TSharedRef<IDatasmithMeshElement> MeshElement, FMeshParameters& MeshParameters)
+TOptional<FMeshDescription> FWireTranslatorImpl::MeshDagNodeWithExternalMesher(TSharedRef<BodyData> Body, TSharedRef<IDatasmithMeshElement> MeshElement, CADLibrary::FMeshParameters& MeshParameters)
 {
 	CADLibrary::CheckedCTError Result;
 
@@ -1830,7 +1830,7 @@ TOptional<FMeshDescription> FWireTranslatorImpl::MeshDagNodeWithExternalMesher(T
 
 #endif
 
-TOptional<FMeshDescription> FWireTranslatorImpl::GetMeshOfShellNode(AlDagNode& DagNode, TSharedRef<IDatasmithMeshElement> MeshElement, FMeshParameters& MeshParameters)
+TOptional<FMeshDescription> FWireTranslatorImpl::GetMeshOfShellNode(AlDagNode& DagNode, TSharedRef<IDatasmithMeshElement> MeshElement, CADLibrary::FMeshParameters& MeshParameters)
 {
 	static bool bUseExternalMesher = true;
 #ifdef CAD_LIBRARY
@@ -1859,13 +1859,13 @@ TOptional<FMeshDescription> FWireTranslatorImpl::GetMeshOfShellNode(AlDagNode& D
 }
 
 
-TOptional<FMeshDescription> FWireTranslatorImpl::GetMeshOfShellBody(TSharedRef<BodyData> Body, TSharedRef<IDatasmithMeshElement> MeshElement, FMeshParameters& MeshParameters)
+TOptional<FMeshDescription> FWireTranslatorImpl::GetMeshOfShellBody(TSharedRef<BodyData> Body, TSharedRef<IDatasmithMeshElement> MeshElement, CADLibrary::FMeshParameters& MeshParameters)
 {
 	TOptional< FMeshDescription > UEMesh = MeshDagNodeWithExternalMesher(Body, MeshElement, MeshParameters);
 	return UEMesh;
 }
 
-TOptional<FMeshDescription> FWireTranslatorImpl::GetMeshOfMeshBody(TSharedRef<BodyData> Body, TSharedRef<IDatasmithMeshElement> MeshElement, FMeshParameters& MeshParameters)
+TOptional<FMeshDescription> FWireTranslatorImpl::GetMeshOfMeshBody(TSharedRef<BodyData> Body, TSharedRef<IDatasmithMeshElement> MeshElement, CADLibrary::FMeshParameters& MeshParameters)
 {
 	FMeshDescription MeshDescription;
 	DatasmithMeshHelper::PrepareAttributeForStaticMesh(MeshDescription);
@@ -1885,7 +1885,7 @@ TOptional<FMeshDescription> FWireTranslatorImpl::GetMeshOfMeshBody(TSharedRef<Bo
 	return MoveTemp(MeshDescription);
 }
 
-TOptional<FMeshDescription> FWireTranslatorImpl::GetMeshOfNodeMesh(AlMeshNode& MeshNode, TSharedRef<IDatasmithMeshElement> MeshElement, FMeshParameters& MeshParameters, AlMatrix4x4* AlMeshInvGlobalMatrix)
+TOptional<FMeshDescription> FWireTranslatorImpl::GetMeshOfNodeMesh(AlMeshNode& MeshNode, TSharedRef<IDatasmithMeshElement> MeshElement, CADLibrary::FMeshParameters& MeshParameters, AlMatrix4x4* AlMeshInvGlobalMatrix)
 {
 	AlMesh * Mesh = MeshNode.mesh();
 	if (AlIsValid(Mesh))
@@ -1900,7 +1900,7 @@ TOptional<FMeshDescription> FWireTranslatorImpl::GetMeshOfNodeMesh(AlMeshNode& M
 }
 
 
-TOptional<FMeshDescription> FWireTranslatorImpl::GetMeshDescription(TSharedRef<IDatasmithMeshElement> MeshElement, FMeshParameters& MeshParameters, TSharedRef<BodyData> Body)
+TOptional<FMeshDescription> FWireTranslatorImpl::GetMeshDescription(TSharedRef<IDatasmithMeshElement> MeshElement, CADLibrary::FMeshParameters& MeshParameters, TSharedRef<BodyData> Body)
 {
 	if(Body->ShellSet.Num() == 0 )
 	{
@@ -1935,7 +1935,7 @@ TOptional<FMeshDescription> FWireTranslatorImpl::GetMeshDescription(TSharedRef<I
 }
 
 
-TOptional<FMeshDescription> FWireTranslatorImpl::GetMeshDescription(TSharedRef<IDatasmithMeshElement> MeshElement, FMeshParameters& MeshParameters)
+TOptional<FMeshDescription> FWireTranslatorImpl::GetMeshDescription(TSharedRef<IDatasmithMeshElement> MeshElement, CADLibrary::FMeshParameters& MeshParameters)
 {
 	AlDagNode ** DagNodeTemp = MeshElementToAlDagNodeMap.Find(&MeshElement.Get());
 	if (DagNodeTemp == nullptr || *DagNodeTemp == nullptr)
@@ -1995,7 +1995,7 @@ TOptional<FMeshDescription> FWireTranslatorImpl::GetMeshDescription(TSharedRef<I
 
 
 // Note that Alias file unit is cm like UE
-TOptional< FMeshDescription > FWireTranslatorImpl::ImportMesh(AlMesh& CurrentMesh, FMeshParameters& MeshParameters)
+TOptional< FMeshDescription > FWireTranslatorImpl::ImportMesh(AlMesh& CurrentMesh, CADLibrary::FMeshParameters& MeshParameters)
 {
 
 	FMeshDescription MeshDescription;
@@ -2070,8 +2070,8 @@ bool FDatasmithWireTranslator::LoadStaticMesh(const TSharedRef<IDatasmithMeshEle
 {
 #ifdef USE_OPENMODEL
 
-	FImportParameters& ImportParameters = Translator->GetImportParameters();
-	FMeshParameters MeshParameters;
+	CADLibrary::FImportParameters& ImportParameters = Translator->GetImportParameters();
+	CADLibrary::FMeshParameters MeshParameters;
 	if (TOptional< FMeshDescription > Mesh = Translator->GetMeshDescription(MeshElement, MeshParameters))
 	{
 		OutMeshPayload.LodMeshes.Add(MoveTemp(Mesh.GetValue()));
diff --git a/Engine/Plugins/Enterprise/DatasmithCADImporter/Source/DatasmithWireTranslator/Private/DatasmithWireTranslator.h b/Engine/Plugins/Enterprise/DatasmithCADImporter/Source/DatasmithWireTranslator/Private/DatasmithWireTranslator.h
--- a/Engine/Plugins/Enterprise/DatasmithCADImporter/Source/DatasmithWireTranslator/Private/DatasmithWireTranslator.h
+++ b/Engine/Plugins/Enterprise/DatasmithCADImporter/Source/DatasmithWireTranslator/Private/DatasmithWireTranslator.h
@@ -4,12 +4,10 @@
 
 #include "CoreMinimal.h"
 
-// Undef USE_OPENMODEL if CAD_LIBRARY is not defined
 #ifndef CAD_LIBRARY
 #undef USE_OPENMODEL
 #endif
 
-#include "CADLibraryOptions.h"
 #include "DatasmithCoreTechTranslator.h"
 #include "UObject/ObjectMacros.h"
 
diff --git a/Engine/Plugins/Enterprise/DatasmithCADImporter/Source/DatasmithWireTranslator/Private/OpenModelUtils.cpp b/Engine/Plugins/Enterprise/DatasmithCADImporter/Source/DatasmithWireTranslator/Private/OpenModelUtils.cpp
--- a/Engine/Plugins/Enterprise/DatasmithCADImporter/Source/DatasmithWireTranslator/Private/OpenModelUtils.cpp
+++ b/Engine/Plugins/Enterprise/DatasmithCADImporter/Source/DatasmithWireTranslator/Private/OpenModelUtils.cpp
@@ -1,13 +1,12 @@
 // Copyright 1998-2019 Epic Games, Inc. All Rights Reserved.
-
 #include "OpenModelUtils.h"
 
-
 #ifdef USE_OPENMODEL
 
+#include "CADOptions.h"
+#include "DatasmithUtils.h"
 #include "Translators/DatasmithTranslator.h"
 
-
 #include "AlDagNode.h"
 #include "AlMesh.h"
 #include "AlLayer.h"
@@ -20,9 +19,6 @@
 #include "Utility/DatasmithMathUtils.h"
 using namespace OpenModelUtils;
 
-
-
-
 const TCHAR * OpenModelUtils::AlObjectTypeToString(AlObjectType type)
 {
 	switch (type)
diff --git a/Engine/Plugins/Enterprise/DatasmithCADImporter/Source/DatasmithWireTranslator/Private/OpenModelUtils.h b/Engine/Plugins/Enterprise/DatasmithCADImporter/Source/DatasmithWireTranslator/Private/OpenModelUtils.h
--- a/Engine/Plugins/Enterprise/DatasmithCADImporter/Source/DatasmithWireTranslator/Private/OpenModelUtils.h
+++ b/Engine/Plugins/Enterprise/DatasmithCADImporter/Source/DatasmithWireTranslator/Private/OpenModelUtils.h
@@ -4,8 +4,7 @@
 
 #include "CoreMinimal.h"
 
-#include "CADLibraryOptions.h"
-
+#include "CADOptions.h"
 
 #ifdef USE_OPENMODEL
 
