Merging 9745597 //UE4/Dev-Enterprise to Release-4.24 (//UE4/Release-4.24)
DatasmithWorker

diff --git a/Engine/Plugins/Enterprise/DatasmithCADImporter/Source/DatasmithDispatcher/Private/DatasmithWorkerHandler.cpp b/Engine/Plugins/Enterprise/DatasmithCADImporter/Source/DatasmithDispatcher/Private/DatasmithWorkerHandler.cpp
--- a/Engine/Plugins/Enterprise/DatasmithCADImporter/Source/DatasmithDispatcher/Private/DatasmithWorkerHandler.cpp
+++ b/Engine/Plugins/Enterprise/DatasmithCADImporter/Source/DatasmithDispatcher/Private/DatasmithWorkerHandler.cpp
@@ -17,11 +17,11 @@ FString FDatasmithWorkerHandler::GetExePath()
 {
 	FString ProcessorPath;
 #if PLATFORM_MAC
-	ProcessorPath = FPaths::Combine(FPaths::EngineDir(), TEXT("Binaries/Mac/NotForLicensees/DatasmithCADWorker"));
+	ProcessorPath = FPaths::Combine(FPaths::EngineDir(), TEXT("Binaries/Mac/DatasmithCADWorker"));
 #elif PLATFORM_LINUX
-	ProcessorPath = FPaths::Combine(FPaths::EngineDir(), TEXT("Binaries/Linux/NotForLicensees/DatasmithCADWorker"));
+	ProcessorPath = FPaths::Combine(FPaths::EngineDir(), TEXT("Binaries/Linux/DatasmithCADWorker"));
 #else
-	ProcessorPath = FPaths::Combine(FPaths::EngineDir(), TEXT("Binaries/Win64/NotForLicensees/DatasmithCADWorker.exe"));
+	ProcessorPath = FPaths::Combine(FPaths::EngineDir(), TEXT("Binaries/Win64/DatasmithCADWorker.exe"));
 #endif
 
 	if (FPaths::FileExists(ProcessorPath))
