Merging 9745597 //UE4/Dev-Enterprise to Release-4.24 (//UE4/Release-4.24)
DatasmithWorker

diff --git a/Engine/Plugins/Enterprise/DatasmithCADImporter/Source/DatasmithCADTranslator/Private/DatasmithCADTranslator.cpp b/Engine/Plugins/Enterprise/DatasmithCADImporter/Source/DatasmithCADTranslator/Private/DatasmithCADTranslator.cpp
--- a/Engine/Plugins/Enterprise/DatasmithCADImporter/Source/DatasmithCADTranslator/Private/DatasmithCADTranslator.cpp
+++ b/Engine/Plugins/Enterprise/DatasmithCADImporter/Source/DatasmithCADTranslator/Private/DatasmithCADTranslator.cpp
@@ -93,7 +93,7 @@ bool FDatasmithCADTranslator::LoadScene(TSharedRef<IDatasmithScene> DatasmithSce
 	Dispatcher.AddTask(FPaths::ConvertRelativePathToFull(GetSource().GetSourceFile()));
 
 	bool bEnableMultiProcess = (NumCores != 1);
-	Dispatcher.Process(bEnableMultiProcess);
+	Dispatcher.Process(false && bEnableMultiProcess);
 	Dispatcher.Clear();
 
 	FDatasmithSceneGraphBuilder SceneGraphBuilder(CADFileToUE4FileMap, MeshElementToCADBRepUuidMap,  CachePath, DatasmithScene, GetSource(), ImportParameters);
diff --git a/Engine/Plugins/Enterprise/DatasmithCADImporter/Source/DatasmithCoreTechExtension/Private/CoreTechParametricSurfaceExtension.cpp b/Engine/Plugins/Enterprise/DatasmithCADImporter/Source/DatasmithCoreTechExtension/Private/CoreTechParametricSurfaceExtension.cpp
--- a/Engine/Plugins/Enterprise/DatasmithCADImporter/Source/DatasmithCoreTechExtension/Private/CoreTechParametricSurfaceExtension.cpp
+++ b/Engine/Plugins/Enterprise/DatasmithCADImporter/Source/DatasmithCoreTechExtension/Private/CoreTechParametricSurfaceExtension.cpp
@@ -181,10 +181,10 @@ bool FCoreTechRetessellate_Impl::ApplyOnOneAsset(UStaticMesh& StaticMesh, UCoreT
 	ImportParameters.ChordTolerance = RetessellateOptions.ChordTolerance;
 	ImportParameters.MaxEdgeLength = RetessellateOptions.MaxEdgeLength;
 	ImportParameters.MaxNormalAngle = RetessellateOptions.NormalTolerance;
+	ImportParameters.ModelCoordSys = CADLibrary::EModelCoordSystem(CoreTechData.SceneParameters.ModelCoordSys);
 	ImportParameters.StitchingTechnique = CADLibrary::EStitchingTechnique(RetessellateOptions.StitchingTechnique);
 
 	CADLibrary::FMeshParameters MeshParameters;
-	MeshParameters.ModelCoordSys = FDatasmithUtils::EModelCoordSystem(CoreTechData.SceneParameters.ModelCoordSys);
 	MeshParameters.bNeedSwapOrientation = CoreTechData.MeshParameters.bNeedSwapOrientation;
 	MeshParameters.bIsSymmetric = CoreTechData.MeshParameters.bIsSymmetric;
 	MeshParameters.SymmetricNormal = CoreTechData.MeshParameters.SymmetricNormal;
