Merging 9745597 //UE4/Dev-Enterprise to Release-4.24 (//UE4/Release-4.24)
DatasmithWorker

diff --git a/Engine/Plugins/Enterprise/DatasmithCADImporter/Source/DatasmithDispatcher/DatasmithDispatcher.Build.cs b/Engine/Plugins/Enterprise/DatasmithCADImporter/Source/DatasmithDispatcher/DatasmithDispatcher.Build.cs
new file mode 100644
--- /dev/null
+++ b/Engine/Plugins/Enterprise/DatasmithCADImporter/Source/DatasmithDispatcher/DatasmithDispatcher.Build.cs
@@ -0,0 +1,33 @@
+// Copyright 1998-2019 Epic Games, Inc. All Rights Reserved.
+
+using System.IO;
+
+namespace UnrealBuildTool.Rules
+{
+	public class DatasmithDispatcher : ModuleRules
+	{
+		public DatasmithDispatcher(ReadOnlyTargetRules Target) : base(Target)
+		{
+            PublicDependencyModuleNames.AddRange(
+				new string[]
+				{
+                    "Core",
+					"Sockets",
+					"CADInterfaces",
+					"CADTools",
+				}
+			);
+
+			PrivateDependencyModuleNames.AddRange(
+				new string[]
+				{
+                }
+			);
+
+            if (System.Type.GetType("CoreTech") != null)
+            {
+                PublicDependencyModuleNames.Add("CoreTech");
+            }
+        }
+    }
+}
