Fix [Incremental UE4Editor Win64]

diff --git a/Engine/Plugins/Enterprise/DatasmithCADImporter/Source/CADInterfaces/CADInterfaces.Build.cs b/Engine/Plugins/Enterprise/DatasmithCADImporter/Source/CADInterfaces/CADInterfaces.Build.cs
--- a/Engine/Plugins/Enterprise/DatasmithCADImporter/Source/CADInterfaces/CADInterfaces.Build.cs
+++ b/Engine/Plugins/Enterprise/DatasmithCADImporter/Source/CADInterfaces/CADInterfaces.Build.cs
@@ -8,6 +8,8 @@ namespace UnrealBuildTool.Rules
 	{
 		public CADInterfaces(ReadOnlyTargetRules Target) : base(Target)
 		{
+			bLegalToDistributeObjectCode = true;
+
 			PublicDependencyModuleNames.AddRange(
 				new string[]
 				{
