[NO NON-PII COMMIT MESSAGE]

diff --git a/Engine/Plugins/Enterprise/DatasmithCADImporter/Source/DatasmithWireTranslator/Private/AliasCoretechWrapper.cpp b/Engine/Plugins/Enterprise/DatasmithCADImporter/Source/DatasmithWireTranslator/Private/AliasCoretechWrapper.cpp
--- a/Engine/Plugins/Enterprise/DatasmithCADImporter/Source/DatasmithWireTranslator/Private/AliasCoretechWrapper.cpp
+++ b/Engine/Plugins/Enterprise/DatasmithCADImporter/Source/DatasmithWireTranslator/Private/AliasCoretechWrapper.cpp
@@ -240,11 +240,6 @@ CT_OBJECT_ID FAliasCoretechWrapper::AddTrimCurve(AlTrimCurve& TrimCurve)
 		AlEdge2CTEdge.Add(&TrimCurve, CoedgeID);
 	}
 
-	if (Result != IO_OK)
-	{
-		return 0;
-	}
-
 	return CoedgeID;
 }
 
