Fix github build

diff --git a/Engine/Plugins/Enterprise/DatasmithCADImporter/Source/CADInterfaces/Private/CTSession.cpp b/Engine/Plugins/Enterprise/DatasmithCADImporter/Source/CADInterfaces/Private/CTSession.cpp
new file mode 100644
--- /dev/null
+++ b/Engine/Plugins/Enterprise/DatasmithCADImporter/Source/CADInterfaces/Private/CTSession.cpp
@@ -0,0 +1,61 @@
+// Copyright 1998-2019 Epic Games, Inc. All Rights Reserved.
+
+#ifdef CAD_INTERFACE // #ueent_wip if WITH_CORETECH
+
+#include "CTSession.h"
+#include "CADData.h"
+
+
+#include <vector>
+
+// If attached with a debugger, errors from CT assigned to a Checked_IO_ERROR will break.
+#define BREAK_ON_CT_USAGE_ERROR 0
+
+namespace CADLibrary
+{
+	TWeakPtr<CTSession> CTSession::SharedSession;
+
+	void CTSession::ClearData()
+	{
+		CheckedCTError Result = CTKIO_UnloadModel();
+
+		// recreate the Main Object
+		CT_OBJECT_ID NullParent = 0;
+		Result = CT_COMPONENT_IO::Create(MainObjectId, NullParent);
+	}
+
+	CheckedCTError CTSession::SaveBrep(const FString& FilePath)
+	{
+		CT_LIST_IO ObjectList;
+		ObjectList.PushBack(MainObjectId);
+		return CTKIO_SaveFile(ObjectList, *FilePath, L"Ct");
+	}
+
+	CheckedCTError CTSession::TopoFixes()
+	{
+		return CADLibrary::Repair(MainObjectId, ImportParams.StitchingTechnique);
+	}
+
+	void CheckedCTError::Validate()
+	{
+		static bool breakOnError = BREAK_ON_CT_USAGE_ERROR;
+		ensure(!breakOnError || bool(*this));
+	}
+
+	void CTSession::SetImportParameters(float ChordTolerance, float MaxEdgeLength, float NormalTolerance, CADLibrary::EStitchingTechnique StitchingTechnique)
+	{
+		ImportParams.ChordTolerance = ChordTolerance / ImportParams.ScaleFactor;
+		ImportParams.MaxEdgeLength = MaxEdgeLength / ImportParams.ScaleFactor;
+		ImportParams.MaxNormalAngle = NormalTolerance;
+		ImportParams.StitchingTechnique = StitchingTechnique;
+
+		CT_TESS_DATA_TYPE VertexType = CT_TESS_DOUBLE;
+		CT_TESS_DATA_TYPE NormalType = CT_TESS_FLOAT;
+		CT_TESS_DATA_TYPE UVType = CT_TESS_DOUBLE;
+		CTKIO_ChangeTesselationParameters(ImportParams.ChordTolerance, ImportParams.MaxEdgeLength, ImportParams.MaxNormalAngle, CT_FALSE, VertexType, NormalType, UVType);
+	}
+
+} // namespace CADLibrary
+
+#endif // CAD_INTERFACE
+
diff --git a/Engine/Plugins/Enterprise/DatasmithCADImporter/Source/CADInterfaces/Private/CoreTechFileParser.cpp b/Engine/Plugins/Enterprise/DatasmithCADImporter/Source/CADInterfaces/Private/CoreTechFileParser.cpp
new file mode 100644
--- /dev/null
+++ b/Engine/Plugins/Enterprise/DatasmithCADImporter/Source/CADInterfaces/Private/CoreTechFileParser.cpp
@@ -0,0 +1,1728 @@
+// Copyright 1998-2019 Epic Games, Inc. All Rights Reserved.
+#include "CoreTechFileParser.h"
+
+#ifdef CAD_INTERFACE
+
+
+#include "CADData.h"
+#include "CADOptions.h"
+#include "Containers/Array.h"
+#include "Containers/Map.h"
+#include "Containers/Queue.h"
+#include "CoreTechTypes.h"
+#include "GenericPlatform/GenericPlatformFile.h"
+#include "HAL/FileManager.h"
+#include "Internationalization/Text.h"
+#include "Misc/FileHelper.h"
+#include "Misc/Paths.h"
+
+#define SGSIZE 100000
+#define EXTREFNUM 5000
+
+namespace CADLibrary 
+{
+
+namespace {
+	double Distance(CT_COORDINATE Point1, CT_COORDINATE Point2)
+	{
+		return sqrt((Point2.xyz[0] - Point1.xyz[0]) * (Point2.xyz[0] - Point1.xyz[0]) + (Point2.xyz[1] - Point1.xyz[1]) * (Point2.xyz[1] - Point1.xyz[1]) + (Point2.xyz[2] - Point1.xyz[2]) * (Point2.xyz[2] - Point1.xyz[2]));
+	}
+
+	template<typename UVType>
+	void ScaleUV(CT_OBJECT_ID FaceID, void* TexCoordArray, int32 VertexCount, UVType Scale)
+	{
+		UVType VMin, VMax, UMin, UMax;
+		VMin = UMin = HUGE_VALF;
+		VMax = UMax = -HUGE_VALF;
+		UVType* UVSet = (UVType*) TexCoordArray;
+		for (int32 Index = 0, UVCoord = 0; Index < VertexCount; ++Index, UVCoord += 2)
+		{
+			UMin = FMath::Min(UVSet[UVCoord + 0], UMin);
+			UMax = FMath::Max(UVSet[UVCoord + 0], UMax);
+			VMin = FMath::Min(UVSet[UVCoord + 1], VMin);
+			VMax = FMath::Max(UVSet[UVCoord + 1], VMax);
+		}
+
+		double PuMin, PuMax, PvMin, PvMax;
+		PuMin = PvMin = HUGE_VALF;
+		PuMax = PvMax = -HUGE_VALF;
+
+		// fast UV min max 
+		CT_FACE_IO::AskUVminmax(FaceID, PuMin, PuMax, PvMin, PvMax);
+
+		const uint32 NbIsoCurves = 7;
+
+		// Compute Point grid on the restricted surface defined by [PuMin, PuMax], [PvMin, PvMax]
+		CT_OBJECT_ID SurfaceID;
+		CT_ORIENTATION Orientation;
+		CT_FACE_IO::AskSurface(FaceID, SurfaceID, Orientation);
+
+		CT_OBJECT_TYPE SurfaceType;
+		CT_SURFACE_IO::AskType(SurfaceID, SurfaceType);
+
+		UVType DeltaU = (PuMax - PuMin) / (NbIsoCurves - 1);
+		UVType DeltaV = (PvMax - PvMin) / (NbIsoCurves - 1);
+		UVType U = PuMin, V = PvMin;
+
+		CT_COORDINATE NodeMatrix[121];
+
+		for (int32 IndexI = 0; IndexI < NbIsoCurves; IndexI++)
+		{
+			for (int32 IndexJ = 0; IndexJ < NbIsoCurves; IndexJ++)
+			{
+				CT_SURFACE_IO::Evaluate(SurfaceID, U, V, NodeMatrix[IndexI*NbIsoCurves + IndexJ]);
+				V += DeltaV;
+			}
+			U += DeltaU;
+			V = PvMin;
+		}
+
+		// Compute length of 7 iso V line
+		UVType LengthU[NbIsoCurves];
+		UVType LengthUMin = HUGE_VAL;
+		UVType LengthUMax = 0;
+		UVType LengthUMed = 0;
+
+		for (int32 IndexJ = 0; IndexJ < NbIsoCurves; IndexJ++)
+		{
+			LengthU[IndexJ] = 0;
+			for (int32 IndexI = 0; IndexI < (NbIsoCurves - 1); IndexI++)
+			{
+				LengthU[IndexJ] += Distance(NodeMatrix[IndexI * NbIsoCurves + IndexJ], NodeMatrix[(IndexI + 1) * NbIsoCurves + IndexJ]);
+			}
+			LengthUMed += LengthU[IndexJ];
+			LengthUMin = FMath::Min(LengthU[IndexJ], LengthUMin);
+			LengthUMax = FMath::Max(LengthU[IndexJ], LengthUMax);
+		}
+		LengthUMed /= NbIsoCurves;
+		LengthUMed = LengthUMed * 2 / 3 + LengthUMax / 3;
+
+		// Compute length of 7 iso U line
+		UVType LengthV[NbIsoCurves];
+		UVType LengthVMin = HUGE_VAL;
+		UVType LengthVMax = 0;
+		UVType LengthVMed = 0;
+
+		for (int32 IndexI = 0; IndexI < NbIsoCurves; IndexI++)
+		{
+			LengthV[IndexI] = 0;
+			for (int32 IndexJ = 0; IndexJ < (NbIsoCurves - 1); IndexJ++)
+			{
+				LengthV[IndexI] += Distance(NodeMatrix[IndexI * NbIsoCurves + IndexJ], NodeMatrix[IndexI * NbIsoCurves + IndexJ + 1]);
+			}
+			LengthVMed += LengthV[IndexI];
+			LengthVMin = FMath::Min(LengthV[IndexI], LengthVMin);
+			LengthVMax = FMath::Max(LengthV[IndexI], LengthVMax);
+		}
+		LengthVMed /= NbIsoCurves;
+		LengthVMed = LengthVMed * 2 / 3 + LengthVMax / 3;
+
+		switch (SurfaceType)
+		{
+		case CT_CONE_TYPE:
+		case CT_CYLINDER_TYPE:
+		case CT_SPHERE_TYPE:
+			Swap(LengthUMed, LengthVMed);
+			break;
+		case CT_S_REVOL_TYPE:
+		case CT_TORUS_TYPE:
+			// Need swap ?
+			// Swap(LengthUMed, LengthVMed);
+			break;
+		case CT_S_NURBS_TYPE:
+		case CT_PLANE_TYPE:
+		case CT_S_OFFSET_TYPE:
+		case CT_S_RULED_TYPE:
+		case CT_TABULATED_RULED_TYPE:
+		case CT_S_LINEARTRANSFO_TYPE:
+		case CT_S_NONLINEARTRANSFO_TYPE:
+		case CT_S_BLEND_TYPE:
+		default:
+			break;
+		}
+
+		// scale the UV map
+		// 0.1 define UV in cm and not in mm
+		UVType VScale = Scale * LengthVMed * 1 / (VMax - VMin) / 100;
+		UVType UScale = Scale * LengthUMed * 1 / (UMax - UMin) / 100;
+
+		for (int32 Index = 0, UVCoord = 0; Index < VertexCount; ++Index, UVCoord += 2)
+		{
+			UVSet[UVCoord + 0] *= UScale;
+			UVSet[UVCoord + 1] *= VScale;
+		}
+	}
+}
+
+uint32 GetFileHash(const FString& FileName, const FFileStatData& FileStatData, const FString& Config)
+{
+	int64 FileSize = FileStatData.FileSize;
+	FDateTime ModificationTime = FileStatData.ModificationTime;
+
+	uint32 FileHash = GetTypeHash(FileName);
+	FileHash = HashCombine(FileHash, GetTypeHash(FileSize));
+	FileHash = HashCombine(FileHash, GetTypeHash(ModificationTime));
+	if (!Config.IsEmpty())
+	{
+		FileHash = HashCombine(FileHash, GetTypeHash(Config));
+	}
+
+	return FileHash;
+}
+
+uint32 GetGeomFileHash(const uint32 InSGHash, const FImportParameters& ImportParam)
+{
+	uint32 FileHash = InSGHash;
+	FileHash = HashCombine(FileHash, GetTypeHash(ImportParam.ChordTolerance));
+	FileHash = HashCombine(FileHash, GetTypeHash(ImportParam.MaxEdgeLength));
+	FileHash = HashCombine(FileHash, GetTypeHash(ImportParam.MaxNormalAngle));
+	FileHash = HashCombine(FileHash, GetTypeHash(ImportParam.MetricUnit));
+	FileHash = HashCombine(FileHash, GetTypeHash(ImportParam.ScaleFactor));
+	FileHash = HashCombine(FileHash, GetTypeHash(ImportParam.StitchingTechnique));
+	return FileHash;
+}
+
+
+
+uint32 GetFaceTessellation(CT_OBJECT_ID FaceID, TArray<FTessellationData>& FaceTessellationSet, int32& OutRawDataSize, const float ScaleFactor)
+{
+	CT_IO_ERROR Error = IO_OK;
+
+	CT_UINT32         VertexCount;
+	CT_UINT32         NormalCount;
+	CT_UINT32         IndexCount;
+	CT_TESS_DATA_TYPE VertexType;
+	CT_TESS_DATA_TYPE TexCoordType;
+	CT_TESS_DATA_TYPE NormalType;
+	CT_LOGICAL        HasRGBColor;
+	CT_UINT16         UserSize;
+	CT_TESS_DATA_TYPE IndexType;
+	void*             VertexArray;
+	void*             TexCoordArray;
+	void*             NormalArray;
+	void*             ColorArray;
+	void*             UserArray;
+	void*             IndexArray;
+
+	Error = CT_FACE_IO::AskTesselation(FaceID, VertexCount, NormalCount, IndexCount,
+		VertexType, TexCoordType, NormalType, HasRGBColor, UserSize, IndexType,
+		VertexArray, TexCoordArray, NormalArray, ColorArray, UserArray, IndexArray);
+
+	// Something wrong happened, either an error or no data to collect
+	if (Error != IO_OK || VertexArray == nullptr || IndexArray == nullptr || IndexCount == 0)
+	{
+		VertexCount = NormalCount = IndexCount = 0;
+		OutRawDataSize = 0;
+		return 0;
+	}
+
+	int32 Index = FaceTessellationSet.Emplace();
+	FTessellationData& Tessellation = FaceTessellationSet[Index];
+	if (TexCoordArray != nullptr)
+	{
+		switch (TexCoordType)
+		{
+		case CT_TESS_FLOAT:
+			ScaleUV<float>(FaceID, TexCoordArray, VertexCount, ScaleFactor);
+			break;
+		case CT_TESS_DOUBLE:
+			ScaleUV<double>(FaceID, TexCoordArray, VertexCount, (double)ScaleFactor);
+			break;
+		}
+	}
+
+	Tessellation.VertexCount = VertexCount;
+	Tessellation.NormalCount = NormalCount;
+	Tessellation.IndexCount = IndexCount;
+	Tessellation.SizeOfVertexType = GetSize(VertexType);
+	Tessellation.SizeOfTexCoordType = GetSize(TexCoordType);
+	Tessellation.SizeOfNormalType = GetSize(NormalType);
+	Tessellation.SizeOfIndexType = GetSize(IndexType);
+
+	Tessellation.VertexArray.Empty(Tessellation.VertexCount * Tessellation.SizeOfVertexType);
+	Tessellation.NormalArray.Empty(Tessellation.NormalCount * Tessellation.SizeOfNormalType);
+	Tessellation.IndexArray.Empty(Tessellation.IndexCount * Tessellation.SizeOfIndexType);
+	Tessellation.TexCoordArray.Empty(Tessellation.VertexCount * Tessellation.SizeOfTexCoordType);
+
+	Tessellation.VertexArray.Append((uint8*) VertexArray, 3 * Tessellation.VertexCount * Tessellation.SizeOfVertexType);
+	Tessellation.NormalArray.Append((uint8*)NormalArray, 3 * Tessellation.NormalCount * Tessellation.SizeOfNormalType);
+	Tessellation.IndexArray.Append((uint8*)IndexArray, Tessellation.IndexCount * Tessellation.SizeOfIndexType);
+	Tessellation.TexCoordArray.Append((uint8*)TexCoordArray, 2 * Tessellation.VertexCount * Tessellation.SizeOfTexCoordType);
+
+	OutRawDataSize = Tessellation.VertexArray.Num() + Tessellation.NormalArray.Num() + Tessellation.IndexArray.Num() + Tessellation.TexCoordArray.Num();
+
+	return (uint32)Tessellation.IndexCount / 3;
+}
+
+
+void GetCTObjectDisplayDataIds(CT_OBJECT_ID ObjectID, FObjectDisplayDataId& Material)
+{
+	if (CT_OBJECT_IO::SearchAttribute(ObjectID, CT_ATTRIB_MATERIALID) == IO_OK)
+	{
+		CT_UINT32 MaterialId = 0;
+		if (CT_CURRENT_ATTRIB_IO::AskIntField(ITH_MATERIALID_VALUE, MaterialId) == IO_OK && MaterialId > 0)
+		{
+			Material.MaterialId = (uint32)MaterialId;
+		}
+	}
+
+	if (CT_OBJECT_IO::SearchAttribute(ObjectID, CT_ATTRIB_COLORID) == IO_OK)
+	{
+		CT_UINT32 ColorId = 0;
+		if (CT_CURRENT_ATTRIB_IO::AskIntField(ITH_COLORID_VALUE, ColorId) == IO_OK && ColorId > 0)
+		{
+			uint8 alpha = 255;
+			if (CT_OBJECT_IO::SearchAttribute(ObjectID, CT_ATTRIB_TRANSPARENCY) == IO_OK)
+			{
+				CT_DOUBLE dbl_value = 0.;
+				if (CT_CURRENT_ATTRIB_IO::AskDblField(0, dbl_value) == IO_OK && dbl_value >= 0.0 && dbl_value <= 1.0)
+				{
+					alpha = (uint8)int((1. - dbl_value) * 255.);
+				}
+			}
+			Material.ColorHId = BuildFastColorHash(ColorId, alpha);
+		}
+	}
+}
+
+int32 BuildColorHash(CT_COLOR CtColor, uint8 Alpha)
+{
+	FColor Color(CtColor[0], CtColor[1], CtColor[2], Alpha);
+	return BuildColorHash(Color);
+}
+
+int32 BuildColorHash(uint32 ColorHId)
+{
+	uint32 ColorId;
+	uint8 Alpha;
+	UnhashFastColorHash(ColorHId, ColorId, Alpha);
+
+	// Ref. BaseHelper.cpp getColorData
+	CT_COLOR CtColor = { 200, 200, 200 };
+	if (ColorId > 0)
+	{
+		CT_MATERIAL_IO::AskIndexedColor((CT_OBJECT_ID) ColorId, CtColor);
+	}
+	return BuildColorHash(CtColor, Alpha);
+}
+
+int32 BuildMaterialHash(uint32 MaterialId)
+{
+	FCADMaterial OutMaterial;
+	bool bReturn = GetMaterial(MaterialId, OutMaterial);
+	if (!bReturn)
+	{
+		return 0;
+	}
+	return BuildMaterialHash(OutMaterial);
+}
+
+void GetColor(uint32 ColorHash, FColor& OutColor)
+{
+	uint32 ColorId;
+	uint8 Alpha;
+	UnhashFastColorHash(ColorHash, ColorId, Alpha);
+
+	CT_COLOR CtColor = { 200, 200, 200 };
+	if (ColorId > 0)
+	{
+		CT_MATERIAL_IO::AskIndexedColor((CT_OBJECT_ID)ColorId, CtColor);
+	}
+
+	OutColor.R = CtColor[0];
+	OutColor.G = CtColor[1];
+	OutColor.B = CtColor[2];
+	OutColor.A = Alpha;
+}
+
+
+bool GetMaterial(uint32 MaterialId, FCADMaterial& OutMaterial)
+{
+	//	// Ref. BaseHelper.cpp
+	CT_STR CtName;
+	CT_COLOR CtDiffuse = { 200, 200, 200 }, CtAmbient = { 200, 200, 200 }, CtSpecular = { 200, 200, 200 };
+	CT_FLOAT CtShininess = 0.f, CtTransparency = 0.f, CtReflexion = 0.f;
+	CT_TEXTURE_ID CtTextureId = 0;
+	if (MaterialId)
+	{
+		CT_IO_ERROR bReturn = CT_MATERIAL_IO::AskParameters(MaterialId, CtName, CtDiffuse, CtAmbient, CtSpecular, CtShininess, CtTransparency, CtReflexion, CtTextureId);
+		if (bReturn != IO_OK)
+		{
+			return false;
+		}
+	}
+
+	CT_STR CtTextureName = "";
+	if (CtTextureId)
+	{
+		CT_INT32  Width = 0, Height = 0;
+		if (!(CT_TEXTURE_IO::AskParameters(CtTextureId, CtTextureName, Width, Height) == IO_OK && Width && Height))
+		{
+			CtTextureName = "";
+		}
+	}
+
+	OutMaterial.MaterialId = MaterialId;
+	OutMaterial.MaterialName = CtName.toUnicode();
+	OutMaterial.Diffuse = FColor(CtDiffuse[0], CtDiffuse[1], CtDiffuse[2], 255);
+	OutMaterial.Ambient = FColor(CtAmbient[0], CtAmbient[1], CtAmbient[2], 255);
+	OutMaterial.Specular = FColor(CtSpecular[0], CtSpecular[1], CtSpecular[2], 255);
+	OutMaterial.Shininess = CtShininess;
+	OutMaterial.Transparency = CtTransparency;
+	OutMaterial.Reflexion = CtReflexion;
+	OutMaterial.TextureName = CtTextureName.toUnicode();
+	return true;
+}
+
+
+
+void GetBodiesMaterials(TArray<CT_OBJECT_ID>& BodySet, TMap<uint32, uint32>& MaterialIdToHash, bool bPreferPartData)
+{
+	CT_MATERIAL_ID BodyCtMaterialId = 0;
+	for (int Index = 0; Index < BodySet.Num(); Index++)
+	{
+		FObjectDisplayDataId BodyMaterialId;
+		GetCTObjectDisplayDataIds(BodySet[Index], BodyMaterialId);
+
+		// Loop through the face of the first body and collect material data
+		CT_LIST_IO FaceList;
+		CT_BODY_IO::AskFaces(BodySet[Index], FaceList);
+		FaceList.IteratorInitialize();
+
+		CT_OBJECT_ID FaceID;
+		while ((FaceID = FaceList.IteratorIter()) != 0)
+		{
+			FTessellationData Tessellation;
+			FObjectDisplayDataId FaceMaterialId;
+			GetCTObjectDisplayDataIds(FaceID, FaceMaterialId);
+			SetFaceMainMaterial(FaceMaterialId, BodyMaterialId, MaterialIdToHash, Tessellation);
+		}
+	}
+}
+
+
+void SetFaceMainMaterial(FObjectDisplayDataId& InFaceMaterial, FObjectDisplayDataId& InBodyMaterial, TMap<uint32, uint32>& MaterialIdToMaterialHashMap, FTessellationData& OutFaceTessellations)
+{
+	uint32 FaceMaterialId = 0; // either CT Material ID or CT Color/Alpha ID (ColorId & Alpha << 24)
+	uint32 BodyMaterialId = 0; // either CT Material ID or CT Color/Alpha ID (ColorId & Alpha << 24)
+
+	uint32 FaceMaterialHash = 0;
+	uint32 BodyMaterialHash = 0;
+
+	if (InFaceMaterial.MaterialId > 0)
+	{
+		FaceMaterialId = InFaceMaterial.MaterialId;
+		if (uint32* MHash = MaterialIdToMaterialHashMap.Find(InFaceMaterial.MaterialId))
+		{
+			FaceMaterialHash = *MHash;
+		}
+		else
+		{
+			FaceMaterialHash = BuildMaterialHash(InFaceMaterial.MaterialId);
+			MaterialIdToMaterialHashMap.Add(InFaceMaterial.MaterialId, FaceMaterialHash);
+		}
+	}
+	else if (InFaceMaterial.ColorHId > 0)
+	{
+		FaceMaterialId = InFaceMaterial.ColorHId;
+		if (uint32* MHash = MaterialIdToMaterialHashMap.Find(InFaceMaterial.ColorHId))
+		{
+			FaceMaterialHash = *MHash;
+		}
+		else
+		{
+			FaceMaterialHash = BuildColorHash(InFaceMaterial.ColorHId);
+			MaterialIdToMaterialHashMap.Add(InFaceMaterial.ColorHId, FaceMaterialHash);
+		}
+	}
+
+	if (InBodyMaterial.MaterialId > 0)
+	{
+		BodyMaterialId = InBodyMaterial.MaterialId;
+		if (uint32* MHash = MaterialIdToMaterialHashMap.Find(InBodyMaterial.MaterialId))
+		{
+			BodyMaterialHash = *MHash;
+		}
+		else
+		{
+			BodyMaterialHash = BuildMaterialHash(InBodyMaterial.MaterialId);
+			MaterialIdToMaterialHashMap.Add(InBodyMaterial.MaterialId, BodyMaterialHash);
+		}
+	}
+	else if (InBodyMaterial.ColorHId > 0)
+	{
+		BodyMaterialId = InBodyMaterial.ColorHId;
+		if (uint32* MHash = MaterialIdToMaterialHashMap.Find(InBodyMaterial.ColorHId))
+		{
+			BodyMaterialHash = *MHash;
+		}
+		else
+		{
+			BodyMaterialHash = BuildColorHash(InBodyMaterial.ColorHId);
+			MaterialIdToMaterialHashMap.Add(InBodyMaterial.ColorHId, BodyMaterialHash);
+		}
+	}
+
+	uint32 MaterialId = 0;
+	uint32 MaterialHash = 0;
+
+	// set output
+	if (true)
+	{
+		if (BodyMaterialId)
+		{
+			MaterialId = BodyMaterialId;
+			MaterialHash = BodyMaterialHash;
+		} 
+		else
+		{
+			MaterialId = FaceMaterialId;
+			MaterialHash = FaceMaterialHash;
+		}
+	}
+	else
+	{
+		if (FaceMaterialId)
+		{
+			MaterialId = FaceMaterialId;
+			MaterialHash = FaceMaterialHash;
+		}
+		else
+		{
+			MaterialId = BodyMaterialId;
+			MaterialHash = BodyMaterialHash;
+		}
+	}
+
+	OutFaceTessellations.MaterialId = MaterialId;
+	OutFaceTessellations.MaterialHash = MaterialHash;
+}
+
+int32 BuildColorHash(uint32 ColorId, uint8 Alpha)
+{
+	// Ref. BaseHelper.cpp getColorData
+	CT_COLOR CtColor = { 200, 200, 200 };
+	if (ColorId > 0)
+	{
+		CT_MATERIAL_IO::AskIndexedColor(ColorId, CtColor);
+	}
+	return BuildColorHash(CtColor, Alpha);
+}
+
+uint32 GetStaticMeshUuid(const TCHAR* OutSgFile, const int32 BodyId)
+{
+	uint32 BodyUUID = GetTypeHash(OutSgFile);
+	FString BodyStr = TEXT("B ") + FString::FromInt(BodyId);
+	BodyUUID = HashCombine(BodyUUID, GetTypeHash(*BodyStr));
+	return BodyUUID;
+}
+
+void FCoreTechFileParser::ExportFileSceneGraph()
+{
+	FFileHelper::SaveStringArrayToFile(SceneGraphDescription, *FPaths::Combine(CachePath, TEXT("scene"), SceneGraphFile + TEXT(".sg"))); 
+	SceneGraphDescription.Empty();
+}
+
+uint32 FCoreTechFileParser::GetMaterialNum()
+{
+	CT_UINT32 IColor = 1;
+	while (true)
+	{
+		CT_COLOR CtColor;
+		if (CT_MATERIAL_IO::AskIndexedColor(IColor, CtColor) != IO_OK)
+		{
+			break;
+		}
+		IColor++;
+	}
+
+	CT_UINT32 IMaterial = 1;
+	while (true)
+	{
+		CT_COLOR Diffuse, Ambient, Specular;
+		CT_FLOAT Shininess, Transparency, Reflexion;
+		CT_STR Name("");
+		CT_TEXTURE_ID TextureId;
+
+		if (CT_MATERIAL_IO::AskParameters(IMaterial, Name, Diffuse, Ambient, Specular, Shininess, Transparency, Reflexion, TextureId) != IO_OK)
+		{
+			break;
+		}
+		IMaterial++;
+	}
+
+	return IColor + IMaterial - 2;
+}
+
+void FCoreTechFileParser::ReadColor()
+{
+	CT_COLOR CtColor;
+	CT_UINT32 IColor = 1;
+
+	SceneGraphDescription[COLORSETLINE] = FString::Printf(TEXT("Color %d"), SceneGraphDescription.Num());
+	while (true)
+	{
+		if (CT_MATERIAL_IO::AskIndexedColor(IColor, CtColor) != IO_OK)
+		{
+			break;
+		}
+		SceneGraphDescription.Add(FString::Printf(TEXT("%d %d %d %d"), IColor, CtColor[0], CtColor[1], CtColor[2]));
+		IColor++;
+	}
+	SceneGraphDescription.Add(TEXT(""));
+}
+
+void FCoreTechFileParser::ReadMaterial()
+{
+	CT_UINT32 MaterialId = 1;
+
+	SceneGraphDescription[MATERIALSETLINE] = FString::Printf(TEXT("Material %d"), SceneGraphDescription.Num());
+	while (true)
+	{
+		FCADMaterial Material;
+		bool bReturn = GetMaterial(MaterialId, Material);
+		if (!bReturn)
+		{
+			break;
+		}
+
+		int32 MaterialHash = BuildMaterialHash(Material);
+
+		SceneGraphDescription.Add(FString::Printf(TEXT("%d %s"), MaterialId, *Material.MaterialName));        // COLORSETLINE 3
+		SceneGraphDescription.Add(FString::Printf(TEXT("%d %d %d %d %d %d %d %d %d %d %d %f %f %f %s"), 
+			MaterialId, MaterialHash, 
+			Material.Diffuse.R, Material.Diffuse.G, Material.Diffuse.B, 
+			Material.Ambient.R, Material.Ambient.G, Material.Ambient.B, 
+			Material.Specular.R, Material.Specular.G, Material.Specular.B, 
+			Material.Shininess, Material.Transparency, Material.Reflexion, *Material.TextureName));
+		MaterialIdToMaterialHashMap.Add(MaterialId, MaterialHash);
+		MaterialId++;
+	}
+	SceneGraphDescription.Add(TEXT(""));
+}
+
+FCoreTechFileParser::FCoreTechFileParser(const FString& InCADFullPath, const FString& InCachePath, const FImportParameters& ImportParams, bool bInPreferBodyData, bool bScaleUVMap)
+	: CachePath(InCachePath)
+	, FullPath(InCADFullPath)
+	, bNeedSaveCTFile(false)
+	, bPreferBodyData(bInPreferBodyData)
+	, bScaleUVMap(bScaleUVMap)
+	, ImportParameters(ImportParams)
+{
+	CTIdToRawLineMap.Reserve(EXTREFNUM);
+	CTKIO_InitializeKernel(ImportParameters.MetricUnit);
+}
+
+void GetRawDataFileExternalRef(const FString& InRawDataFile, TSet<FString>& ExternalReferences)
+{
+	TArray<FString> SGDescription;
+	FFileHelper::LoadFileToStringArray(SGDescription, *InRawDataFile);
+
+	if (SGDescription.Num() < 10)
+		return;
+
+	TArray<FString> HeaderSplit;
+	FString& HeaderExternalRef = SGDescription[EXTERNALREFLINE];
+	HeaderExternalRef.ParseIntoArray(HeaderSplit, TEXT(" "));
+
+	if (HeaderSplit.Num() != 3)
+	{
+		return;
+	}
+
+	int32 StartRef = FCString::Atoi(*HeaderSplit[1]);
+	int32 EndRef = StartRef + FCString::Atoi(*HeaderSplit[2]);
+	for (int32 Ref = StartRef; Ref < EndRef; ++Ref)
+	{
+		ExternalReferences.Add(SGDescription[Ref]);
+	}
+}
+
+EProcessState FCoreTechFileParser::ProcessFile()
+{
+	FileConfiguration.Empty();
+
+	CADFile = FPaths::GetCleanFilename(*FullPath);
+
+	// Check if configuration is passed with file name
+	FString NewCADFile;
+	const FString CutOffMarker = TEXT("|");
+	if (FullPath.Split(CutOffMarker, &NewCADFile, &FileConfiguration))
+	{
+		FullPath = NewCADFile;
+	}
+
+	if (!IFileManager::Get().FileExists(*FullPath))
+	{
+		return FileNotFound;
+	}
+
+	FFileStatData FileStatData = IFileManager::Get().GetStatData(*FullPath);
+	uint32 FileHash = GetFileHash(CADFile, FileStatData, FileConfiguration);
+
+	SceneGraphFile = FString::Printf(TEXT("UEx%08x"), FileHash);
+
+	FString RawDataFile = FPaths::Combine(CachePath, TEXT("scene"), SceneGraphFile + TEXT(".sg"));
+	FString CTFile = FPaths::Combine(CachePath, TEXT("cad"), SceneGraphFile + TEXT(".ct"));
+
+	uint32 MeshFileHash = GetGeomFileHash(FileHash, ImportParameters);
+	MeshFile = FString::Printf(TEXT("UEx%08x"), MeshFileHash);
+	RawDataGeom = FPaths::Combine(CachePath, TEXT("mesh"), MeshFile + TEXT(".gm"));
+
+	bool bNeedToProceed = true;
+	if (IFileManager::Get().FileExists(*RawDataFile))
+	{
+		if (!IFileManager::Get().FileExists(*CTFile)) // the file is scene graph only because no CT file
+		{
+			bNeedToProceed = false;
+		}
+		else if (IFileManager::Get().FileExists(*RawDataGeom)) // the file has been proceed with same meshing parameters
+		{
+			bNeedToProceed = false;
+		}
+		else // the file has been converted into CT file but meshed with different parameters
+		{
+			FullPath = CTFile;
+		}
+	}
+
+	if (!bNeedToProceed)
+	{
+		// The file has been yet proceed, get ExternalRef
+		GetRawDataFileExternalRef(RawDataFile, ExternalRefSet);
+		return ProcessOk;
+	}
+
+	// Process the file
+	return ReadFileWithKernelIO();
+}
+
+
+
+
+EProcessState FCoreTechFileParser::ReadFileWithKernelIO()
+{
+	CT_IO_ERROR Result = IO_OK;
+	CT_OBJECT_ID MainId = 0;
+
+	Result = CT_KERNEL_IO::UnloadModel();
+
+	CT_FLAGS CTImportOption = SetCoreTechImportOption(FPaths::GetExtension(CADFile));
+
+	FString LoadOption;
+	CT_UINT32 NumberOfIds = 1;
+	if (!FileConfiguration.IsEmpty())
+	{
+		NumberOfIds = CT_KERNEL_IO::AskFileNbOfIds(*FullPath);
+		if (NumberOfIds > 1)
+		{
+			CT_UINT32 ActiveConfig = CT_KERNEL_IO::AskFileActiveConfig(*FullPath);
+			for (CT_UINT32 i = 0; i < NumberOfIds; i++)
+			{
+				CT_STR ConfValue = CT_KERNEL_IO::AskFileIdIthName(*FullPath, i);
+				if (FileConfiguration == ConfValue.toUnicode()) {
+					ActiveConfig = i;
+					break;
+				}
+			}
+
+			CTImportOption |= CT_LOAD_FLAGS_READ_SPECIFIC_OBJECT;
+			LoadOption = FString::FromInt((int32) ActiveConfig);
+		}
+	}
+
+	Result = CT_KERNEL_IO::LoadFile(*FullPath, MainId, CTImportOption, 0, *LoadOption);
+	if (Result == IO_ERROR_EMPTY_ASSEMBLY)
+	{
+		Result = CT_KERNEL_IO::UnloadModel();
+		if (Result != IO_OK)
+		{
+			return ProcessFailed;
+		}
+		Result = CT_KERNEL_IO::LoadFile(*FullPath, MainId, CTImportOption | CT_LOAD_FLAGS_LOAD_EXTERNAL_REF);
+	}
+
+	if (Result != IO_OK && Result != IO_OK_MISSING_LICENSES)
+	{
+		CT_KERNEL_IO::UnloadModel();
+		return ProcessFailed;
+	}
+
+	Repair(MainId, ImportParameters.StitchingTechnique);
+	SetCoreTechTessellationState(ImportParameters);
+
+	const CT_OBJECT_TYPE TypeSet[] = { CT_INSTANCE_TYPE, CT_ASSEMBLY_TYPE, CT_PART_TYPE, CT_COMPONENT_TYPE, CT_BODY_TYPE, CT_UNLOADED_COMPONENT_TYPE, CT_UNLOADED_ASSEMBLY_TYPE, CT_UNLOADED_PART_TYPE };
+	uint32 NbElements[8], NbTotal = 10;
+	for (int32 index = 0; index < 8; index++)
+	{
+		CT_KERNEL_IO::AskNbObjectsType(NbElements[index], TypeSet[index]);
+		NbTotal += NbElements[index];
+	}
+
+	{
+		TArray<uint8> RawData;
+		RawData.Append((uint8*)&NbElements[4], sizeof(uint32));
+		FFileHelper::SaveArrayToFile(RawData, *RawDataGeom, &IFileManager::Get());
+	}
+
+	SceneGraphDescription.Reset(NbTotal * 20);
+	ExternalRefSet.Empty(NbTotal);
+	CTIdToRawLineMap.Reset();
+
+	// Header
+	SceneGraphDescription.Add(CADFile);  // line 0: File name
+	SceneGraphDescription.Add(FullPath);  // line 1: File path
+	SceneGraphDescription.Add(TEXT(""));  // Line 2 is reserved
+	SceneGraphDescription.Add(TEXT("Color 0"));        // COLORSETLINE 3
+	SceneGraphDescription.Add(TEXT("Material 0"));     // MATERIALSETLINE 4
+	SceneGraphDescription.Add(TEXT(""));
+	SceneGraphDescription.Add(TEXT(""));
+	SceneGraphDescription.Add(TEXT("ExternalRef 0"));  // EXTERNALREFLINE 7
+	SceneGraphDescription.Add(TEXT("MapCTId 0"));      // MAPCTIDLINE 8
+	SceneGraphDescription.Add(TEXT(""));
+
+	uint32 MaterialNum = GetMaterialNum();
+	MaterialIdToMaterialHashMap.Empty(MaterialNum);
+
+	ReadColor();
+	ReadMaterial();
+
+	// Parse the file
+	bool Ret = ReadNode(MainId);
+	// End of parsing
+
+	if (bNeedSaveCTFile)
+	{
+		CT_LIST_IO ObjectList;
+		ObjectList.PushBack(MainId);
+
+		CT_KERNEL_IO::SaveFile(ObjectList, *FPaths::Combine(CachePath, TEXT("cad"), SceneGraphFile+TEXT(".ct")), L"Ct");
+	}
+
+	CT_KERNEL_IO::UnloadModel();
+
+	if (!Ret)
+	{
+		return ProcessFailed;
+	}
+
+	uint32 LastLine = SceneGraphDescription.Num();
+	if (ExternalRefSet.Num() > 0)
+	{
+		SceneGraphDescription[EXTERNALREFLINE] = FString::Printf(TEXT("ExternalRef %d %d"), LastLine, ExternalRefSet.Num());
+		for (auto ExternalFile : ExternalRefSet)
+		{
+			SceneGraphDescription.Add(ExternalFile);
+		}
+	}
+	SceneGraphDescription.Add("");
+
+	LastLine = SceneGraphDescription.Num();
+	SceneGraphDescription[MAPCTIDLINE] = FString::Printf(TEXT("MapCTId %d"), LastLine);
+
+	FString MapCTId;
+	uint32 SizeId = (uint32) log(CTIdToRawLineMap.Num());
+	uint32 SizeLine = (uint32)log(LastLine);
+	uint32 NbChar = (SizeId > SizeLine ? SizeId : SizeLine) +1;
+	MapCTId.Reserve(CTIdToRawLineMap.Num() * NbChar * 2);
+
+	for (auto ID : CTIdToRawLineMap)
+	{
+		MapCTId.AppendInt(ID.Key);
+		MapCTId.Append(TEXT(" "));
+		MapCTId.AppendInt(ID.Value);
+		MapCTId.Append(TEXT(" "));
+	}
+	SceneGraphDescription.Add(MapCTId);
+	SceneGraphDescription.Add(TEXT(""));
+	SceneGraphDescription.Add(TEXT("F"));
+	
+	ExportFileSceneGraph();
+
+	return ProcessOk;
+}
+
+CT_FLAGS FCoreTechFileParser::SetCoreTechImportOption(const FString& MainFileExt)
+{
+	// Set import option
+	CT_FLAGS Flags = CT_LOAD_FLAGS_USE_DEFAULT;
+
+	if (MainFileExt == TEXT("jt"))
+	{
+		Flags |= CT_LOAD_FLAGS_READ_META_DATA;
+	}
+
+	if (MainFileExt == TEXT("catpart") || MainFileExt == TEXT("catproduct") || MainFileExt == TEXT("cgr"))
+	{
+		Flags |= CT_LOAD_FLAGS_V5_READ_GEOM_SET;
+	}
+
+	// All the BRep topology is not available in IGES import
+	// Ask Kernel IO to complete or create missing topology
+	if (MainFileExt == TEXT(".igs") || MainFileExt == TEXT("iges"))
+	{
+		Flags |= CT_LOAD_FLAG_SEARCH_NEW_TOPOLOGY | CT_LOAD_FLAG_COMPLETE_TOPOLOGY;
+	}
+
+	Flags |= CT_LOAD_FLAGS_V5_READ_GEOM_SET;
+	
+	Flags &= ~CT_LOAD_FLAGS_LOAD_EXTERNAL_REF;
+
+	return Flags;
+}
+
+bool FCoreTechFileParser::ReadNode(CT_OBJECT_ID NodeId)
+{
+	CT_OBJECT_TYPE Type;
+	CT_OBJECT_IO::AskType(NodeId, Type);
+
+	uint32* LineNumber = CTIdToRawLineMap.Find(NodeId);
+	if (LineNumber != nullptr) 
+	{
+		return true;
+	}
+
+	switch (Type)
+	{
+	case CT_INSTANCE_TYPE:
+		return ReadInstance(NodeId);
+
+	case CT_ASSEMBLY_TYPE:
+	case CT_PART_TYPE:
+	case CT_COMPONENT_TYPE:
+		return ReadComponent(NodeId);
+
+	case CT_UNLOADED_ASSEMBLY_TYPE:
+	case CT_UNLOADED_COMPONENT_TYPE:
+	case CT_UNLOADED_PART_TYPE:
+		return ReadUnloadedComponent(NodeId);
+
+	case CT_BODY_TYPE:
+		return ReadBody(NodeId);
+
+		//Treat all CT_CURVE_TYPE :
+	case CT_CURVE_TYPE:
+	case CT_C_NURBS_TYPE:
+	case CT_CONICAL_TYPE:
+	case CT_ELLIPSE_TYPE:
+	case CT_CIRCLE_TYPE:
+	case CT_PARABOLA_TYPE:
+	case CT_HYPERBOLA_TYPE:
+	case CT_LINE_TYPE:
+	case CT_C_COMPO_TYPE:
+	case CT_POLYLINE_TYPE:
+	case CT_EQUATION_CURVE_TYPE:
+	case CT_CURVE_ON_SURFACE_TYPE:
+	case CT_INTERSECTION_CURVE_TYPE:
+	default:
+		break;
+	}
+	return true;
+}
+
+bool FCoreTechFileParser::ReadUnloadedComponent(CT_OBJECT_ID ComponentId)
+{
+	CTIdToRawLineMap.Add(ComponentId, SceneGraphDescription.Num());
+
+	FString Data = TEXT("U ") + FString::FromInt(ComponentId);
+	SceneGraphDescription.Add(*Data);
+
+	ReadNodeMetaDatas(ComponentId);
+
+	CT_STR Filename, FileType;
+	CT_COMPONENT_IO::AskExternalDefinition(ComponentId, Filename, FileType);
+	Data = TEXT("ext ");
+	Data += Filename.toUnicode();
+	Data += TEXT(" ");
+	Data += FileType.toUnicode();
+	SceneGraphDescription.Add(*Data);
+
+	SceneGraphDescription.Add(TEXT(""));
+
+	return true;
+}
+
+bool FCoreTechFileParser::ReadComponent(CT_OBJECT_ID ComponentId)
+{
+	CTIdToRawLineMap.Add(ComponentId, SceneGraphDescription.Num());
+
+	FString Data = TEXT("C ") + FString::FromInt(ComponentId);
+	SceneGraphDescription.Add(*Data);
+
+	ReadNodeMetaDatas(ComponentId);
+
+	CT_LIST_IO Children;
+	CT_COMPONENT_IO::AskChildren(ComponentId, Children);
+
+	// Parse children
+	Data = TEXT("children ") + FString::FromInt(Children.Count());
+	SceneGraphDescription.Add(*Data);
+
+	Children.IteratorInitialize();
+	CT_OBJECT_ID CurrentObjectId;
+	while ((CurrentObjectId = Children.IteratorIter()) != 0)
+	{
+		SceneGraphDescription.Add(*FString::FromInt(CurrentObjectId));
+	}
+	SceneGraphDescription.Add(TEXT(""));
+
+	Children.IteratorInitialize();
+	while ((CurrentObjectId = Children.IteratorIter()) != 0)
+	{
+		ReadNode(CurrentObjectId);
+	}
+
+	return true;
+}
+
+bool FCoreTechFileParser::ReadInstance(CT_OBJECT_ID InstanceNodeId)
+{
+	CTIdToRawLineMap.Add(InstanceNodeId, SceneGraphDescription.Num());
+
+	FString Data = TEXT("I ") + FString::FromInt(InstanceNodeId);
+	SceneGraphDescription.Add(*Data);
+
+	ReadNodeMetaDatas(InstanceNodeId);
+
+	// Ask the transformation of the instance
+	double CTMatrix[16];
+	CT_INSTANCE_IO::AskTransformation(InstanceNodeId, CTMatrix);
+	SceneGraphDescription.Add(FString::Printf(TEXT("matrix %f %f %f %f %f %f %f %f %f %f %f %f %f %f %f %f"), CTMatrix[0], CTMatrix[1], CTMatrix[2], CTMatrix[3], CTMatrix[4], CTMatrix[5], CTMatrix[6], CTMatrix[7], CTMatrix[8], CTMatrix[9], CTMatrix[10], CTMatrix[11], CTMatrix[12], CTMatrix[13], CTMatrix[14], CTMatrix[15]));
+
+	// Ask the reference
+	CT_OBJECT_ID ReferenceNodeId;
+	CT_IO_ERROR CTReturn = CT_INSTANCE_IO::AskChild(InstanceNodeId, ReferenceNodeId);
+	if (CTReturn != CT_IO_ERROR::IO_OK)
+		return false;
+
+	NodeConfiguration.Empty(); // #ueent_CAD
+	CT_OBJECT_TYPE type;
+	CT_OBJECT_IO::AskType(ReferenceNodeId, type);
+	if (type == CT_UNLOADED_PART_TYPE || type == CT_UNLOADED_COMPONENT_TYPE || type == CT_UNLOADED_ASSEMBLY_TYPE)
+	{
+		CT_STR ComponentFile, FileType;
+		CT_COMPONENT_IO::AskExternalDefinition(ReferenceNodeId, ComponentFile, FileType);
+		FString ExternalRef = FPaths::GetCleanFilename(ComponentFile.toUnicode());
+		if(!NodeConfiguration.IsEmpty())
+		{
+			ExternalRef += TEXT("|") + NodeConfiguration;
+		}
+
+		Data = TEXT("ext ") + FString::FromInt(ReferenceNodeId) + TEXT(" ") + ExternalRef;
+		SceneGraphDescription.Add(*Data);
+
+		ExternalRefSet.Add(ExternalRef);
+	}
+	else
+	{
+		Data = TEXT("ref ") + FString::FromInt(ReferenceNodeId);
+		SceneGraphDescription.Add(Data);
+	}
+	SceneGraphDescription.Add(TEXT(""));
+
+	return ReadNode(ReferenceNodeId);
+}
+
+
+uint32 GetBodiesFaceSetNum(TArray<CT_OBJECT_ID>& BodySet)
+{
+	uint32 size = 0;
+	for (int Index = 0; Index < BodySet.Num(); Index++)
+	{
+		// Loop through the face of the first body and collect material data
+		CT_LIST_IO FaceList;
+		CT_BODY_IO::AskFaces(BodySet[Index], FaceList);
+		size += FaceList.Count();
+	}
+	return size;
+}
+
+uint32 GetBodiesTessellations(TArray<CT_OBJECT_ID>& BodySet, TArray<FTessellationData>& FaceTessellationSet, TMap<uint32, uint32>& MaterialIdToMaterialHash, int32& OutRawDataSize, const float ScaleFactor)
+{
+	uint32 FaceSize = GetBodiesFaceSetNum(BodySet);
+
+	// Allocate memory space for tessellation data
+	FaceTessellationSet.Reserve(FaceSize);
+
+	uint32 GlobalTriangleCount = 0;
+
+	CT_MATERIAL_ID BodyCtMaterialId = 0;
+	for (int Index = 0; Index < BodySet.Num(); Index++)
+	{
+		FObjectDisplayDataId BodyMaterial;
+		GetCTObjectDisplayDataIds(BodySet[Index], BodyMaterial);
+
+		// Loop through the face of the first body and collect all tessellation data
+		CT_LIST_IO FaceList;
+		CT_BODY_IO::AskFaces(BodySet[Index], FaceList);
+		FaceList.IteratorInitialize();
+
+		int32 LastIndex = FaceTessellationSet.Num();
+		int32 FaceIndex = LastIndex;
+
+		CT_OBJECT_ID FaceID;
+		while ((FaceID = FaceList.IteratorIter()) != 0)
+		{
+			int32 FaceRawSize = 0;
+			uint32 TriangleCount = GetFaceTessellation(FaceID, FaceTessellationSet, FaceRawSize, ScaleFactor);
+
+			if (TriangleCount == 0)
+			{
+				continue;
+			}
+
+			GlobalTriangleCount += TriangleCount;
+			OutRawDataSize += (FaceRawSize + ECTTessellationLine::LastLine * sizeof(uint32));
+			FObjectDisplayDataId FaceMaterial;
+			GetCTObjectDisplayDataIds(FaceID, FaceMaterial);
+			SetFaceMainMaterial(FaceMaterial, BodyMaterial, MaterialIdToMaterialHash, FaceTessellationSet[FaceIndex]);
+			FaceIndex++;
+		}
+
+		int32 FaceNum = FaceTessellationSet.Num() - LastIndex;
+		for (FaceIndex = LastIndex; FaceIndex < FaceTessellationSet.Num(); FaceIndex++)
+		{
+			FaceTessellationSet[FaceIndex].BodyId = (uint32)BodySet[Index];
+			FaceTessellationSet[FaceIndex].BodyFaceNum = FaceNum;
+		}
+	}
+	return GlobalTriangleCount;
+}
+
+
+bool FCoreTechFileParser::ReadBody(CT_OBJECT_ID BodyId)
+{
+	CTIdToRawLineMap.Add(BodyId, SceneGraphDescription.Num());
+
+	bNeedSaveCTFile = true;
+
+	FString Data = TEXT("B ") + FString::FromInt(BodyId);
+	SceneGraphDescription.Add(*Data);
+	
+	ReadNodeMetaDatas(BodyId);
+
+	// Parse Body to set Get Mesh
+	TArray<CT_OBJECT_ID> BodySet;
+	BodySet.Add(BodyId);
+	TArray<FTessellationData> FaceTessellationSet;
+
+	uint32 BodyUuId = GetStaticMeshUuid(*SceneGraphFile, BodyId);
+	double ScaleFactor = 0.1;
+                
+	int32 BodyRawDataSize = 0;
+	uint32 NbTriangles = GetBodiesTessellations(BodySet, FaceTessellationSet, MaterialIdToMaterialHashMap, BodyRawDataSize, ScaleFactor);
+	TArray<uint8> GlobalRawData;
+
+	GlobalRawData.Reserve(BodyRawDataSize);
+
+	int32 FaceNum = FaceTessellationSet.Num();
+	for (FTessellationData& Tessellation : FaceTessellationSet)
+	{
+		Tessellation.BodyUuId = BodyUuId;
+		Tessellation.BodyFaceNum = FaceNum;
+		WriteTessellationInRawData(Tessellation, GlobalRawData);
+	}
+
+	bool bWriteOk = FFileHelper::SaveArrayToFile(GlobalRawData, *RawDataGeom, &IFileManager::Get(), FILEWRITE_Append);
+	ensure(bWriteOk);
+	GlobalRawData.Empty();
+
+	CT_LIST_IO ObjectList;
+	ObjectList.PushBack(BodyId);
+	FString BodyFile = FString::Printf(TEXT("UEx%08x"), BodyUuId);
+	CT_KERNEL_IO::SaveFile(ObjectList, *FPaths::Combine(CachePath, TEXT("body"), BodyFile + TEXT(".ct")), L"Ct");
+
+	FString MapCTId;
+	MapCTId.Reserve((MaterialIdToMaterialHashMap.Num()+1) * 22);
+
+	MapCTId.Append(TEXT("materialMap "));
+	for (auto Material : MaterialIdToMaterialHashMap)
+	{
+		MapCTId.Append(FString::Printf(TEXT("%u %u "), Material.Key, Material.Value));
+	}
+	SceneGraphDescription.Add(MapCTId);
+	SceneGraphDescription.Add(TEXT(""));
+
+	return true;
+}
+
+void FCoreTechFileParser::GetAttributeValue(CT_ATTRIB_TYPE AttributType, int IthField, FString& Value)
+{
+	CT_STR               FieldName;
+	CT_ATTRIB_FIELD_TYPE FieldType;
+
+	Value = "";
+
+	if (CT_ATTRIB_DEFINITION_IO::AskFieldDefinition(AttributType, IthField, FieldType, FieldName) != IO_OK) return;
+
+	switch (FieldType) {
+	case CT_ATTRIB_FIELD_UNKNOWN:
+	{
+		break;
+	}
+	case CT_ATTRIB_FIELD_INTEGER:
+	{
+		int IValue;
+		if (CT_CURRENT_ATTRIB_IO::AskIntField(IthField, IValue) != IO_OK) break;
+		Value = FString::FromInt(IValue);
+		break;
+	}
+	case CT_ATTRIB_FIELD_DOUBLE:
+	{
+		double DValue;
+		if (CT_CURRENT_ATTRIB_IO::AskDblField(IthField, DValue) != IO_OK) break;
+		Value = FString::Printf(TEXT("%lf"), DValue);
+		break;
+	}
+	case CT_ATTRIB_FIELD_STRING:
+	{
+		CT_STR StrValue;
+		if (CT_CURRENT_ATTRIB_IO::AskStrField(IthField, StrValue) != IO_OK) break;
+		Value = StrValue.toUnicode();
+		break;
+	}
+	case CT_ATTRIB_FIELD_POINTER:
+	{
+		break;
+	}
+	}
+}
+
+void FCoreTechFileParser::ReadNodeMetaDatas(CT_OBJECT_ID NodeId)
+{
+	const FString ConfigName = TEXT("Configuration Name");
+	FString AttributeNameLigne = TEXT("M");
+	SceneGraphDescription.Add(AttributeNameLigne);
+	int32 AttributIndex = SceneGraphDescription.Num();
+
+	if (CT_COMPONENT_IO::IsA(NodeId, CT_COMPONENT_TYPE))
+	{
+		CT_STR FileName, FileType;
+		CT_COMPONENT_IO::AskExternalDefinition(NodeId, FileName, FileType);
+		if (!FileName.IsEmpty())
+		{
+			SceneGraphDescription.Add(TEXT("ExternalDefinition"));
+			SceneGraphDescription.Add(FileName.toUnicode());
+		}
+	}
+
+	CT_UINT32 NbAttributes;
+	CT_OBJECT_IO::AskNbAttributes(NodeId, CT_ATTRIB_ALL, NbAttributes);
+	CT_UINT32 ith_attrib = 0;
+
+	CT_SHOW_ATTRIBUTE IsShow = CT_UNKNOWN;
+	if (CT_OBJECT_IO::AskShowAttribute(NodeId, IsShow) == IO_OK)
+	{
+		switch (IsShow)
+		{
+		case CT_SHOW:
+			SceneGraphDescription.Add(TEXT("ShowAttribute"));
+			SceneGraphDescription.Add(TEXT("show"));
+			break;
+		case CT_NOSHOW:
+			SceneGraphDescription.Add(TEXT("ShowAttribute"));
+			SceneGraphDescription.Add(TEXT("noShow"));
+			break;
+		case CT_UNKNOWN:
+			SceneGraphDescription.Add(TEXT("ShowAttribute"));
+			SceneGraphDescription.Add(TEXT("unknown"));
+			break;
+		}
+	}
+
+	while (CT_OBJECT_IO::SearchAttribute(NodeId, CT_ATTRIB_ALL, ith_attrib++) == IO_OK)
+	{
+		// Get the current attribute type
+		CT_ATTRIB_TYPE       AttributeType;
+		CT_STR               TypeName;
+
+		CT_STR               FieldName;
+		CT_STR               FieldStrValue;
+		CT_INT32             FieldIntValue;
+		CT_DOUBLE            FieldDoubleValue0, FieldDoubleValue1, FieldDoubleValue2;
+		FString              FieldValue;
+
+
+		if (CT_CURRENT_ATTRIB_IO::AskAttributeType(AttributeType) != IO_OK) continue;;
+		switch (AttributeType) {
+
+		case CT_ATTRIB_SPLT:
+			//if (CT_CURRENT_ATTRIB_IO::AskStrField(ITH_SPLT_VALUE, field_strValue) != IO_OK) break;
+			//nodeAttributeSet[TEXT("SPLT")] = field_strValue.toUnicode();
+			break;
+
+		case CT_ATTRIB_NAME:
+			if (CT_CURRENT_ATTRIB_IO::AskStrField(ITH_NAME_VALUE, FieldStrValue) != IO_OK) break;
+			if (FieldStrValue.IsEmpty()) { break; }
+			SceneGraphDescription.Add(TEXT("CTName"));
+			SceneGraphDescription.Add(FieldStrValue.toUnicode());
+			break;
+
+		case CT_ATTRIB_ORIGINAL_NAME:
+			if (CT_CURRENT_ATTRIB_IO::AskStrField(ITH_NAME_VALUE, FieldStrValue) != IO_OK) break;
+			if (FieldStrValue.IsEmpty()) { break; }
+			SceneGraphDescription.Add(TEXT("Name"));
+			SceneGraphDescription.Add(FieldStrValue.toUnicode());
+			break;
+
+		case CT_ATTRIB_ORIGINAL_FILENAME:
+			if (CT_CURRENT_ATTRIB_IO::AskStrField(ITH_FILENAME_VALUE, FieldStrValue) != IO_OK) break;
+			if (FieldStrValue.IsEmpty()) { break; }
+			SceneGraphDescription.Add(TEXT("FileName"));
+			SceneGraphDescription.Add(FieldStrValue.toUnicode());
+			break;
+
+		case CT_ATTRIB_UUID:
+			if (CT_CURRENT_ATTRIB_IO::AskStrField(ITH_UUID_VALUE, FieldStrValue) != IO_OK) break;
+			if (FieldStrValue.IsEmpty()) { break; }
+			SceneGraphDescription.Add(TEXT("UUID"));
+			SceneGraphDescription.Add(FieldStrValue.toUnicode());
+			break;
+
+		case CT_ATTRIB_INPUT_FORMAT_AND_EMETTOR:
+			if (CT_CURRENT_ATTRIB_IO::AskStrField(ITH_INPUT_FORMAT_AND_EMETTOR, FieldStrValue) != IO_OK) break;
+			if (FieldStrValue.IsEmpty()) { break; }
+			SceneGraphDescription.Add(TEXT("Input_Format_and_Emitter"));
+			SceneGraphDescription.Add(FieldStrValue.toUnicode());
+			break;
+
+		case CT_ATTRIB_CONFIGURATION_NAME:
+			if (CT_CURRENT_ATTRIB_IO::AskStrField(ITH_NAME_VALUE, FieldStrValue) != IO_OK) break;
+			if (FieldStrValue.IsEmpty()) { break; }
+			SceneGraphDescription.Add(TEXT("ConfigurationName"));
+			SceneGraphDescription.Add(FieldStrValue.toUnicode());
+			break;
+
+		case CT_ATTRIB_LAYERID:
+			GetAttributeValue(AttributeType, ITH_LAYERID_VALUE, FieldValue);
+			SceneGraphDescription.Add(TEXT("LayerId"));
+			SceneGraphDescription.Add(FieldValue);
+			GetAttributeValue(AttributeType, ITH_LAYERID_NAME, FieldValue);
+			SceneGraphDescription.Add(TEXT("LayerName"));
+			SceneGraphDescription.Add(FieldValue);
+			GetAttributeValue(AttributeType, ITH_LAYERID_FLAG, FieldValue);
+			SceneGraphDescription.Add(TEXT("LayerFlag"));
+			SceneGraphDescription.Add(FieldValue);
+			break;
+
+		case CT_ATTRIB_COLORID:
+			{
+				if (CT_CURRENT_ATTRIB_IO::AskIntField(ITH_COLORID_VALUE, FieldIntValue) != IO_OK) break;
+				uint32 ColorId = FieldIntValue;
+				uint8 Alpha = 255;
+				SceneGraphDescription.Add(TEXT("ColorId"));
+				SceneGraphDescription.Add(FString::FromInt(FieldIntValue));
+
+				CT_COLOR CtColor;
+				if (CT_MATERIAL_IO::AskIndexedColor(FieldIntValue, CtColor) != IO_OK) break;
+
+				if (CT_OBJECT_IO::SearchAttribute(NodeId, CT_ATTRIB_TRANSPARENCY) == IO_OK)
+				{
+					if (CT_CURRENT_ATTRIB_IO::AskDblField(0, FieldDoubleValue0) == IO_OK)
+					{
+						Alpha = FMath::Max((1. - FieldDoubleValue0), FieldDoubleValue0) * 255.;
+					}
+				}
+				FString colorHexa = FString::Printf(TEXT("%02x%02x%02x%02x"), CtColor[0], CtColor[1], CtColor[2], Alpha);
+				SceneGraphDescription.Add(TEXT("ColorValue"));
+				SceneGraphDescription.Add(colorHexa);
+
+				SceneGraphDescription.Add(TEXT("ColorUEId"));
+				uint32 ColorUuId = BuildFastColorHash(ColorId, Alpha);
+				SceneGraphDescription.Add(FString::Printf(TEXT("%u"), ColorUuId));
+
+				int32 ColorHash = BuildColorHash(CtColor, Alpha);
+				MaterialIdToMaterialHashMap.Add(ColorUuId, ColorHash);
+			}
+			break;
+
+		case CT_ATTRIB_MATERIALID:
+			if (CT_CURRENT_ATTRIB_IO::AskIntField(ITH_MATERIALID_VALUE, FieldIntValue) != IO_OK) break;
+			SceneGraphDescription.Add(TEXT("MaterialId"));
+			SceneGraphDescription.Add(FString::FromInt(FieldIntValue));
+			break;
+
+		case CT_ATTRIB_TRANSPARENCY:
+			if (CT_CURRENT_ATTRIB_IO::AskDblField(ITH_TRANSPARENCY_VALUE, FieldDoubleValue0) != IO_OK) break;
+			FieldIntValue = FMath::Max((1. - FieldDoubleValue0), FieldDoubleValue0) * 255.;
+			SceneGraphDescription.Add(TEXT("Transparency"));
+			SceneGraphDescription.Add(FString::FromInt(FieldIntValue));
+			break;
+
+		case CT_ATTRIB_COMMENT:
+			//ITH_COMMENT_POSX, ITH_COMMENT_POSY, ITH_COMMENT_POSZ, ITH_COMMENT_TEXT
+			break;
+
+		case CT_ATTRIB_REFCOUNT:
+			if (CT_CURRENT_ATTRIB_IO::AskIntField(ITH_REFCOUNT_VALUE, FieldIntValue) != IO_OK) break;
+			//NodeAttributeSet.Add(TEXT("RefCount"));
+			SceneGraphDescription.Add(FString::FromInt(FieldIntValue));
+			break;
+
+		case CT_ATTRIB_TESS_PARAMS:
+			//ITH_TESS_PARAMS_MAXSAG, ITH_TESS_PARAMS_MAXANGLE, ITH_TESS_PARAMS_MAXLENGTH
+			break;
+
+		case CT_ATTRIB_COMPARE_RESULT:
+			//ITH_COMPARE_TYPE, ITH_COMPARE_RESULT_MAXDIST, ITH_COMPARE_RESULT_MINDIST, ITH_COMPARE_THICKNESS, ITH_COMPARE_BACKLASH, ITH_COMPARE_CLEARANCE, ITH_COMPARE_MAXTHICKNESS, ITH_COMPARE_RESULT_NEAREST, ITH_COMPARE_RESULT_ORIGIN1_X, ITH_COMPARE_RESULT_ORIGIN1_Y, ITH_COMPARE_RESULT_ORIGIN1_Z, ITH_COMPARE_RESULT_DIRECTION1_X, ITH_COMPARE_RESULT_DIRECTION1_Y, ITH_COMPARE_RESULT_DIRECTION1_Z, ITH_COMPARE_RESULT_RADIUS1_1, ITH_COMPARE_RESULT_RADIUS2_1, ITH_COMPARE_RESULT_HALFHANGLE1, ITH_COMPARE_RESULT_ORIGIN2_X, ITH_COMPARE_RESULT_ORIGIN2_Y, ITH_COMPARE_RESULT_ORIGIN2_Z, ITH_COMPARE_RESULT_DIRECTION2_X, ITH_COMPARE_RESULT_DIRECTION2_Y, ITH_COMPARE_RESULT_DIRECTION2_Z, ITH_COMPARE_RESULT_RADIUS1_2, ITH_COMPARE_RESULT_RADIUS2_2, ITH_COMPARE_RESULT_HALFHANGLE2, ITH_COMPARE_RESULT_COLOR1, ITH_COMPARE_RESULT_COLOR2, ITH_COMPARE_DIST_VECTOR1_X, ITH_COMPARE_DIST_VECTOR1_Y, ITH_COMPARE_DIST_VECTOR1_Z, ITH_COMPARE_DIST_VECTOR2_X, ITH_COMPARE_DIST_VECTOR2_Y, ITH_COMPARE_DIST_VECTOR2_Z
+			break;
+		case CT_ATTRIB_DENSITY:
+			//ITH_VOLUME_DENSITY_VALUE
+			break;
+
+		case CT_ATTRIB_MASS_PROPERTIES:
+			if (CT_CURRENT_ATTRIB_IO::AskDblField(ITH_MASS_PROPERTIES_AREA, FieldDoubleValue0) != IO_OK) break;
+			SceneGraphDescription.Add(TEXT("Area"));
+			SceneGraphDescription.Add(FString::Printf(TEXT("%lf"), FieldDoubleValue0));
+			if (CT_CURRENT_ATTRIB_IO::AskDblField(ITH_MASS_PROPERTIES_VOLUME, FieldDoubleValue0) != IO_OK) break;
+			SceneGraphDescription.Add(TEXT("Volume"));
+			SceneGraphDescription.Add(FString::Printf(TEXT("%lf"), FieldDoubleValue0));
+			if (CT_CURRENT_ATTRIB_IO::AskDblField(ITH_MASS_PROPERTIES_MASS, FieldDoubleValue0) != IO_OK) break;
+			SceneGraphDescription.Add(TEXT("Mass"));
+			SceneGraphDescription.Add(FString::Printf(TEXT("%lf"), FieldDoubleValue0));
+			if (CT_CURRENT_ATTRIB_IO::AskDblField(ITH_MASS_PROPERTIES_LENGTH, FieldDoubleValue0) != IO_OK) break;
+			SceneGraphDescription.Add(TEXT("Length"));
+			SceneGraphDescription.Add(FString::Printf(TEXT("%lf"), FieldDoubleValue0));
+			//ITH_MASS_PROPERTIES_COGX, ITH_MASS_PROPERTIES_COGY, ITH_MASS_PROPERTIES_COGZ
+			//ITH_MASS_PROPERTIES_M1, ITH_MASS_PROPERTIES_M2, ITH_MASS_PROPERTIES_M3
+			//ITH_MASS_PROPERTIES_IXXG,ITH_MASS_PROPERTIES_IYYG, ITH_MASS_PROPERTIES_IZZG, ITH_MASS_PROPERTIES_IXYG, ITH_MASS_PROPERTIES_IYZG, ITH_MASS_PROPERTIES_IZXG
+			//ITH_MASS_PROPERTIES_AXIS1X, ITH_MASS_PROPERTIES_AXIS1Y, ITH_MASS_PROPERTIES_AXIS1Z, ITH_MASS_PROPERTIES_AXIS2X, ITH_MASS_PROPERTIES_AXIS2Y, ITH_MASS_PROPERTIES_AXIS2Z, ITH_MASS_PROPERTIES_AXIS3X, ITH_MASS_PROPERTIES_AXIS3Y, ITH_MASS_PROPERTIES_AXIS3Z
+			//ITH_MASS_PROPERTIES_XMIN, ITH_MASS_PROPERTIES_YMIN, ITH_MASS_PROPERTIES_ZMIN, ITH_MASS_PROPERTIES_XMAX, ITH_MASS_PROPERTIES_YMAX, ITH_MASS_PROPERTIES_ZMAX
+			break;
+
+		case CT_ATTRIB_THICKNESS:
+			//ITH_THICKNESS_VALUE
+			break;
+
+		case CT_ATTRIB_INTEGER_METADATA:
+			if (CT_CURRENT_ATTRIB_IO::AskStrField(ITH_INTEGER_METADATA_NAME, FieldName) != IO_OK) break;
+			if (CT_CURRENT_ATTRIB_IO::AskIntField(ITH_INTEGER_METADATA_VALUE, FieldIntValue) != IO_OK) break;
+			SceneGraphDescription.Add(FieldName.toUnicode());
+			SceneGraphDescription.Add(FString::FromInt(FieldIntValue));
+			break;
+
+		case CT_ATTRIB_DOUBLE_METADATA:
+			if (CT_CURRENT_ATTRIB_IO::AskStrField(ITH_DOUBLE_METADATA_NAME, FieldName) != IO_OK) break;
+			if (CT_CURRENT_ATTRIB_IO::AskDblField(ITH_DOUBLE_METADATA_VALUE, FieldDoubleValue0) != IO_OK) break;
+			SceneGraphDescription.Add(FieldName.toUnicode());
+			SceneGraphDescription.Add(FString::Printf(TEXT("%lf"), FieldDoubleValue0));
+			break;
+
+		case CT_ATTRIB_STRING_METADATA:
+			if (CT_CURRENT_ATTRIB_IO::AskStrField(ITH_STRING_METADATA_NAME, FieldName) != IO_OK) break;
+			if (CT_CURRENT_ATTRIB_IO::AskStrField(ITH_STRING_METADATA_VALUE, FieldStrValue) != IO_OK) break;
+			if (FieldStrValue.IsEmpty()) break;
+			if(ConfigName== FieldName.toUnicode())
+			{
+				NodeConfiguration = FieldStrValue.toUnicode();
+			}
+			SceneGraphDescription.Add(FieldName.toUnicode());
+			SceneGraphDescription.Add(FieldStrValue.toUnicode());
+			break;
+
+		case CT_ATTRIB_ORIGINAL_UNITS:
+			if (CT_CURRENT_ATTRIB_IO::AskDblField(ITH_ORIGINAL_UNITS_MASS, FieldDoubleValue0) != IO_OK) break;
+			if (CT_CURRENT_ATTRIB_IO::AskDblField(ITH_ORIGINAL_UNITS_LENGTH, FieldDoubleValue1) != IO_OK) break;
+			if (CT_CURRENT_ATTRIB_IO::AskDblField(ITH_ORIGINAL_UNITS_DURATION, FieldDoubleValue2) != IO_OK) break;
+			SceneGraphDescription.Add(TEXT("OriginalUnitsMass"));
+			SceneGraphDescription.Add(FString::Printf(TEXT("%lf"), FieldDoubleValue0));
+			SceneGraphDescription.Add(TEXT("OriginalUnitsLength"));
+			SceneGraphDescription.Add(FString::Printf(TEXT("%lf"), FieldDoubleValue1));
+			SceneGraphDescription.Add(TEXT("OriginalUnitsDuration"));
+			SceneGraphDescription.Add(FString::Printf(TEXT("%lf"), FieldDoubleValue2));
+			break;
+
+		case CT_ATTRIB_ORIGINAL_TOLERANCE:
+			//ITH_ORIGINAL_TOLERANCE_LENGTH, ITH_ORIGINAL_TOLERANCE_MAXCOORD, ITH_ORIGINAL_TOLERANCE_ANGLE
+			break;
+		case CT_ATTRIB_IGES_PARAMETERS:
+			//ITH_IGES_DELIMITOR, ITH_IGES_ENDDELIMITOR, ITH_IGES_EMETTORID, ITH_IGES_CREATION_FILENAME, ITH_IGES_EMETTOR_NAME, ITH_IGES_EMETTOR_VERSION, ITH_IGES_NBOFBITS, ITH_IGES_MAXPOW_SINGLE_PREC, ITH_IGES_NB_DIGITS_SINGLE_PRE, ITH_IGES_MAXPOW_DOUBLE_PREC, ITH_IGES_NB_DIGITS_DOUBLE_PRE, ITH_IGES_RECEPTOR_ID, ITH_IGES_SCALE, ITH_IGES_UNIT, ITH_IGES_UNIT_NAME, ITH_IGES_MAXLINEGRADATIONS, ITH_IGES_MAXLINEWIDTH, ITH_IGES_CREATIONFILEDATE, ITH_IGES_MINIRESOLUTION, ITH_IGES_MAXCOORD, ITH_IGES_AUTHORNAME, ITH_IGES_ORGANIZATION, ITH_IGES_IGESVERSION, ITH_IGES_DRAFTINGSTANDARD, ITH_IGES_MODELDATE, ITH_IGES_APPLICATIONPROTOCOL
+			break;
+		case CT_ATTRIB_READ_V4_MARKER:
+			//ITH_READ_V4_MARKER_VALUE
+			break;
+
+		case CT_ATTRIB_PRODUCT:
+			if (CT_CURRENT_ATTRIB_IO::AskStrField(ITH_PRODUCT_REVISION, FieldStrValue) != IO_OK) break;
+			if (FieldStrValue.IsEmpty()) { break; }
+			SceneGraphDescription.Add(TEXT("ProductRevision"));
+			SceneGraphDescription.Add(FieldStrValue.toUnicode());
+			if (CT_CURRENT_ATTRIB_IO::AskStrField(ITH_PRODUCT_DEFINITION, FieldStrValue) != IO_OK) break;
+			if (FieldStrValue.IsEmpty()) { break; }
+			SceneGraphDescription.Add(TEXT("ProductDefinition"));
+			SceneGraphDescription.Add(FieldStrValue.toUnicode());
+			if (CT_CURRENT_ATTRIB_IO::AskStrField(ITH_PRODUCT_NOMENCLATURE, FieldStrValue) != IO_OK) break;
+			if (FieldStrValue.IsEmpty()) { break; }
+			SceneGraphDescription.Add(TEXT("ProductNomenclature"));
+			SceneGraphDescription.Add(FieldStrValue.toUnicode());
+			if (CT_CURRENT_ATTRIB_IO::AskStrField(ITH_PRODUCT_SOURCE, FieldStrValue) != IO_OK) break;
+			if (FieldStrValue.IsEmpty()) { break; }
+			SceneGraphDescription.Add(TEXT("ProductSource"));
+			SceneGraphDescription.Add(FieldStrValue.toUnicode());
+			if (CT_CURRENT_ATTRIB_IO::AskStrField(ITH_PRODUCT_DESCRIPTION, FieldStrValue) != IO_OK) break;
+			if (FieldStrValue.IsEmpty()) { break; }
+			SceneGraphDescription.Add(TEXT("ProductDescription"));
+			SceneGraphDescription.Add(FieldStrValue.toUnicode());
+			break;
+
+		case CT_ATTRIB_SIMPLIFY:
+			//ITH_SIMPLIFY_CURRENT_ID, ITH_SIMPLIFY_CONNECTED_ID
+			break;
+		case CT_ATTRIB_MIDFACE:
+			//ITH_MIDFACE_TYPE, ITH_MIDFACE_THICKNESS1, ITH_MIDFACE_THICKNESS2
+			break;
+		case CT_ATTRIB_DEBUG_STRING:
+			//ITH_DEBUG_STRING_VALUE
+			break;
+		case CT_ATTRIB_DEFEATURING:
+			break;
+		case CT_ATTRIB_BREPLINKID:
+			//ITH_BREPLINKID_BRANCHID, ITH_BREPLINKID_FACEID
+			break;
+		case CT_ATTRIB_MARKUPS_REF:
+			//ITH_MARKUPS_REF
+			break;
+		case CT_ATTRIB_COLLISION:
+			//ITH_COLLISION_ID, ITH_COLLISION_MATID
+			break;
+		case CT_ATTRIB_EXTERNAL_ID:
+			//ITH_EXTERNAL_ID_VALUE
+			break;
+		case CT_ATTRIB_MODIFIER:
+			//ITH_MODIFIER_TYPE, ITH_MODIFIER_INTVALUE, ITH_MODIFIER_DBLVALUE, ITH_MODIFIER_STRVALUE
+			break;
+		case CT_ATTRIB_ORIGINAL_SURF_OLD:
+			break;
+		case CT_ATTRIB_RESULT_BREPLINKID:
+			break;
+		case CT_ATTRIB_AREA:
+			//ITH_AREA_VALUE
+			break;
+		case CT_ATTRIB_ACIS_SG_PIDNAME:
+			//ITH_ACIS_SG_PIDNAME_BASE_NAME, ITH_ACIS_SG_PIDNAME_TIME_VAL, ITH_ACIS_SG_PIDNAME_INDEX, ITH_ACIS_SG_PIDNAME_COPY_NUM
+			break;
+		case CT_ATTRIB_CURVE_ORIGINAL_BOUNDARY_PARAMS:
+			//ITH_CURVE_ORIGINAL_BOUNDARY_PARAMS_START, ITH_CURVE_ORIGINAL_BOUNDARY_PARAMS_END, ITH_CURVE_ORIGINAL_BOUNDARY_PARAMS_SCALE
+			break;
+
+		case CT_ATTRIB_INTEGER_PARAMETER:
+			if (CT_CURRENT_ATTRIB_IO::AskStrField(ITH_INTEGER_PARAMETER_NAME, FieldName) != IO_OK) break;
+			if (CT_CURRENT_ATTRIB_IO::AskIntField(ITH_INTEGER_PARAMETER_VALUE, FieldIntValue) != IO_OK) break;
+			SceneGraphDescription.Add(FieldName.toUnicode());
+			SceneGraphDescription.Add(FString::FromInt(FieldIntValue));
+			break;
+
+		case CT_ATTRIB_DOUBLE_PARAMETER:
+			if (CT_CURRENT_ATTRIB_IO::AskStrField(ITH_DOUBLE_PARAMETER_NAME, FieldName) != IO_OK) break;
+			if (CT_CURRENT_ATTRIB_IO::AskDblField(ITH_DOUBLE_PARAMETER_VALUE, FieldDoubleValue0) != IO_OK) break;
+			SceneGraphDescription.Add(FieldName.toUnicode());
+			SceneGraphDescription.Add(FString::Printf(TEXT("%lf"), FieldDoubleValue0));
+			break;
+
+		case CT_ATTRIB_STRING_PARAMETER:
+			if (CT_CURRENT_ATTRIB_IO::AskStrField(ITH_STRING_PARAMETER_NAME, FieldName) != IO_OK) break;
+			if (CT_CURRENT_ATTRIB_IO::AskStrField(ITH_STRING_PARAMETER_VALUE, FieldStrValue) != IO_OK) break;
+			if (FieldStrValue.IsEmpty()) { break; }
+			SceneGraphDescription.Add(FieldName.toUnicode());
+			SceneGraphDescription.Add(FieldStrValue.toUnicode());
+			break;
+
+		case CT_ATTRIB_PARAMETER_ARRAY:
+			//ITH_PARAMETER_ARRAY_NAME
+			//ITH_PARAMETER_ARRAY_NUMBER
+			//ITH_PARAMETER_ARRAY_VALUES
+			break;
+
+		case CT_ATTRIB_SAVE_OPTION:
+			if (CT_CURRENT_ATTRIB_IO::AskStrField(ITH_SAVE_OPTION_AUTHOR, FieldStrValue) != IO_OK) break;
+			if (FieldStrValue.IsEmpty()) { break; }
+			//NodeAttributeSet.Add(TEXT("SaveOptionAuthor"));
+			SceneGraphDescription.Add(FieldStrValue.toUnicode());
+			if (CT_CURRENT_ATTRIB_IO::AskStrField(ITH_SAVE_OPTION_ORGANIZATION, FieldStrValue) != IO_OK) break;
+			if (FieldStrValue.IsEmpty()) { break; }
+			//NodeAttributeSet.Add(TEXT("SaveOptionOrganization"));
+			SceneGraphDescription.Add(FieldStrValue.toUnicode());
+			if (CT_CURRENT_ATTRIB_IO::AskStrField(ITH_SAVE_OPTION_FILE_DESCRIPTION, FieldStrValue) != IO_OK) break;
+			if (FieldStrValue.IsEmpty()) { break; }
+			//NodeAttributeSet.Add(TEXT("SaveOptionFileDescription"));
+			SceneGraphDescription.Add(FieldStrValue.toUnicode());
+			if (CT_CURRENT_ATTRIB_IO::AskStrField(ITH_SAVE_OPTION_AUTHORISATION, FieldStrValue) != IO_OK) break;
+			if (FieldStrValue.IsEmpty()) { break; }
+			//NodeAttributeSet.Add(TEXT("SaveOptionAuthorisation"));
+			SceneGraphDescription.Add(FieldStrValue.toUnicode());
+			if (CT_CURRENT_ATTRIB_IO::AskStrField(ITH_SAVE_OPTION_PREPROCESSOR, FieldStrValue) != IO_OK) break;
+			if (FieldStrValue.IsEmpty()) { break; }
+			//NodeAttributeSet.Add(TEXT("SaveOptionPreprocessor"));
+			SceneGraphDescription.Add(FieldStrValue.toUnicode());
+			break;
+
+		case CT_ATTRIB_ORIGINAL_ID:
+			GetAttributeValue(AttributeType, ITH_ORIGINAL_ID_VALUE, FieldValue);
+			SceneGraphDescription.Add(TEXT("OriginalId"));
+			SceneGraphDescription.Add(FieldValue);
+			break;
+
+		case CT_ATTRIB_ORIGINAL_ID_STRING:
+			if (CT_CURRENT_ATTRIB_IO::AskStrField(ITH_ORIGINAL_ID_VALUE_STRING, FieldStrValue) != IO_OK) break;
+			if (FieldStrValue.IsEmpty()) { break; }
+			SceneGraphDescription.Add(TEXT("OriginalIdStr"));
+			SceneGraphDescription.Add(FieldStrValue.toUnicode());
+			break;
+
+		case CT_ATTRIB_COLOR_RGB_DOUBLE:
+			if (CT_CURRENT_ATTRIB_IO::AskDblField(ITH_ATTRIB_COLOR_R_DOUBLE, FieldDoubleValue0) != IO_OK) break;
+			if (CT_CURRENT_ATTRIB_IO::AskDblField(ITH_ATTRIB_COLOR_G_DOUBLE, FieldDoubleValue1) != IO_OK) break;
+			if (CT_CURRENT_ATTRIB_IO::AskDblField(ITH_ATTRIB_COLOR_B_DOUBLE, FieldDoubleValue2) != IO_OK) break;
+			FieldValue = FString::Printf(TEXT("%lf"), FieldDoubleValue0) + TEXT(", ") + FString::Printf(TEXT("%lf"), FieldDoubleValue1) + TEXT(", ") + FString::Printf(TEXT("%lf"), FieldDoubleValue2);
+			//OutMetaDataSet.Add(TEXT("ColorRGBDouble"))
+			//OutMetaDataSet.Add(FieldValue);
+			break;
+
+		case CT_ATTRIB_REVERSE_COLORID:
+			break;
+		case CT_ATTRIB_INITIAL_FILTER:
+			break;
+		case CT_ATTRIB_ORIGINAL_SURF:
+			break;
+		case CT_ATTRIB_LINKMANAGER_BRANCH_FACE:
+			//ITH_LINKMANAGER_BRANCH_FACE_PART_ID, ITH_LINKMANAGER_BRANCH_FACE_SOLID_ID, ITH_LINKMANAGER_BRANCH_FACE_BRANCH_ID, ITH_LINKMANAGER_BRANCH_FACE_FACE_ID
+			break;
+		case CT_ATTRIB_LINKMANAGER_PMI:
+			//ITH_LINKMANAGER_PMI_ID, ITH_LINKMANAGER_PMI_SENSE, ITH_LINKMANAGER_PMI_INSTANCE
+			break;
+		case CT_ATTRIB_NULL:
+			//ITH_NULL_TITLE
+			break;
+		case CT_ATTRIB_MEASURE_VALIDATION_ATTRIBUTE:
+			//ITH_MEASURE_VALIDATION_NAME
+			//ITH_MEASURE_VALIDATION_UNIT
+			//ITH_MEASURE_VALIDATION_VALUE
+			break;
+
+		case CT_ATTRIB_INTEGER_VALIDATION_ATTRIBUTE:
+			if (CT_CURRENT_ATTRIB_IO::AskStrField(ITH_INTEGER_VALIDATION_NAME, FieldName) != IO_OK) break;
+			if (CT_CURRENT_ATTRIB_IO::AskIntField(ITH_INTEGER_VALIDATION_VALUE, FieldIntValue) != IO_OK) break;
+			SceneGraphDescription.Add(FieldName.toUnicode());
+			SceneGraphDescription.Add(FString::FromInt(FieldIntValue));
+			break;
+
+		case CT_ATTRIB_DOUBLE_VALIDATION_ATTRIBUTE:
+			if (CT_CURRENT_ATTRIB_IO::AskStrField(ITH_DOUBLE_VALIDATION_NAME, FieldName) != IO_OK) break;
+			if (CT_CURRENT_ATTRIB_IO::AskDblField(ITH_DOUBLE_VALIDATION_VALUE, FieldDoubleValue0) != IO_OK) break;
+			SceneGraphDescription.Add(FieldName.toUnicode());
+			SceneGraphDescription.Add(FString::Printf(TEXT("%lf"), FieldDoubleValue0));
+			break;
+
+		case CT_ATTRIB_STRING_VALIDATION_ATTRIBUTE:
+			if (CT_CURRENT_ATTRIB_IO::AskStrField(ITH_STRING_VALIDATION_NAME, FieldName) != IO_OK) break;
+			if (CT_CURRENT_ATTRIB_IO::AskStrField(ITH_STRING_VALIDATION_VALUE, FieldStrValue) != IO_OK) break;
+			if (FieldStrValue.IsEmpty()) { break; }
+			SceneGraphDescription.Add(FieldName.toUnicode());
+			SceneGraphDescription.Add(FieldStrValue.toUnicode());
+			break;
+
+		case CT_ATTRIB_BOUNDING_BOX:
+			//ITH_BOUNDING_BOX_XMIN, ITH_BOUNDING_BOX_YMIN, ITH_BOUNDING_BOX_ZMIN, ITH_BOUNDING_BOX_XMAX, ITH_BOUNDING_BOX_YMAX, ITH_BOUNDING_BOX_ZMAX
+			break;
+		case CT_ATTRIB_DATABASE:
+			//ITH_DATABASE_NAME, ITH_DATABASE_TABLE
+			break;
+		case CT_ATTRIB_CURVE_FONT:
+			//ITH_ATTRIB_CURVE_FONT_VALUE
+			break;
+		case CT_ATTRIB_CURVE_WEIGHT:
+			//ITH_ATTRIB_CURVE_WEIGHT_VALUE
+			break;
+		case CT_ATTRIB_COMPARE_TOPO:
+			break;
+		case CT_ATTRIB_MONIKER_GUID_TABLE:
+			//ITH_MONIKER_GUID_TABLE_NAME, ITH_MONIKER_GUID_TABLE_GUID, ITH_MONIKER_GUID_TABLE_INDEX, ITH_MONIKER_GUID_TABLE_APPLICATION_NAME
+			break;
+		case CT_ATTRIB_MONIKER_DATA:
+			//ITH_MONIKER_DATA_NAME, ITH_MONIKER_DATA_TABLE_INDEX, ITH_MONIKER_DATA_ENTITY_ID, ITH_MONIKER_DATA_LABEL, ITH_MONIKER_DATA_BODY_GUID
+			break;
+		case CT_ATTRIB_MONIKER_BODY_ID:
+			//ITH_MONIKER_BODY_ID_NAME, ITH_MONIKER_BODY_ID_ENTITY_ID, ITH_MONIKER_BODY_ID_APPLICATION_NAME
+			break;
+		case CT_ATTRIB_NO_INSTANCE:
+			break;
+
+		case CT_ATTRIB_GROUPNAME:
+			if (CT_CURRENT_ATTRIB_IO::AskStrField(ITH_GROUPNAME_VALUE, FieldStrValue) != IO_OK) break;
+			if (FieldStrValue.IsEmpty()) { break; }
+			SceneGraphDescription.Add(TEXT("GroupName"));
+			SceneGraphDescription.Add(FieldStrValue.toUnicode());
+			break;
+
+		case CT_ATTRIB_ANALYZE_ID:
+			//ITH_ANALYZE_OBJECT_ID
+			break;
+		case CT_ATTRIB_ANALYZER_DISPLAY_MODE:
+			//ITH_ANALYZER_DISPLAY_MODE_VALUE
+			break;
+		case CT_ATTRIB_ANIMATION_ID:
+			//ITH_ANIMATION_ID_VALUE
+			break;
+		case CT_ATTRIB_PROJECTED_SURFACE_ID:
+			//ITH_PROJECTED_SURFACE_ID_VALUE
+			break;
+		case CT_ATTRIB_ANALYZE_LINK:
+			//ITH_ANALYZE_LINK_ID
+			break;
+		case CT_ATTRIB_TOPO_EVENT_ID:
+			//ITH_TOPO_EVENT_ID_VALUE
+			break;
+		case CT_ATTRIB_ADDITIVE_MANUFACTURING:
+			//ITH_AM_MODE, ITH_AM_ROUGHNESS_MIN, ITH_AM_SECTION_DIST, ITH_AM_FUNCTIONAL_TYPE, ITH_AM_LATTICE_TYPE
+			break;
+		case CT_ATTRIB_MOLDING_RESULT:
+			//ITH_MOLDING_MODE, ITH_MOLDING_MEASURE_MIN, ITH_MOLDING_MEASURE_MAX
+			break;
+		case CT_ATTRIB_AMF_ID:
+			//ITH_AMF_ID_VALUE, ITH_AMF_TYPE
+			break;
+		case CT_ATTRIB_PARAMETER_LINK:
+			//ITH_PARAMETER_LINK_ID, ITH_PARAMETER_LINK_PARENT_ID, ITH_PARAMETER_LINK_SENSE
+			break;
+		default:
+			break;
+		}
+	}
+
+	AttributeNameLigne = "M " + FString::FromInt(SceneGraphDescription.Num() - AttributIndex);
+	SceneGraphDescription[AttributIndex-1] = *AttributeNameLigne;
+}
+
+uint32 GetSize(CT_TESS_DATA_TYPE type)
+{
+	switch (type)
+	{
+	case CT_TESS_USE_DEFAULT:
+		return sizeof(uint32);
+	case CT_TESS_UBYTE:
+		return sizeof(uint8_t);
+	case CT_TESS_BYTE:
+		return sizeof(int8_t);
+	case CT_TESS_USHORT:
+		return sizeof(int16_t);
+	case CT_TESS_SHORT:
+		return sizeof(uint16_t);
+	case CT_TESS_UINT:
+		return sizeof(uint32);
+	case CT_TESS_INT:
+		return sizeof(int32);
+	case CT_TESS_ULONG:
+		return sizeof(uint64);
+	case CT_TESS_LONG:
+		return sizeof(int64);
+	case CT_TESS_FLOAT:
+		return sizeof(float);
+	case CT_TESS_DOUBLE:
+		return sizeof(double);
+	}
+	return 0;
+}
+
+}
+
+#endif // CAD_INTERFACE
diff --git a/Engine/Plugins/Enterprise/DatasmithCADImporter/Source/CADInterfaces/Private/CoreTechTypes.cpp b/Engine/Plugins/Enterprise/DatasmithCADImporter/Source/CADInterfaces/Private/CoreTechTypes.cpp
new file mode 100644
--- /dev/null
+++ b/Engine/Plugins/Enterprise/DatasmithCADImporter/Source/CADInterfaces/Private/CoreTechTypes.cpp
@@ -0,0 +1,160 @@
+// Copyright 1998-2019 Epic Games, Inc. All Rights Reserved.
+#include "CoreTechTypes.h"
+
+#ifdef CAD_INTERFACE
+
+const char* CoreTechLicenseKey =
+#include "CoreTech.lic"
+;
+
+namespace CADLibrary
+{
+	// Note: CTKIO_* functions are not functionally useful.
+	// This wrapping allows a correct profiling of the CT API.
+
+	CT_IO_ERROR CTKIO_InitializeKernel(double Unit)
+	{
+		CT_STR appName = CoreTechLicenseKey;
+		return CT_KERNEL_IO::InitializeKernel(appName, Unit, 0.00001 / Unit);
+	}
+
+	CT_IO_ERROR CTKIO_ShutdownKernel()
+	{
+		return CT_KERNEL_IO::ShutdownKernel(); // ignorable. Just in case CT was not previously stopped
+	}
+
+	CT_IO_ERROR CTKIO_UnloadModel()
+	{
+		return CT_KERNEL_IO::UnloadModel();
+	}
+
+	CT_IO_ERROR CTKIO_AskNbObjectsType(CT_UINT32& object_count, CT_OBJECT_TYPE type)
+	{
+		return CT_KERNEL_IO::AskNbObjectsType(object_count, type);
+	}
+
+	CT_IO_ERROR CTKIO_AskMainObject(CT_OBJECT_ID& OutMainObjectId)
+	{
+		return CT_KERNEL_IO::AskMainObject(OutMainObjectId);
+	}
+
+	CT_IO_ERROR CTKIO_ChangeTolerance(double Precision)
+	{
+		return CT_KERNEL_IO::ChangeTolerance(Precision);
+	}
+
+	CT_IO_ERROR CTKIO_AskUnit(double Unit)
+	{
+		return CT_KERNEL_IO::AskUnit(Unit);
+	}
+
+	CT_IO_ERROR CTKIO_ChangeTesselationParameters(double max_sag, double max_length, double max_angle, CT_LOGICAL high_quality, CT_TESS_DATA_TYPE vertex_type, CT_TESS_DATA_TYPE normal_type, CT_TESS_DATA_TYPE texture_type)
+	{
+		return CT_KERNEL_IO::ChangeTesselationParameters(max_sag, max_length, max_angle, high_quality, vertex_type, normal_type, texture_type);
+	}
+
+	CT_IO_ERROR CTKIO_LoadFile(const TCHAR* FileName, CT_OBJECT_ID& MainObject, CT_FLAGS LoadFlags, uint32 Lod, const TCHAR* StringOption)
+	{
+		return CT_KERNEL_IO::LoadFile(FileName, MainObject, LoadFlags, Lod, StringOption);
+	}
+
+	CT_IO_ERROR CTKIO_SaveFile(CT_LIST_IO& objects_list_to_save, const TCHAR* file_name, const TCHAR* format, const CT_OBJECT_ID coordsystem)
+	{
+		return CT_KERNEL_IO::SaveFile(objects_list_to_save, file_name, format, coordsystem);
+	}
+
+	CT_IO_ERROR CTKIO_AskTolerance(double tolerance)
+	{
+		return CT_KERNEL_IO::AskTolerance(tolerance);
+	}
+
+	CT_IO_ERROR CTKIO_AskModelSize(double model_size)
+	{
+		return CT_KERNEL_IO::AskModelSize(model_size);
+	}
+
+	CT_IO_ERROR CTKIO_AskTesselationParameters(double max_face_sag, double max_face_length, double max_face_angle, double max_curve_sag, double max_curve_length, double max_curve_angle, CT_LOGICAL& high_quality, CT_TESS_DATA_TYPE& vertex_type, CT_TESS_DATA_TYPE& normal_type, CT_TESS_DATA_TYPE& texture_type)
+	{
+		return CT_KERNEL_IO::AskTesselationParameters(max_face_sag, max_face_length, max_face_angle, max_curve_sag, max_curve_length, max_curve_angle, high_quality, vertex_type, normal_type, texture_type);
+	}
+
+	CT_IO_ERROR Repair(CT_OBJECT_ID MainObjectID, EStitchingTechnique StitchingTechnique)
+	{
+		switch (StitchingTechnique)
+		{
+		case EStitchingTechnique::StitchingSew:
+		{
+			// Sew disconnected faces back together
+			CT_UINT32 SewingToleranceFactor = 100; // factor of 'tolerance' value
+			return CT_REPAIR_IO::Sew(MainObjectID, SewingToleranceFactor, CT_SEW_CREATE_BODIES_BY_TOPOLOGY);
+		}
+
+		case EStitchingTechnique::StitchingHeal:
+		{
+			// Heal disconnected faces back together
+			CT_UINT32 StitchingToleranceFactor = 100; // topo only if deformation=0, factor of 'tolerance' value
+			CT_LOGICAL   OnlyVisible = CT_TRUE;
+			CT_LOGICAL   MergeBeforeHealing = CT_FALSE;
+			CT_LOGICAL   RemoveTinyObjects = CT_TRUE;
+			return CT_REPAIR_IO::Heal(MainObjectID, StitchingToleranceFactor, OnlyVisible, MergeBeforeHealing, RemoveTinyObjects);
+		}
+		}
+		return CT_IO_ERROR::IO_OK;
+	}
+
+	CT_IO_ERROR SetCoreTechTessellationState(const FImportParameters& ImportParams)
+	{
+		CT_IO_ERROR Result = IO_OK;
+		CT_DOUBLE CurrentUnit = 0.001;
+
+		CT_TESS_DATA_TYPE VertexType = CT_TESS_DOUBLE;
+		CT_TESS_DATA_TYPE NormalType = CT_TESS_FLOAT;
+		CT_TESS_DATA_TYPE UVType = CT_TESS_DOUBLE;
+
+		// convert max edge length to model units
+		double MaxEdgeLength_modelUnit = FLT_MAX; // default value is huge, as 0.0 causes bugs in KernelIO...
+		if (ImportParams.MaxEdgeLength > SMALL_NUMBER)
+		{
+			MaxEdgeLength_modelUnit = ImportParams.MaxEdgeLength / ImportParams.ScaleFactor;
+		}
+
+		// Apply retrieved tessellation parameters to CoreTech tessellation engine
+		static CT_LOGICAL high_quality = CT_FALSE;
+		Result = CTKIO_ChangeTesselationParameters(ImportParams.ChordTolerance / ImportParams.ScaleFactor, MaxEdgeLength_modelUnit, ImportParams.MaxNormalAngle, high_quality, VertexType, NormalType, UVType);
+
+		return Result;
+	}
+
+
+	const TCHAR* CoreTechSessionBase::Owner = nullptr;
+
+	CoreTechSessionBase::CoreTechSessionBase(const TCHAR* InOwner, double Unit)
+	{
+		// Lib init
+		ensureAlways(InOwner != nullptr); // arg must be valid
+		ensureAlways(Owner == nullptr); // static Owner must be unset
+
+		CT_IO_ERROR e = CTKIO_ShutdownKernel(); // ignorable. Just in case CT was not previously stopped
+		CT_IO_ERROR Result = CTKIO_InitializeKernel(Unit);
+		Owner = (Result == IO_OK) ? InOwner : nullptr;
+
+		// Create a main object to hold the BRep
+		CT_OBJECT_ID nullParent = 0;
+		Result = CT_COMPONENT_IO::Create(MainObjectId, nullParent);
+		Result = CTKIO_AskMainObject(MainObjectId); // not required, just a validation.
+	}
+
+	CoreTechSessionBase::~CoreTechSessionBase()
+	{
+		if (Owner != nullptr)
+		{
+			CheckedCTError Result;
+			Result = CTKIO_UnloadModel();
+			Result = CTKIO_ShutdownKernel();
+			Owner = nullptr;
+		}
+	}
+} // ns CADLibrary
+
+#endif // CAD_INTERFACE
+
diff --git a/Engine/Plugins/Enterprise/DatasmithCADImporter/Source/CADInterfaces/Public/CTSession.h b/Engine/Plugins/Enterprise/DatasmithCADImporter/Source/CADInterfaces/Public/CTSession.h
new file mode 100644
--- /dev/null
+++ b/Engine/Plugins/Enterprise/DatasmithCADImporter/Source/CADInterfaces/Public/CTSession.h
@@ -0,0 +1,71 @@
+// Copyright 1998-2019 Epic Games, Inc. All Rights Reserved.
+
+#pragma once
+
+#ifdef CAD_INTERFACE
+
+#include "CoreMinimal.h"
+#include "CoreTechTypes.h"
+#include "CADOptions.h"
+
+struct FMeshDescription;
+
+// Fill data array with a debug value (eg -1) to help debugging
+#define MARK_UNINITIALIZED_MEMORY 0
+
+namespace CADLibrary
+{
+class CADINTERFACES_API CTSession : public CoreTechSessionBase
+{
+public:
+	/**
+	 * Make sure CT is initialized, and a main object is ready.
+	 * Handle input file unit and an output unit
+	 * @param InOwner:        text that describe the owner of the session (helps to fix initialization issues)
+	 * @param FileMetricUnit: number of meters per file unit.
+	 * eg. For a file in inches, arg should be 0.0254
+	 */
+	CTSession(const TCHAR* InOwner, double InFileMetricUnit, double InScaleFactor)
+		: CoreTechSessionBase(InOwner, InFileMetricUnit)
+	{
+		ImportParams.ScaleFactor = InScaleFactor;
+		ImportParams.MetricUnit = InFileMetricUnit;
+	}
+
+	void ClearData();
+
+	CheckedCTError SaveBrep(const FString& FilePath);
+	CheckedCTError TopoFixes();
+
+	/**
+	 * @param InScaleFactor : use to scale meshing from Kernel-IO
+	 */
+	void SetScaleFactor(double InScaleFactor)
+	{
+		ImportParams.ScaleFactor = InScaleFactor;
+	}
+
+	/**
+	 * Set Import parameters,
+	 * Tack care to set scale factor before because import parameters will be scale according to scale factor
+	 * @param ChordTolerance : SAG	
+	 * @param MaxEdgeLength : max length of element's edge
+	 * @param NormalTolerance : Angle between two adjacent triangles
+	 * @param StitchingTechnique : CAD topology correction technique
+	 */
+	void SetImportParameters(float ChordTolerance, float MaxEdgeLength, float NormalTolerance, CADLibrary::EStitchingTechnique StitchingTechnique);
+	
+	CADLibrary::FImportParameters& GetImportParameters()
+	{
+		return ImportParams;
+	}
+
+protected:
+	CADLibrary::FImportParameters ImportParams;
+	static TWeakPtr<CTSession> SharedSession;
+};
+
+}
+
+#endif // CAD_INTERFACE
+
diff --git a/Engine/Plugins/Enterprise/DatasmithCADImporter/Source/CADInterfaces/Public/CoreTechFileParser.h b/Engine/Plugins/Enterprise/DatasmithCADImporter/Source/CADInterfaces/Public/CoreTechFileParser.h
new file mode 100644
--- /dev/null
+++ b/Engine/Plugins/Enterprise/DatasmithCADImporter/Source/CADInterfaces/Public/CoreTechFileParser.h
@@ -0,0 +1,158 @@
+// Copyright 1998-2019 Epic Games, Inc. All Rights Reserved.
+#pragma once
+
+#include "CoreMinimal.h"
+
+#include "CADData.h"
+#include "CADOptions.h"
+#include "Containers/Map.h"
+#include "Containers/Queue.h"
+#include "Misc/Paths.h"
+
+#define COLORSETLINE    3
+#define MATERIALSETLINE 4
+#define EXTERNALREFLINE 7
+#define MAPCTIDLINE     8
+
+#ifdef CAD_INTERFACE
+#pragma warning(push)
+#pragma warning(disable:4996) // unsafe sprintf
+#pragma warning(disable:4828) // illegal character
+#include "kernel_io/kernel_io.h"
+#include "kernel_io/kernel_io_error.h"
+#include "kernel_io/kernel_io_type.h"
+#include "kernel_io/list_io/list_io.h"
+#include "kernel_io/material_io/material_io.h"
+#include "kernel_io/object_io/asm_io/component_io/component_io.h"
+#include "kernel_io/object_io/asm_io/instance_io/instance_io.h"
+#include "kernel_io/object_io/geom_io/curve_io/cnurbs_io/cnurbs_io.h"
+#include "kernel_io/object_io/geom_io/surface_io/snurbs_io/snurbs_io.h"
+#include "kernel_io/object_io/geom_io/surface_io/surface_io.h"
+#include "kernel_io/object_io/history_io/hnode_io/hbranch_io/hbranch_io.h"
+#include "kernel_io/object_io/history_io/hnode_io/hleaf_io/hleaf_io.h"
+#include "kernel_io/object_io/history_io/solid_io/solid_io.h"
+#include "kernel_io/object_io/topo_io/body_io/body_io.h"
+#include "kernel_io/object_io/topo_io/coedge_io/coedge_io.h"
+#include "kernel_io/object_io/topo_io/face_io/face_io.h"
+#include "kernel_io/object_io/topo_io/loop_io/loop_io.h"
+#include "kernel_io/repair_io/repair_io.h"
+#pragma warning(pop)
+#endif CAD_INTERFACE
+
+
+struct FFileStatData;
+
+enum EProcessState  // #ueent_CAD
+{
+	Unknown = 0x00,
+	UnTreated = 0x01,
+	Running = 0x02,
+	ProcessOk = 0x04,
+	ProcessFailed = 0x08,
+	FileNotFound = 0x10
+};
+
+#ifdef CAD_INTERFACE
+namespace CADLibrary
+{
+CADINTERFACES_API void GetColor(uint32 ColorHash, FColor& OutMaterial);
+CADINTERFACES_API bool GetMaterial(uint32 MaterialID, FCADMaterial& OutMaterial);
+
+CADINTERFACES_API int32 BuildColorHash(uint32 ColorHId);
+CADINTERFACES_API int32 BuildMaterialHash(uint32 MaterialId);
+
+CADINTERFACES_API uint32 GetFaceTessellation(CT_OBJECT_ID FaceID, TArray<FTessellationData>& FaceTessellationSet, int32& OutRawDataSize, const float ScaleFactor);
+CADINTERFACES_API uint32 GetBodiesTessellations(TArray<CT_OBJECT_ID>& BodySet, TArray<FTessellationData>& FaceTessellations, TMap<uint32, uint32>& MaterialIdToMaterialHashMap, int32& OutRawDataSize, const float ScaleFactor);
+CADINTERFACES_API uint32 GetBodiesFaceSetNum(TArray<CT_OBJECT_ID>& BodySet);
+CADINTERFACES_API void GetCTObjectDisplayDataIds(CT_OBJECT_ID ObjectID, FObjectDisplayDataId& Material);
+CADINTERFACES_API void GetBodiesMaterials(TArray<CT_OBJECT_ID>& BodySet, TMap<uint32, uint32>& MaterialIdToHash, bool bPreferPartData);
+
+uint32 GetFileHash(const FString& FileName, const FFileStatData& FileStatData, const FString& Config);
+uint32 GetGeomFileHash(const uint32 InSGHash, CADLibrary::FImportParameters& ImportParam);
+
+/**
+ * According to the preference (bPreferPartData), select the best material to use
+ */
+CADINTERFACES_API void SetFaceMainMaterial(FObjectDisplayDataId& InFaceMaterial, FObjectDisplayDataId& InBodyMaterial, TMap<uint32, uint32>& MaterialIdToMaterialHashMap, FTessellationData& OutFaceTessellations);
+
+uint32 GetSize(CT_TESS_DATA_TYPE type);
+
+
+class CADINTERFACES_API FCoreTechFileParser
+{
+public:
+	FCoreTechFileParser(const FString& InCTFullPath, const FString& InCachePath, const FImportParameters& ImportParams, bool bPreferBodyData, bool bScaleUVMap);
+
+	double GetMetricUnit() const { return 0.01; }
+
+	EProcessState ProcessFile();
+
+	const TSet<FString>& GetExternalRefSet()
+	{
+		return ExternalRefSet;
+	}
+
+	const FString& GetSceneGraphFile()
+	{
+		return SceneGraphFile;
+	}
+	const FString& GetMeshFile()
+	{
+		return MeshFile;
+	}
+	const FString& GetFileName()
+	{
+		return CADFile;
+	}
+
+
+
+private:
+	EProcessState ReadFileWithKernelIO();
+
+	bool ReadNode(CT_OBJECT_ID NodeId);
+	bool ReadInstance(CT_OBJECT_ID NodeId);
+	bool ReadComponent(CT_OBJECT_ID NodeId);
+	bool ReadUnloadedComponent(CT_OBJECT_ID NodeId);
+	bool ReadBody(CT_OBJECT_ID NodeId);
+
+	uint32 GetMaterialNum();
+	void ReadColor();
+	void ReadMaterial();
+
+	void ReadNodeMetaDatas(CT_OBJECT_ID NodeId);
+
+	CT_FLAGS SetCoreTechImportOption(const FString& MainFileExt);
+	void GetAttributeValue(CT_ATTRIB_TYPE attrib_type, int ith_field, FString& value);
+
+	void ExportFileSceneGraph();
+
+protected:
+	const FString& CachePath;
+
+	FString FullPath;
+	FString CADFile;
+
+	FString FileConfiguration;
+	FString NodeConfiguration;
+
+	FString SceneGraphFile;
+	FString RawDataGeom;
+	FString MeshFile;
+
+	TSet<FString> ExternalRefSet;
+
+	TArray<FString> SceneGraphDescription;
+	TMap<uint32, uint32> CTIdToRawLineMap;
+
+	TMap<uint32, uint32> MaterialIdToMaterialHashMap;
+
+	bool bNeedSaveCTFile;
+	bool bPreferBodyData;
+	bool bScaleUVMap;
+
+	const FImportParameters& ImportParameters;
+};
+
+}
+#endif // CAD_INTERFACE
diff --git a/Engine/Plugins/Enterprise/DatasmithCADImporter/Source/CADInterfaces/Public/CoreTechTypes.h b/Engine/Plugins/Enterprise/DatasmithCADImporter/Source/CADInterfaces/Public/CoreTechTypes.h
new file mode 100644
--- /dev/null
+++ b/Engine/Plugins/Enterprise/DatasmithCADImporter/Source/CADInterfaces/Public/CoreTechTypes.h
@@ -0,0 +1,151 @@
+// Copyright 1998-2019 Epic Games, Inc. All Rights Reserved.
+#pragma once
+
+
+#ifdef CAD_INTERFACE
+
+#include "CADOptions.h"
+#include "Containers/Array.h"
+#include "Math/Vector.h"
+#include "Math/Vector2D.h"
+
+#pragma warning(push)
+#pragma warning(disable:4996) // unsafe sprintf
+#pragma warning(disable:4828) // illegal character
+#include "kernel_io/kernel_io.h"
+#include "kernel_io/kernel_io_error.h"
+#include "kernel_io/kernel_io_type.h"
+#include "kernel_io/list_io/list_io.h"
+#include "kernel_io/material_io/material_io.h"
+#include "kernel_io/object_io/asm_io/component_io/component_io.h"
+#include "kernel_io/object_io/asm_io/instance_io/instance_io.h"
+#include "kernel_io/object_io/geom_io/curve_io/cnurbs_io/cnurbs_io.h"
+#include "kernel_io/object_io/geom_io/surface_io/snurbs_io/snurbs_io.h"
+#include "kernel_io/object_io/geom_io/surface_io/surface_io.h"
+#include "kernel_io/object_io/history_io/hnode_io/hbranch_io/hbranch_io.h"
+#include "kernel_io/object_io/history_io/hnode_io/hleaf_io/hleaf_io.h"
+#include "kernel_io/object_io/history_io/solid_io/solid_io.h"
+#include "kernel_io/object_io/topo_io/body_io/body_io.h"
+#include "kernel_io/object_io/topo_io/coedge_io/coedge_io.h"
+#include "kernel_io/object_io/topo_io/face_io/face_io.h"
+#include "kernel_io/object_io/topo_io/loop_io/loop_io.h"
+#include "kernel_io/repair_io/repair_io.h"
+#pragma warning(pop)
+
+
+namespace CADLibrary
+{
+
+	struct CTMesh
+	{
+		TArray<uint32> Materials; //GPure Material Id
+		TArray<uint32> MaterialUUIDs; //Material Hash from color value
+		TArray<FVector> Vertices;
+		TArray<FVector> Normals;
+		TArray<FVector2D> TexCoords;
+		TArray<int32_t> Indices;
+		TArray<uint32> TriangleMaterials;
+	};
+
+	/**
+	 * Mandatory: Kernel_IO has to be initialize with the final MetricUnit, set function doesn`t work
+	 * @param MetricUnit: Length unit express in meter
+	 */
+	CADINTERFACES_API CT_IO_ERROR CTKIO_InitializeKernel(double MetricUnit);
+	CADINTERFACES_API CT_IO_ERROR CTKIO_ShutdownKernel();
+
+	CADINTERFACES_API CT_IO_ERROR CTKIO_UnloadModel();
+	CADINTERFACES_API CT_IO_ERROR CTKIO_AskMainObject(CT_OBJECT_ID& OutMainObjectId);
+	CADINTERFACES_API CT_IO_ERROR CTKIO_ChangeTolerance(double Precision);
+	CADINTERFACES_API CT_IO_ERROR CTKIO_AskUnit(double Unit);
+	CADINTERFACES_API CT_IO_ERROR CTKIO_AskTolerance(double tolerance);
+	CADINTERFACES_API CT_IO_ERROR CTKIO_AskModelSize(double model_size);
+	CADINTERFACES_API CT_IO_ERROR CTKIO_AskTesselationParameters(
+		double      max_face_sag,      
+		double      max_face_length,   
+		double      max_face_angle,    
+		double      max_curve_sag,     
+		double      max_curve_length,  
+		double      max_curve_angle,   
+		CT_LOGICAL&        high_quality,
+		CT_TESS_DATA_TYPE& vertex_type, 
+		CT_TESS_DATA_TYPE& normal_type, 
+		CT_TESS_DATA_TYPE& texture_type 
+	);
+
+	CADINTERFACES_API CT_IO_ERROR CTKIO_AskNbObjectsType
+	(
+		CT_UINT32&     object_count,            /*!< [out] Number of objects found. Type \c CT_UINT. */
+		CT_OBJECT_TYPE type = CT_OBJECT_ALL_TYPE /*!< [in] [optional] Searched type. Type \c CT_OBJECT_TYPE. */
+	);
+
+	CADINTERFACES_API CT_IO_ERROR CTKIO_ChangeTesselationParameters(
+		double     max_sag,                           
+		double     max_length,                        
+		double     max_angle,                         
+		CT_LOGICAL        high_quality = CT_FALSE,     
+		CT_TESS_DATA_TYPE vertex_type = CT_TESS_USE_DEFAULT, 
+		CT_TESS_DATA_TYPE normal_type = CT_TESS_USE_DEFAULT, 
+		CT_TESS_DATA_TYPE texture_type = CT_TESS_USE_DEFAULT 
+	);
+
+	CADINTERFACES_API CT_IO_ERROR CTKIO_LoadFile
+	(
+		const TCHAR*  file_name,                             
+		CT_OBJECT_ID& main_object,                           
+		CT_FLAGS      load_flags = CT_LOAD_FLAGS_USE_DEFAULT,
+		uint32        lod = 0,                               
+		const TCHAR*  string_option = TEXT("")               
+	);
+
+	CADINTERFACES_API CT_IO_ERROR CTKIO_SaveFile
+	(
+		CT_LIST_IO&        objects_list_to_save,
+		const TCHAR*       file_name,           
+		const TCHAR*       format,              
+		const CT_OBJECT_ID coordsystem = 0      
+	);
+
+	CADINTERFACES_API CT_IO_ERROR Repair(CT_OBJECT_ID MainObjectID, EStitchingTechnique StitchingTechnique);
+	CADINTERFACES_API CT_IO_ERROR SetCoreTechTessellationState(const FImportParameters& ImportParams);
+
+
+	struct CADINTERFACES_API CheckedCTError
+	{
+		CheckedCTError(CT_IO_ERROR in_e = CT_IO_ERROR::IO_OK, bool in_otherError = false) : e(in_e), OtherError(in_otherError) { Validate(); }
+		CheckedCTError& operator=(CT_IO_ERROR in_e) { e = in_e; Validate(); return *this; }
+		CheckedCTError& operator=(const CheckedCTError& in_e) { new (this) CheckedCTError(in_e.e, in_e.OtherError); return *this; }
+		operator bool() const { return e == CT_IO_ERROR::IO_OK && OtherError == false; }
+		operator CT_IO_ERROR() const { return e; }
+		void RaiseOtherError(const char* msg) { OtherErrorMsg = msg; OtherError = true; Validate(); }
+		CT_IO_ERROR e;
+
+	private:
+		void Validate();
+		bool OtherError;
+		const char* OtherErrorMsg = nullptr;
+	};
+
+	class CADINTERFACES_API CoreTechSessionBase
+	{
+	public:
+		/**
+		 * Make sure CT is initialized, and a main object is ready.
+		 * Handle input file unit and an output unit
+		 * @param FileMetricUnit number of meters per file unit.
+		 * eg. For a file in inches, arg should be 0.0254
+		 */
+		CoreTechSessionBase(const TCHAR* Owner, double Unit);
+		bool IsSessionValid() { return Owner != nullptr && MainObjectId != 0; }
+		virtual ~CoreTechSessionBase();
+
+	protected:
+		CT_OBJECT_ID MainObjectId;
+
+	private:
+		static const TCHAR* Owner;
+	};
+}
+
+#endif // CAD_INTERFACE
+
diff --git a/Engine/Plugins/Enterprise/DatasmithCADImporter/Source/DatasmithCADTranslator/Private/CoreTechParserMT.cpp b/Engine/Plugins/Enterprise/DatasmithCADImporter/Source/DatasmithCADTranslator/Private/CoreTechParserMT.cpp
deleted file mode 100644
diff --git a/Engine/Plugins/Enterprise/DatasmithCADImporter/Source/DatasmithCADTranslator/Private/CoreTechParserMT.h b/Engine/Plugins/Enterprise/DatasmithCADImporter/Source/DatasmithCADTranslator/Private/CoreTechParserMT.h
deleted file mode 100644
diff --git a/Engine/Plugins/Enterprise/DatasmithCADImporter/Source/DatasmithCADTranslator/Private/DatasmithCADTranslator.h b/Engine/Plugins/Enterprise/DatasmithCADImporter/Source/DatasmithCADTranslator/Private/DatasmithCADTranslator.h
--- a/Engine/Plugins/Enterprise/DatasmithCADImporter/Source/DatasmithCADTranslator/Private/DatasmithCADTranslator.h
+++ b/Engine/Plugins/Enterprise/DatasmithCADImporter/Source/DatasmithCADTranslator/Private/DatasmithCADTranslator.h
@@ -4,42 +4,39 @@
 
 #include "CoreMinimal.h"
 
-#include "CADOptions.h"
 #include "DatasmithCoreTechTranslator.h"
+
+#ifdef CAD_LIBRARY
+#include "CADOptions.h"
 #include "DatasmithDispatcher.h"
 #include "DatasmithMeshBuilder.h"
 #include "DatasmithSceneGraphBuilder.h"
+#endif // CAD_LIBRARY
+
 #include "UObject/ObjectMacros.h"
 
 class IDatasmithMeshElement;
-class FCoreTechParser;
 
 class FDatasmithCADTranslator : public FDatasmithCoreTechTranslator
 {
 public:
-#ifndef CAD_LIBRARY
-	// Begin IDatasmithTranslator overrides
 	virtual FName GetFName() const override { return "DatasmithCADTranslator"; };
-	// End IDatasmithTranslator overrides
+
+#ifndef CAD_LIBRARY
+	virtual void Initialize(FDatasmithTranslatorCapabilities& OutCapabilities) override { OutCapabilities.bIsEnabled = false; }
 #else
 	FDatasmithCADTranslator();
-
-	// Begin IDatasmithTranslator overrides
-	virtual FName GetFName() const override { return "DatasmithCADTranslator"; };
-
 	virtual void Initialize(FDatasmithTranslatorCapabilities& OutCapabilities) override;
 
 	virtual bool IsSourceSupported(const FDatasmithSceneSource& Source) override;
 
 	virtual bool LoadScene(TSharedRef<IDatasmithScene> OutScene) override;
+
 	virtual void UnloadScene() override;
 
 	virtual bool LoadStaticMesh(const TSharedRef<IDatasmithMeshElement> MeshElement, FDatasmithMeshElementPayload& OutMeshPayload) override;
-	// End IDatasmithTranslator overrides
 
-	// Begin ADatasmithCoreTechTranslator overrides
 	virtual void SetSceneImportOptions(TArray<TStrongObjectPtr<UObject>>& Options) override;
-	// End ADatasmithCoreTechTranslator overrides
 
 private:
 	TMap<FString, FString> CADFileToUE4FileMap;
diff --git a/Engine/Plugins/Enterprise/DatasmithCADImporter/Source/DatasmithCADTranslator/Private/DatasmithMeshBuilder.cpp b/Engine/Plugins/Enterprise/DatasmithCADImporter/Source/DatasmithCADTranslator/Private/DatasmithMeshBuilder.cpp
--- a/Engine/Plugins/Enterprise/DatasmithCADImporter/Source/DatasmithCADTranslator/Private/DatasmithMeshBuilder.cpp
+++ b/Engine/Plugins/Enterprise/DatasmithCADImporter/Source/DatasmithCADTranslator/Private/DatasmithMeshBuilder.cpp
@@ -1,13 +1,14 @@
 // Copyright 1998-2019 Epic Games, Inc. All Rights Reserved.
 #include "DatasmithMeshBuilder.h"
 
-#include "CoreTechHelper.h"
 #include "CoreTechFileParser.h"
 #include "DatasmithMeshHelper.h"
-#include "HAL/FileManager.h"
 #include "IDatasmithSceneElements.h"
+
+#include "HAL/FileManager.h"
 #include "MeshDescription.h"
 #include "Misc/Paths.h"
+
 using namespace CADLibrary;
 
 FDatasmithMeshBuilder::FDatasmithMeshBuilder(TMap<FString, FString>& InCADFileToUE4GeomMap, TMap< TSharedPtr< IDatasmithMeshElement >, uint32 >& InMeshElementToCTBodyUuidMap)
@@ -44,6 +45,7 @@ void FDatasmithMeshBuilder::LoadRawDataGeom()
 
 TOptional<FMeshDescription> FDatasmithMeshBuilder::GetMeshDescription(TSharedRef<IDatasmithMeshElement> OutMeshElement, CADLibrary::FMeshParameters& OutMeshParameters)
 {
+#ifdef CAD_INTERFACE
 	uint32* BodyUuid = MeshElementToBodyUuidMap.Find(OutMeshElement);
 	if (BodyUuid == nullptr)
 	{
@@ -63,10 +65,11 @@ TOptional<FMeshDescription> FDatasmithMeshBuilder::GetMeshDescription(TSharedRef
 	FMeshDescription MeshDescription;
 	DatasmithMeshHelper::PrepareAttributeForStaticMesh(MeshDescription);
 
-	if (!ConvertCTBodySetToMeshDescription(ImportParameters, OutMeshParameters, Body.GetTriangleCount(), Body.GetTessellationSet(), MaterialIdToMaterialHash, MeshDescription))
+	if (ConvertCTBodySetToMeshDescription(ImportParameters, OutMeshParameters, Body.GetTriangleCount(), Body.GetTessellationSet(), MaterialIdToMaterialHash, MeshDescription))
 	{
-		return TOptional<FMeshDescription>();
+		return MoveTemp(MeshDescription);
 	}
+#endif // CAD_INTERFACE
 
-	return MoveTemp(MeshDescription);
+	return TOptional<FMeshDescription>();
 }
diff --git a/Engine/Plugins/Enterprise/DatasmithCADImporter/Source/DatasmithCADTranslator/Private/DatasmithMeshBuilder.h b/Engine/Plugins/Enterprise/DatasmithCADImporter/Source/DatasmithCADTranslator/Private/DatasmithMeshBuilder.h
--- a/Engine/Plugins/Enterprise/DatasmithCADImporter/Source/DatasmithCADTranslator/Private/DatasmithMeshBuilder.h
+++ b/Engine/Plugins/Enterprise/DatasmithCADImporter/Source/DatasmithCADTranslator/Private/DatasmithMeshBuilder.h
@@ -4,12 +4,12 @@
 #include "CoreMinimal.h"
 
 #include "CADData.h"
-#include "CoretechHelper.h"
+#include "CADOptions.h"
 
+#include "MeshDescription.h"
 
-class IDatasmithMeshElement;
 
-struct FMeshParameters;
+class IDatasmithMeshElement;
 
 
 class FDatasmithMeshBuilder
diff --git a/Engine/Plugins/Enterprise/DatasmithCADImporter/Source/DatasmithCADTranslator/Private/DatasmithSceneGraphBuilder.cpp b/Engine/Plugins/Enterprise/DatasmithCADImporter/Source/DatasmithCADTranslator/Private/DatasmithSceneGraphBuilder.cpp
--- a/Engine/Plugins/Enterprise/DatasmithCADImporter/Source/DatasmithCADTranslator/Private/DatasmithSceneGraphBuilder.cpp
+++ b/Engine/Plugins/Enterprise/DatasmithCADImporter/Source/DatasmithCADTranslator/Private/DatasmithSceneGraphBuilder.cpp
@@ -1,6 +1,7 @@
 // Copyright 1998-2019 Epic Games, Inc. All Rights Reserved.
 #include "DatasmithSceneGraphBuilder.h"
 
+#ifdef CAD_INTERFACE
 #include "CADData.h"
 #include "CoreTechFileParser.h"
 #include "CoreTechHelper.h"
@@ -948,3 +949,5 @@ void FDatasmithSceneGraphBuilder::AddChildren(TSharedPtr< IDatasmithActorElement
 		}
 	}
 }
+
+#endif // CAD_INTERFACE
diff --git a/Engine/Plugins/Enterprise/DatasmithCADImporter/Source/DatasmithCADTranslator/Private/DatasmithSceneGraphBuilder.h b/Engine/Plugins/Enterprise/DatasmithCADImporter/Source/DatasmithCADTranslator/Private/DatasmithSceneGraphBuilder.h
--- a/Engine/Plugins/Enterprise/DatasmithCADImporter/Source/DatasmithCADTranslator/Private/DatasmithSceneGraphBuilder.h
+++ b/Engine/Plugins/Enterprise/DatasmithCADImporter/Source/DatasmithCADTranslator/Private/DatasmithSceneGraphBuilder.h
@@ -1,6 +1,7 @@
 // Copyright 1998-2019 Epic Games, Inc. All Rights Reserved.
 #pragma once
 
+#ifdef CAD_INTERFACE
 #include "CoreMinimal.h"
 
 #include "CADData.h"
@@ -193,3 +194,5 @@ protected:
 	bool bPreferMaterial;
 	bool bMaterialPropagationIsTopDown;
 };
+
+#endif // CAD_INTERFACE
diff --git a/Engine/Plugins/Enterprise/DatasmithCADImporter/Source/DatasmithDispatcher/Private/DatasmithDispatcher.cpp b/Engine/Plugins/Enterprise/DatasmithCADImporter/Source/DatasmithDispatcher/Private/DatasmithDispatcher.cpp
--- a/Engine/Plugins/Enterprise/DatasmithCADImporter/Source/DatasmithDispatcher/Private/DatasmithDispatcher.cpp
+++ b/Engine/Plugins/Enterprise/DatasmithCADImporter/Source/DatasmithDispatcher/Private/DatasmithDispatcher.cpp
@@ -1,8 +1,9 @@
 // Copyright 1998-2019 Epic Games, Inc. All Rights Reserved.
 #include "DatasmithDispatcher.h"
 
-#include "Containers/Array.h"
 #include "CoreTechFileParser.h"
+
+#include "Containers/Array.h"
 #include "HAL/FileManager.h"
 #include "HAL/Thread.h"
 #include "Misc/Optional.h"
@@ -92,6 +93,7 @@ void FDatasmithDispatcher::Process(bool bWithProcessor)
 
 void FDatasmithDispatcher::ProcessLocal()
 {
+#ifdef CAD_INTERFACE
 	TOptional<FTask> Task = GetTask();
 	while (Task.IsSet())
 	{
@@ -117,6 +119,7 @@ void FDatasmithDispatcher::ProcessLocal()
 
 		Task = GetTask();
 	}
+#endif // CAD_INTERFACE
 }
 
 void FDatasmithDispatcher::AddTask(const FString& InFile)
diff --git a/Engine/Plugins/Enterprise/DatasmithCADImporter/Source/DatasmithDispatcher/Public/DatasmithCommands.h b/Engine/Plugins/Enterprise/DatasmithCADImporter/Source/DatasmithDispatcher/Public/DatasmithCommands.h
--- a/Engine/Plugins/Enterprise/DatasmithCADImporter/Source/DatasmithDispatcher/Public/DatasmithCommands.h
+++ b/Engine/Plugins/Enterprise/DatasmithCADImporter/Source/DatasmithDispatcher/Public/DatasmithCommands.h
@@ -8,12 +8,8 @@
 #include "Misc/DateTime.h"
 #include "CoreTechFileParser.h"
 
-class FCoreTechFileParser;
 
-/**
- * Serializable interface.
- * Should allow writing/Reading from Socket/file/any stream
- */
+class FCoreTechFileParser;
 
 namespace DatasmithDispatcher
 {
@@ -37,9 +33,6 @@ struct FTask
 class FDatasmithDispatcherSocket;
 class FDatasmithDispatcher;
 
-/**
- * Enum with all command types
- */
 enum DatasmithCommandType
 {
 	Ping,
