Fix Clang build (Compile ClangEditor Win64)

diff --git a/Engine/Plugins/Enterprise/DatasmithCADImporter/Source/CADInterfaces/Public/CoreTechFileParser.h b/Engine/Plugins/Enterprise/DatasmithCADImporter/Source/CADInterfaces/Public/CoreTechFileParser.h
--- a/Engine/Plugins/Enterprise/DatasmithCADImporter/Source/CADInterfaces/Public/CoreTechFileParser.h
+++ b/Engine/Plugins/Enterprise/DatasmithCADImporter/Source/CADInterfaces/Public/CoreTechFileParser.h
@@ -37,7 +37,7 @@
 #include "kernel_io/object_io/topo_io/loop_io/loop_io.h"
 #include "kernel_io/repair_io/repair_io.h"
 #pragma warning(pop)
-#endif CAD_INTERFACE
+#endif // CAD_INTERFACE
 
 
 struct FFileStatData;
