Fix Non Unity build
- missing header in DatasmithMeshBuilder.cpp

diff --git a/Engine/Plugins/Enterprise/DatasmithCADImporter/Source/DatasmithCADTranslator/Private/DatasmithMeshBuilder.cpp b/Engine/Plugins/Enterprise/DatasmithCADImporter/Source/DatasmithCADTranslator/Private/DatasmithMeshBuilder.cpp
--- a/Engine/Plugins/Enterprise/DatasmithCADImporter/Source/DatasmithCADTranslator/Private/DatasmithMeshBuilder.cpp
+++ b/Engine/Plugins/Enterprise/DatasmithCADImporter/Source/DatasmithCADTranslator/Private/DatasmithMeshBuilder.cpp
@@ -1,6 +1,10 @@
 // Copyright 1998-2019 Epic Games, Inc. All Rights Reserved.
 #include "DatasmithMeshBuilder.h"
 
+#ifdef CAD_INTERFACE
+#include "CoreTechHelper.h"
+#endif // CAD_INTERFACE
+
 #include "CoreTechFileParser.h"
 #include "DatasmithMeshHelper.h"
 #include "IDatasmithSceneElements.h"
@@ -34,7 +38,7 @@ void FDatasmithMeshBuilder::LoadRawDataGeom()
 	for (const auto& FilePair : CADFileToUE4GeomMap)
 	{
 		FString RawDataFile = FPaths::Combine(CachePath, TEXT("mesh"), FilePair.Value + TEXT(".gm"));
-		if (!IFileManager::Get().FileExists(*RawDataFile)) 
+		if (!IFileManager::Get().FileExists(*RawDataFile))
 		{
 			continue;
 		}
