Max Edge Length tessellation option is not applied

diff --git a/Engine/Plugins/Enterprise/DatasmithCADImporter/Source/CADInterfaces/Private/CTSession.cpp b/Engine/Plugins/Enterprise/DatasmithCADImporter/Source/CADInterfaces/Private/CTSession.cpp
--- a/Engine/Plugins/Enterprise/DatasmithCADImporter/Source/CADInterfaces/Private/CTSession.cpp
+++ b/Engine/Plugins/Enterprise/DatasmithCADImporter/Source/CADInterfaces/Private/CTSession.cpp
@@ -44,8 +44,8 @@ namespace CADLibrary
 
 	void CTSession::SetImportParameters(float ChordTolerance, float MaxEdgeLength, float NormalTolerance, CADLibrary::EStitchingTechnique StitchingTechnique)
 	{
-		ImportParams.ChordTolerance = ChordTolerance / ImportParams.ScaleFactor;
-		ImportParams.MaxEdgeLength = MaxEdgeLength / ImportParams.ScaleFactor;
+		ImportParams.ChordTolerance = ChordTolerance;
+		ImportParams.MaxEdgeLength = MaxEdgeLength;
 		ImportParams.MaxNormalAngle = NormalTolerance;
 		ImportParams.StitchingTechnique = StitchingTechnique;
 
