Buildfix for job [Compile UE4Editor Win64]:
- cleaned CADInterfaces.Build.cs

diff --git a/Engine/Plugins/Enterprise/DatasmithCADImporter/Source/CADInterfaces/CADInterfaces.Build.cs b/Engine/Plugins/Enterprise/DatasmithCADImporter/Source/CADInterfaces/CADInterfaces.Build.cs
--- a/Engine/Plugins/Enterprise/DatasmithCADImporter/Source/CADInterfaces/CADInterfaces.Build.cs
+++ b/Engine/Plugins/Enterprise/DatasmithCADImporter/Source/CADInterfaces/CADInterfaces.Build.cs
@@ -18,22 +18,16 @@ namespace UnrealBuildTool.Rules
 				}
 			);
 
-			// Module is effectively empty without non redistributable content
-			bool bHasNonRedistSources = Directory.Exists(Path.Combine(ModuleDirectory, "Public", "NotForLicensees"));
-
 			// CAD library is only available if CoreTech is available too
 			bool bHasCoretech = System.Type.GetType("CoreTech") != null;
 
 			// Support for Windows only
 			bool bIsPlateformSupported = Target.Platform == UnrealTargetPlatform.Win64;
 
-			if (bHasNonRedistSources && bIsPlateformSupported && bHasCoretech)
+			if (bIsPlateformSupported && bHasCoretech)
 			{
 				PublicDefinitions.Add("CAD_INTERFACE");
 				PublicDependencyModuleNames.Add("CoreTech");
-
-				PublicIncludePaths.Add(ModuleDirectory + "/Public/NotForLicensees");
-				PrivateIncludePaths.Add(ModuleDirectory + "/Private/NotForLicensees");
 			}
 		}
 	}
