Fix Rhino import UVMap

diff --git a/Engine/Plugins/Enterprise/DatasmithCADImporter/Source/DatasmithOpenNurbsTranslator/Private/DatasmithOpenNurbsTranslator.cpp b/Engine/Plugins/Enterprise/DatasmithCADImporter/Source/DatasmithOpenNurbsTranslator/Private/DatasmithOpenNurbsTranslator.cpp
--- a/Engine/Plugins/Enterprise/DatasmithCADImporter/Source/DatasmithOpenNurbsTranslator/Private/DatasmithOpenNurbsTranslator.cpp
+++ b/Engine/Plugins/Enterprise/DatasmithCADImporter/Source/DatasmithOpenNurbsTranslator/Private/DatasmithOpenNurbsTranslator.cpp
@@ -664,7 +664,6 @@ private:
 	void SetLayers(const TSharedPtr<IDatasmithActorElement>& ActorElement, const FOpenNurbsObjectWrapper& Object);
 
 	bool TranslateBRep(ON_Brep* brep, const ON_3dmObjectAttributes& Attributes, FMeshDescription& OutMesh, const TSharedRef< IDatasmithMeshElement >& MeshElement, const FString& Name, bool& bHasNormal);
-	void ApplyTextureMapping(FMeshDescription& Mesh, const ON_3dmObjectAttributes& Attributes);
 
 private:
 	TArray<FOpenNurbsTranslatorImpl*> ChildTranslators;
@@ -2849,8 +2848,6 @@ bool FOpenNurbsTranslatorImpl::TranslateBRep(ON_Brep* Brep, const ON_3dmObjectAt
 	CADLibrary::FMeshParameters MeshParameters;
 	Result = LocalSession->Tessellate(OutMesh, MeshParameters);
 
-	ApplyTextureMapping(OutMesh, Attributes);
-
 	return bool(Result);
 #else
 	// .. Trying to load the mesh tessellated by Rhino
@@ -2885,90 +2882,6 @@ bool FOpenNurbsTranslatorImpl::TranslateBRep(ON_Brep* Brep, const ON_3dmObjectAt
 #endif
 }
 
-void FOpenNurbsTranslatorImpl::ApplyTextureMapping(FMeshDescription& Mesh, const ON_3dmObjectAttributes& Attributes)
-{
-	// Apply the OpenNurbs UV mapping, if present
-	if (Attributes.m_rendering_attributes.m_mappings.Count() == 0)
-	{
-		return;
-	}
-
-	TVertexInstanceAttributesRef<FVector2D> VertexInstanceUVs = Mesh.VertexInstanceAttributes().GetAttributesRef<FVector2D>(MeshAttribute::VertexInstance::TextureCoordinate);
-
-	// Using the first one assuming its from Rhino, ignoring the other mappings
-	int NumChannels = Attributes.m_rendering_attributes.m_mappings.At(0)->m_mapping_channels.Count();
-	if (NumChannels > 1)
-	{
-		VertexInstanceUVs.SetNumIndices(NumChannels);
-	}
-
-	const FVertexInstanceArray& VertexInstances = Mesh.VertexInstances();
-	for (int ChannelIndex = 0; ChannelIndex < NumChannels; ++ChannelIndex)
-	{
-		// Retrieve the mapping from the texture mapping table
-		ON_UUID MappingId = Attributes.m_rendering_attributes.m_mappings.At(0)->m_mapping_channels.At(ChannelIndex)->m_mapping_id;
-
-		const ON_TextureMapping** TextureMappingPtr = UUIDToTextureMapping.Find(MappingId);
-		if (TextureMappingPtr == nullptr)
-		{
-			continue;
-		}
-
-		const ON_TextureMapping& TextureMapping = **TextureMappingPtr;
-
-		// No support for custom mapping as that is dependent on the render plug-in
-		// Nor surface mapping since there's no way to evaluate it
-		if (TextureMapping.m_type == ON_TextureMapping::TYPE::mesh_mapping_primitive ||
-			TextureMapping.m_type == ON_TextureMapping::TYPE::srf_mapping_primitive ||
-			TextureMapping.m_type == ON_TextureMapping::TYPE::brep_mapping_primitive ||
-			TextureMapping.m_type == ON_TextureMapping::TYPE::srfp_mapping)
-		{
-			continue;
-		}
-
-		TVertexInstanceAttributesRef<FVector> VertexInstanceNormals = Mesh.VertexInstanceAttributes().GetAttributesRef<FVector>(MeshAttribute::VertexInstance::Normal);
-		TVertexAttributesRef<FVector> VertexPositions = Mesh.VertexAttributes().GetAttributesRef<FVector>(MeshAttribute::Vertex::Position);
-
-		float MinV = FLT_MAX;
-		float MaxV = -FLT_MAX;
-		for (FVertexInstanceID VertexInstanceID : VertexInstances.GetElementIDs())
-		{
-			// Positions must be converted to right-handed Z-up in the file unit
-			const FVertexID VertexID = Mesh.GetVertexInstanceVertex(VertexInstanceID);
-			const FVector& VertexPosition = VertexPositions[VertexID] / ScalingFactor;
-			const FVector& Normal = VertexInstanceNormals[VertexInstanceID];
-			FVector2D UV = VertexInstanceUVs.Get(VertexInstanceID, ChannelIndex);
-
-			ON_3dPoint ON_VertexPosition(-VertexPosition.X, VertexPosition.Y, VertexPosition.Z);
-			ON_3dVector ON_Normal(-Normal.X, Normal.Y, Normal.Z);
-			ON_3dPoint ON_UV(UV.X, UV.Y, 0);
-
-			int Result = TextureMapping.Evaluate(ON_VertexPosition, ON_Normal, &ON_UV);
-			if (Result != 0)
-			{
-				VertexInstanceUVs.Set(VertexInstanceID, ChannelIndex, FVector2D(ON_UV.x, ON_UV.y));
-			}
-
-			if (ON_UV.y < MinV)
-			{
-				MinV = ON_UV.y;
-			}
-			if (ON_UV.y > MaxV)
-			{
-				MaxV = ON_UV.y;
-			}
-		}
-
-		// Reorient UV along V axis
-		for (FVertexInstanceID VertexInstanceID : VertexInstances.GetElementIDs())
-		{
-			FVector2D UV = VertexInstanceUVs.Get(VertexInstanceID, ChannelIndex);
-			UV.Y = MinV + MaxV - UV.Y;
-			VertexInstanceUVs.Set(VertexInstanceID, ChannelIndex, UV);
-		}
-	}
-}
-
 TOptional< FMeshDescription > FOpenNurbsTranslatorImpl::GetMeshDescription(TSharedRef< IDatasmithMeshElement > MeshElement)
 {
 	// Ref. visitNonInstanceObject (mesh collection part)
