Fix Material parent  label

diff --git a/Engine/Plugins/Enterprise/DatasmithCADImporter/Source/DatasmithCADTranslator/Private/DatasmithSceneGraphBuilder.cpp b/Engine/Plugins/Enterprise/DatasmithCADImporter/Source/DatasmithCADTranslator/Private/DatasmithSceneGraphBuilder.cpp
--- a/Engine/Plugins/Enterprise/DatasmithCADImporter/Source/DatasmithCADTranslator/Private/DatasmithSceneGraphBuilder.cpp
+++ b/Engine/Plugins/Enterprise/DatasmithCADImporter/Source/DatasmithCADTranslator/Private/DatasmithSceneGraphBuilder.cpp
@@ -793,7 +793,6 @@ TSharedPtr<IDatasmithMaterialIDElement> FDatasmithSceneGraphBuilder::FindOrAddMa
 		{
 			DatasmithScene->AddMaterial(MaterialElement);
 			MaterialUuidMap.Add(MaterialUuid, MaterialElement);
-			DatasmithScene->AddMaterial(MaterialElement);
 		}
 	}
 
