Fix import scene graph, Improve import classes naming

diff --git a/Engine/Plugins/Enterprise/DatasmithCADImporter/Source/DatasmithCADTranslator/Private/DatasmithSceneGraphBuilder.cpp b/Engine/Plugins/Enterprise/DatasmithCADImporter/Source/DatasmithCADTranslator/Private/DatasmithSceneGraphBuilder.cpp
--- a/Engine/Plugins/Enterprise/DatasmithCADImporter/Source/DatasmithCADTranslator/Private/DatasmithSceneGraphBuilder.cpp
+++ b/Engine/Plugins/Enterprise/DatasmithCADImporter/Source/DatasmithCADTranslator/Private/DatasmithSceneGraphBuilder.cpp
@@ -15,7 +15,7 @@
 
 namespace
 {
-	void GetMainMaterial(const TMap<FString, FString>& InNodeMetaDataMap, FCTNode* InNode, ActorData& OutNodeData, bool bMaterialPropagationIsTopDown)
+	void GetMainMaterial(const TMap<FString, FString>& InNodeMetaDataMap, FSceneNodeDescription* InNode, ActorData& OutNodeData, bool bMaterialPropagationIsTopDown)
 	{
 		if (const FString* MaterialIdStr = InNodeMetaDataMap.Find(TEXT("MaterialId")))
 		{
@@ -48,13 +48,13 @@ namespace
 }
 
 
-FCTRawDataFile::FCTRawDataFile(const FString& InFileName) 
+FCADSceneGraphDescriptionFile::FCADSceneGraphDescriptionFile(const FString& InFileName) 
 	: FileName(InFileName)
 {
 	ReadFile();
 }
 
-void FCTRawDataFile::GetMaterialDescription(int32 LineNumber, CADLibrary::FCADMaterial& Material) const
+void FCADSceneGraphDescriptionFile::GetMaterialDescription(int32 LineNumber, CADLibrary::FCADMaterial& Material) const
 {
 	FString OutMaterialName;
 	uint8 OutDiffuse[3];
@@ -104,7 +104,7 @@ void FCTRawDataFile::GetMaterialDescription(int32 LineNumber, CADLibrary::FCADMa
 	Material.TextureName = OutTextureName;
 }
 
-bool FCTRawDataFile::GetColor(uint32 ColorHId, FColor& Color) const
+bool FCADSceneGraphDescriptionFile::GetColor(uint32 ColorHId, FColor& Color) const
 {
 	uint32 ColorId;
 	uint8 Alpha;
@@ -119,7 +119,7 @@ bool FCTRawDataFile::GetColor(uint32 ColorHId, FColor& Color) const
 	return false;
 }
 
-bool FCTRawDataFile::GetMaterial(int32 MaterialId, CADLibrary::FCADMaterial& Material) const
+bool FCADSceneGraphDescriptionFile::GetMaterial(int32 MaterialId, CADLibrary::FCADMaterial& Material) const
 {
 	const CADLibrary::FCADMaterial* PMaterial = MaterialIdToMaterial.Find(MaterialId);
 	if (PMaterial)
@@ -130,64 +130,64 @@ bool FCTRawDataFile::GetMaterial(int32 MaterialId, CADLibrary::FCADMaterial& Mat
 	return false;
 }
 
-void FCTRawDataFile::ReadFile()
+void FCADSceneGraphDescriptionFile::ReadFile()
 {
-	FFileHelper::LoadFileToStringArray(RawData, *FileName);
-	if (RawData.Num() < 10)
+	FFileHelper::LoadFileToStringArray(SceneGraphDescriptionData, *FileName);
+	if (SceneGraphDescriptionData.Num() < 10)
 		return;
 
 	FileName = FPaths::GetBaseFilename(FileName);
 
 	FString Left, Right;
-	FString HeaderCTId = RawData[MAPCTIDLINE];
+	FString HeaderCTId = SceneGraphDescriptionData[MAPCTIDLINE];
 	HeaderCTId.Split(SPACE, &Left, &Right);
 	int32 Line = FCString::Atoi(*Right);
 
 	TArray<FString> HeaderCTIdSplit;
-	RawData[Line].ParseIntoArray(HeaderCTIdSplit, TEXT(" "));
+	SceneGraphDescriptionData[Line].ParseIntoArray(HeaderCTIdSplit, TEXT(" "));
 
 	for (int32 index = 0; index < HeaderCTIdSplit.Num(); index += 2)
 	{
-		CTIdToRawEntryMap.Emplace(FCString::Atoi(*HeaderCTIdSplit[index]), FCTNode(*this, FCString::Atoi(*HeaderCTIdSplit[index + 1])));
+		SceneNodeIdToSceneNodeDescriptionMap.Emplace(FCString::Atoi(*HeaderCTIdSplit[index]), FSceneNodeDescription(*this, FCString::Atoi(*HeaderCTIdSplit[index + 1])));
 	}
 
 	// Parse color
-	HeaderCTId = RawData[COLORSETLINE];
+	HeaderCTId = SceneGraphDescriptionData[COLORSETLINE];
 	HeaderCTId.Split(SPACE, &Left, &Right);
 	if (Right != TEXT("0"))
 	{
 		Line = FCString::Atoi(*Right);
-		FString ColorString = RawData[Line];
+		FString ColorString = SceneGraphDescriptionData[Line];
 		while (!ColorString.IsEmpty())
 		{
 			TArray<FString> ColorData;
 			ColorString.ParseIntoArray(ColorData, TEXT(" "));
 			ColorIdToColor.Emplace(FCString::Atoi(*ColorData[0]), FColor(FCString::Atoi(*ColorData[1]), FCString::Atoi(*ColorData[2]), FCString::Atoi(*ColorData[3])));
 			Line++;
-			ColorString = RawData[Line];
+			ColorString = SceneGraphDescriptionData[Line];
 		}
 	}
 	
 	// Parse material
-	HeaderCTId = RawData[MATERIALSETLINE];
+	HeaderCTId = SceneGraphDescriptionData[MATERIALSETLINE];
 	HeaderCTId.Split(SPACE, &Left, &Right);
 	if (Right != TEXT("0"))
 	{
 		Line = FCString::Atoi(*Right);
-		FString MaterialString = RawData[Line];
+		FString MaterialString = SceneGraphDescriptionData[Line];
 		while (!MaterialString.IsEmpty())
 		{
 			CADLibrary::FCADMaterial Material;
 			GetMaterialDescription(Line, Material);
 			MaterialIdToMaterial.Add(Material.MaterialId, Material);
 			Line += 2;
-			MaterialString = RawData[Line];
+			MaterialString = SceneGraphDescriptionData[Line];
 		}
 	}
 
 }
 
-void FCTRawDataFile::SetMaterialMaps(TMap< uint32, FColor>& MaterialUuidToColor, TMap< uint32, CADLibrary::FCADMaterial>& MaterialUuidToMaterial)
+void FCADSceneGraphDescriptionFile::SetMaterialMaps(TMap< uint32, FColor>& MaterialUuidToColor, TMap< uint32, CADLibrary::FCADMaterial>& MaterialUuidToMaterial)
 {
 	for (const auto& Color : ColorIdToColor)
 	{
@@ -202,9 +202,9 @@ void FCTRawDataFile::SetMaterialMaps(TMap< uint32, FColor>& MaterialUuidToColor,
 }
 
 
-void FCTNode::GetMetaDatas(TMap<FString, FString>& MetaDataMap)
+void FSceneNodeDescription::GetMetaDatas(TMap<FString, FString>& MetaDataMap)
 {
-	const FString& RawDataString = RawData.GetString(Line+1);
+	const FString& RawDataString = SceneGraphDescription.GetString(Line+1);
 	if (RawDataString.GetCharArray()[0] != L'M')
 	{
 		return;
@@ -218,13 +218,13 @@ void FCTNode::GetMetaDatas(TMap<FString, FString>& MetaDataMap)
 	EndMeta += StartMeta;
 	for (int32 LineIndex = StartMeta; LineIndex < EndMeta; LineIndex += 2)
 	{
-		MetaDataMap.Add(RawData.GetString(LineIndex), RawData.GetString(LineIndex + 1));
+		MetaDataMap.Add(SceneGraphDescription.GetString(LineIndex), SceneGraphDescription.GetString(LineIndex + 1));
 	}
 }
 
-void FCTNode::GetMaterialSet(TMap<uint32, uint32>& MaterialIdToMaterialHashMap)
+void FSceneNodeDescription::GetMaterialSet(TMap<uint32, uint32>& MaterialIdToMaterialHashMap)
 {
-	const FString& MaterialString = RawData.GetString(EndMeta);
+	const FString& MaterialString = SceneGraphDescription.GetString(EndMeta);
 	TArray<FString> MaterialIDToHashSet;
 	MaterialString.ParseIntoArray(MaterialIDToHashSet, TEXT(" "));
 	if (MaterialIDToHashSet[0] != TEXT("materialMap"))
@@ -241,7 +241,7 @@ void FCTNode::GetMaterialSet(TMap<uint32, uint32>& MaterialIdToMaterialHashMap)
 	}
 }
 
-void FCTNode::GetChildren(TArray<int32>& Children)
+void FSceneNodeDescription::GetChildren(TArray<int32>& Children)
 {
 	if(EndMeta == -1)
 	{
@@ -249,14 +249,14 @@ void FCTNode::GetChildren(TArray<int32>& Children)
 		GetMetaDatas(TempMap);
 	}
 
-	const FString& RawDataString = RawData.GetString(EndMeta);
-	if (RawDataString.GetCharArray()[0] != L'c')
+	const FString& DescriptionString = SceneGraphDescription.GetString(EndMeta);
+	if (DescriptionString.GetCharArray()[0] != L'c')
 	{
 		return;
 	}
 
 	FString Left, Right;
-	RawDataString.Split(SPACE, &Left, &Right);
+	DescriptionString.Split(SPACE, &Left, &Right);
 	int32 nbChildren = FCString::Atoi(*Right);
 	Children.Reserve(nbChildren);
 
@@ -264,11 +264,11 @@ void FCTNode::GetChildren(TArray<int32>& Children)
 	int32 EndChildren = StartChildren + nbChildren;
 	for (int32 LineIndex = StartChildren; LineIndex < EndChildren; LineIndex += 1)
 	{
-		Children.Add(FCString::Atoi(*RawData.GetString(LineIndex)));
+		Children.Add(FCString::Atoi(*SceneGraphDescription.GetString(LineIndex)));
 	}
 }
 
-void FCTNode::GetNodeReference(int32& RefId, FString& ExternalFile, NODE_TYPE& NodeType)
+void FSceneNodeDescription::GetNodeReference(int32& RefId, FString& ExternalFile, NODE_TYPE& NodeType)
 {
 	RefId = 0;
 	ExternalFile = "";
@@ -279,10 +279,10 @@ void FCTNode::GetNodeReference(int32& RefId, FString& ExternalFile, NODE_TYPE& N
 		GetMetaDatas(TempMap);
 	}
 
-	const FString& RawDataString = RawData.GetString(EndMeta + 1);
+	const FString& DescriptionString = SceneGraphDescription.GetString(EndMeta + 1);
 
 	FString TypeStr, Right;
-	RawDataString.Split(SPACE, &TypeStr, &Right);
+	DescriptionString.Split(SPACE, &TypeStr, &Right);
 
 	if (TypeStr == TEXT("ref"))
 	{
@@ -303,14 +303,14 @@ void FCTNode::GetNodeReference(int32& RefId, FString& ExternalFile, NODE_TYPE& N
 	}
 }
 
-void FCTNode::AddTransformToActor(TSharedPtr< IDatasmithActorElement > Actor, const CADLibrary::FImportParameters& ImportParameters)
+void FSceneNodeDescription::AddTransformToActor(TSharedPtr< IDatasmithActorElement > Actor, const CADLibrary::FImportParameters& ImportParameters)
 {
 	if (!Actor.IsValid())
 	{
 		return;
 	}
 
-	const FString& RawDataString = RawData.GetString(EndMeta);
+	const FString& RawDataString = SceneGraphDescription.GetString(EndMeta);
 	TArray<FString> TransformString;
 	RawDataString.ParseIntoArray(TransformString, TEXT(" "));
 
@@ -342,7 +342,7 @@ void FCTNode::AddTransformToActor(TSharedPtr< IDatasmithActorElement > Actor, co
 	Actor->SetRotation(Quat);
 }
 
-void FCTNode::SetNodeType()
+void FSceneNodeDescription::SetNodeType()
 {
 	if (Line == 0)
 	{
@@ -350,7 +350,7 @@ void FCTNode::SetNodeType()
 	} 
 	else
 	{
-		const FString& RawDataString = RawData.GetString(Line);
+		const FString& RawDataString = SceneGraphDescription.GetString(Line);
 		switch (RawDataString.GetCharArray()[0])
 		{
 		case L'C':
@@ -372,33 +372,33 @@ void FCTNode::SetNodeType()
 	}
 }
 
-uint32 FCTNode::GetStaticMeshUuid()
+uint32 FSceneNodeDescription::GetStaticMeshUuid()
 {
 	if (!BodyUUID)
 	{
-		BodyUUID = GetTypeHash(RawData.GetFileName());
-		BodyUUID = HashCombine(BodyUUID, GetTypeHash(*RawData.GetString(Line)));  // "B "+TEXT(CT_OBJECT_ID)
+		BodyUUID = GetTypeHash(SceneGraphDescription.GetFileName());
+		BodyUUID = HashCombine(BodyUUID, GetTypeHash(*SceneGraphDescription.GetString(Line)));  // "B "+TEXT(CT_OBJECT_ID)
 	}
 	return BodyUUID;
 }
 
-bool FCTNode::GetColor(int32 ColorHId, FColor& OutColor) const
+bool FSceneNodeDescription::GetColor(int32 ColorHId, FColor& OutColor) const
 {
-	return RawData.GetColor(ColorHId, OutColor);
+	return SceneGraphDescription.GetColor(ColorHId, OutColor);
 }
 
-bool FCTNode::GetMaterial(int32 MaterialId, CADLibrary::FCADMaterial& OutMaterial) const
+bool FSceneNodeDescription::GetMaterial(int32 MaterialId, CADLibrary::FCADMaterial& OutMaterial) const
 {
-	return RawData.GetMaterial(MaterialId, OutMaterial);
+	return SceneGraphDescription.GetMaterial(MaterialId, OutMaterial);
 }
 
-FCTNode* FCTNode::GetCTNode(int32 NodeId)
+FSceneNodeDescription* FSceneNodeDescription::GetCTNode(int32 NodeId)
 {
-	return RawData.GetCTNode(NodeId);
+	return SceneGraphDescription.GetCTNode(NodeId);
 }
 
-FCTNode::FCTNode(FCTRawDataFile& InRawData, int32 InLine)
-	: RawData(InRawData)
+FSceneNodeDescription::FSceneNodeDescription(FCADSceneGraphDescriptionFile& InSceneGraphDescription, int32 InLine)
+	: SceneGraphDescription(InSceneGraphDescription)
 	, Line(InLine)
 	, StartMeta(-1)
 	, EndMeta(-1)
@@ -409,8 +409,8 @@ FCTNode::FCTNode(FCTRawDataFile& InRawData, int32 InLine)
 }
 
 FDatasmithSceneGraphBuilder::FDatasmithSceneGraphBuilder(TMap<FString, FString>& InCADFileToUE4FileMap, TMap< TSharedPtr< IDatasmithMeshElement >, uint32 >& InMeshElementToCTBodyUuidMap, const FString& InCachePath, TSharedRef<IDatasmithScene> InScene, const FDatasmithSceneSource& InSource, const CADLibrary::FImportParameters& InImportParameters)
-	: CADFileToRawDataFile(InCADFileToUE4FileMap)
-	, MeshElementToCTBodyUuidMap(InMeshElementToCTBodyUuidMap)
+	: CADFileToSceneGraphDescriptionFile(InCADFileToUE4FileMap)
+	, MeshElementToCADBodyUuidMap(InMeshElementToCTBodyUuidMap)
 	, CachePath(InCachePath)
 	, DatasmithScene(InScene)
 	, Source(InSource)
@@ -421,17 +421,17 @@ FDatasmithSceneGraphBuilder::FDatasmithSceneGraphBuilder(TMap<FString, FString>&
 
 bool FDatasmithSceneGraphBuilder::Build()
 {
-	LoadRawDataFile();
+	LoadSceneGraphDescriptionFiles();
 
 	const FString& RootFile = FPaths::GetCleanFilename(Source.GetSourceFile());
 
-	FCTRawDataFile** RawData = CADFileToRawData.Find(RootFile);
+	FCADSceneGraphDescriptionFile** RawData = CADFileToSceneGraphDescription.Find(RootFile);
 	if (RawData == nullptr || *RawData == nullptr)
 	{
 		return false;
 	}
 
-	FCTNode* RootNode = (*RawData)->GetCTNode(1);
+	FSceneNodeDescription* RootNode = (*RawData)->GetCTNode(1);
 	if (RootNode == nullptr)
 	{
 		return false;
@@ -446,26 +446,26 @@ bool FDatasmithSceneGraphBuilder::Build()
 
 void FDatasmithSceneGraphBuilder::Clear()
 {
-	CADFileToRawData.Empty();
-	RawDataArray.Empty();
+	CADFileToSceneGraphDescription.Empty();
+	ArrayOfSceneGraphDescription.Empty();
 	BodyUuidToMeshElementMap.Empty();
 }
 
-void FDatasmithSceneGraphBuilder::LoadRawDataFile()
+void FDatasmithSceneGraphBuilder::LoadSceneGraphDescriptionFiles()
 {
-	RawDataArray.Reserve(CADFileToRawDataFile.Num());
-	CADFileToRawData.Reserve(CADFileToRawDataFile.Num());
-	for (auto FilePair : CADFileToRawDataFile)
+	ArrayOfSceneGraphDescription.Reserve(CADFileToSceneGraphDescriptionFile.Num());
+	CADFileToSceneGraphDescription.Reserve(CADFileToSceneGraphDescriptionFile.Num());
+	for (auto FilePair : CADFileToSceneGraphDescriptionFile)
 	{
 		//OutSgFile = FilePair.Value;
 		FString RawDataFile = FPaths::Combine(CachePath, TEXT("scene"), FilePair.Value + TEXT(".sg"));
-		uint32 index = RawDataArray.Emplace(RawDataFile);
-		RawDataArray[index].SetMaterialMaps(MaterialUuidToColor, MaterialUuidToMaterial);
-		CADFileToRawData.Emplace(FilePair.Key, &RawDataArray[index]);
+		uint32 index = ArrayOfSceneGraphDescription.Emplace(RawDataFile);
+		ArrayOfSceneGraphDescription[index].SetMaterialMaps(MaterialUuidToColor, MaterialUuidToMaterial);
+		CADFileToSceneGraphDescription.Emplace(FilePair.Key, &ArrayOfSceneGraphDescription[index]);
 	}
 }
 
-TSharedPtr< IDatasmithActorElement >  FDatasmithSceneGraphBuilder::BuildNode(FCTNode& Node, ActorData& ParentData)
+TSharedPtr< IDatasmithActorElement >  FDatasmithSceneGraphBuilder::BuildNode(FSceneNodeDescription& Node, ActorData& ParentData)
 {
 	NODE_TYPE Type = Node.GetNodeType();
 
@@ -486,7 +486,7 @@ TSharedPtr< IDatasmithActorElement >  FDatasmithSceneGraphBuilder::BuildNode(FCT
 	return TSharedPtr< IDatasmithActorElement >();
 }
 
-TSharedPtr< IDatasmithActorElement >  FDatasmithSceneGraphBuilder::BuildInstance(FCTNode& Node, ActorData& ParentData)
+TSharedPtr< IDatasmithActorElement >  FDatasmithSceneGraphBuilder::BuildInstance(FSceneNodeDescription& Node, ActorData& ParentData)
 {
 	TMap<FString, FString> InstanceNodeMetaDataMap;
 	TMap<FString, FString> ReferenceNodeMetaDataMap;
@@ -502,7 +502,7 @@ TSharedPtr< IDatasmithActorElement >  FDatasmithSceneGraphBuilder::BuildInstance
 	FString ReferenceFile;
 	Node.GetNodeReference(RefId, ReferenceFile, NodeType);
 
-	FCTNode* ComponentNode = nullptr;
+	FSceneNodeDescription* ComponentNode = nullptr;
 	switch(NodeType)
 	{
 		case COMPONENT:
@@ -513,7 +513,7 @@ TSharedPtr< IDatasmithActorElement >  FDatasmithSceneGraphBuilder::BuildInstance
 
 		case EXTERNALCOMPONENT:
 		{
-			FCTRawDataFile** RawData = CADFileToRawData.Find(ReferenceFile);
+			FCADSceneGraphDescriptionFile** RawData = CADFileToSceneGraphDescription.Find(ReferenceFile);
 			if (RawData != nullptr && *RawData != nullptr)
 			{
 				ComponentNode = (*RawData)->GetCTNode(1);
@@ -610,11 +610,11 @@ void FDatasmithSceneGraphBuilder::GetNodeUUIDAndName(
 	{
 		UEUUID = HashCombine(UEUUID, GetTypeHash(*IUUID));
 	}
-	else if (IOriginalName)
+	if (IOriginalName)
 	{
 		UEUUID = HashCombine(UEUUID, GetTypeHash(*IOriginalName));
 	}
-	else if (IName)
+	if (IName)
 	{
 		UEUUID = HashCombine(UEUUID, GetTypeHash(*IName));
 	}
@@ -623,11 +623,11 @@ void FDatasmithSceneGraphBuilder::GetNodeUUIDAndName(
 	{
 		UEUUID = HashCombine(UEUUID, GetTypeHash(*RUUID));
 	}
-	else if (ROriginalName)
+	if (ROriginalName)
 	{
 		UEUUID = HashCombine(UEUUID, GetTypeHash(*ROriginalName));
 	}
-	else if (RName)
+	if (RName)
 	{
 		UEUUID = HashCombine(UEUUID, GetTypeHash(*RName));
 	}
@@ -635,7 +635,7 @@ void FDatasmithSceneGraphBuilder::GetNodeUUIDAndName(
 	OutUEUUID = FString::Printf(TEXT("0x%08x"), UEUUID);
 }
 
-TSharedPtr< IDatasmithActorElement > FDatasmithSceneGraphBuilder::BuildComponent(FCTNode& Node, ActorData& ParentData)
+TSharedPtr< IDatasmithActorElement > FDatasmithSceneGraphBuilder::BuildComponent(FSceneNodeDescription& Node, ActorData& ParentData)
 {
 	TMap<FString, FString> InstanceNodeMetaDataMap;
 	TMap<FString, FString> ReferenceNodeMetaDataMap;
@@ -662,7 +662,7 @@ TSharedPtr< IDatasmithActorElement > FDatasmithSceneGraphBuilder::BuildComponent
 	return Actor;
 }
 
-TSharedPtr< IDatasmithActorElement > FDatasmithSceneGraphBuilder::BuildBody(FCTNode& Node, ActorData& ParentData)
+TSharedPtr< IDatasmithActorElement > FDatasmithSceneGraphBuilder::BuildBody(FSceneNodeDescription& Node, ActorData& ParentData)
 {
 	TMap<FString, FString> InstanceNodeMetaDataMap;
 	TMap<FString, FString> BodyNodeMetaDataMap;
@@ -708,7 +708,7 @@ TSharedPtr< IDatasmithActorElement > FDatasmithSceneGraphBuilder::BuildBody(FCTN
 	return ActorElement;
 }
 
-TSharedPtr< IDatasmithMeshElement > FDatasmithSceneGraphBuilder::FindOrAddMeshElement(FCTNode& Node, FString& BodyName)
+TSharedPtr< IDatasmithMeshElement > FDatasmithSceneGraphBuilder::FindOrAddMeshElement(FSceneNodeDescription& Node, FString& BodyName)
 {
 	const uint32 ShellUuid = Node.GetStaticMeshUuid();
 	FString ShellUuidName = FString::Printf(TEXT("0x%08x"), ShellUuid);
@@ -753,7 +753,7 @@ TSharedPtr< IDatasmithMeshElement > FDatasmithSceneGraphBuilder::FindOrAddMeshEl
 	DatasmithScene->AddMesh(MeshElement);
 
 	BodyUuidToMeshElementMap.Add(ShellUuid, MeshElement);
-	MeshElementToCTBodyUuidMap.Add(MeshElement, ShellUuid);
+	MeshElementToCADBodyUuidMap.Add(MeshElement, ShellUuid);
 
 	return MeshElement;
 }
@@ -930,13 +930,13 @@ void FDatasmithSceneGraphBuilder::LinkActor(TSharedPtr< IDatasmithActorElement >
 	}
 }
 
-void FDatasmithSceneGraphBuilder::AddChildren(TSharedPtr< IDatasmithActorElement > Actor, FCTNode& ComponentNode, ActorData& ParentData)
+void FDatasmithSceneGraphBuilder::AddChildren(TSharedPtr< IDatasmithActorElement > Actor, FSceneNodeDescription& ComponentNode, ActorData& ParentData)
 {
 	TArray<int32> ChildrenIds;
 	ComponentNode.GetChildren(ChildrenIds);
 	for (const auto& ChildId : ChildrenIds)
 	{
-		FCTNode* ChildNode = ComponentNode.GetCTNode(ChildId);
+		FSceneNodeDescription* ChildNode = ComponentNode.GetCTNode(ChildId);
 		if (ChildNode == nullptr)
 		{
 			continue;
diff --git a/Engine/Plugins/Enterprise/DatasmithCADImporter/Source/DatasmithCADTranslator/Private/DatasmithSceneGraphBuilder.h b/Engine/Plugins/Enterprise/DatasmithCADImporter/Source/DatasmithCADTranslator/Private/DatasmithSceneGraphBuilder.h
--- a/Engine/Plugins/Enterprise/DatasmithCADImporter/Source/DatasmithCADTranslator/Private/DatasmithSceneGraphBuilder.h
+++ b/Engine/Plugins/Enterprise/DatasmithCADImporter/Source/DatasmithCADTranslator/Private/DatasmithSceneGraphBuilder.h
@@ -16,7 +16,7 @@
 
 static const FString SPACE(" ");
 
-class FCTRawDataFile;
+class FCADSceneGraphDescriptionFile;
 class FDatasmithSceneSource;
 class IDatasmithActorElement;
 class IDatasmithMeshElement;
@@ -34,10 +34,10 @@ enum NODE_TYPE
 };
 
 
-class FCTNode
+class FSceneNodeDescription
 {
 public:
-	FCTNode(FCTRawDataFile& InRawData, int32 InLine);
+	FSceneNodeDescription(FCADSceneGraphDescriptionFile& InRawData, int32 InLine);
 
 	NODE_TYPE GetNodeType() { return Type; }
 
@@ -47,7 +47,7 @@ public:
 	void GetNodeReference(int32& OutRefId, FString& OutExternalFile, NODE_TYPE& OutType);
 	void AddTransformToActor(TSharedPtr< IDatasmithActorElement > Actor, const CADLibrary::FImportParameters& InImportParameters);
 
-	FCTNode* GetCTNode(int32 NodeId);
+	FSceneNodeDescription* GetCTNode(int32 NodeId);
 
 	uint32 GetStaticMeshUuid();
 	void GetMaterialSet(TMap<uint32, uint32>& MaterialIdToMaterialHashMap);
@@ -58,7 +58,7 @@ public:
 protected:
 	void SetNodeType();
 	FString SceneGraphFile;
-	FCTRawDataFile& RawData;
+	FCADSceneGraphDescriptionFile& SceneGraphDescription;
 	int32 Line;
 	int32 StartMeta;
 	int32 EndMeta;
@@ -99,21 +99,21 @@ public:
 
 
 
-class FCTRawDataFile //#ueent_CAD
+class FCADSceneGraphDescriptionFile
 {
 public:
-	FCTRawDataFile(const FString& InFileName);
+	FCADSceneGraphDescriptionFile(const FString& InFileName);
 
 	void ReadFile();
 
-	FCTNode* GetCTNode(int32 nodeId)
+	FSceneNodeDescription* GetCTNode(int32 nodeId)
 	{
-		return CTIdToRawEntryMap.Find(nodeId);
+		return SceneNodeIdToSceneNodeDescriptionMap.Find(nodeId);
 	}
 
 	const FString& GetString(int32 line) const
 	{
-		return RawData[line];
+		return SceneGraphDescriptionData[line];
 	}
 
 	const FString& GetFileName() const
@@ -122,7 +122,6 @@ public:
 	}
 
 	bool GetColor(uint32 ColorHId, FColor& Color) const;
-
 	bool GetMaterial(int32 MaterialId, CADLibrary::FCADMaterial& material) const;
 
 	void GetMaterialDescription(int32 LineNumber, CADLibrary::FCADMaterial& Material) const;
@@ -131,8 +130,8 @@ public:
 
 protected:
 	FString FileName;
-	TArray<FString> RawData;
-	TMap<int32, FCTNode> CTIdToRawEntryMap;
+	TArray<FString> SceneGraphDescriptionData;
+	TMap<int32, FSceneNodeDescription> SceneNodeIdToSceneNodeDescriptionMap;
 	TMap<int32, FColor> ColorIdToColor;
 	TMap<int32, CADLibrary::FCADMaterial> MaterialIdToMaterial;
 };
@@ -150,37 +149,37 @@ public:
 	void Clear();
 
 protected:
-	void LoadRawDataFile();
+	void LoadSceneGraphDescriptionFiles();
 
-	TSharedPtr< IDatasmithActorElement > BuildNode(FCTNode& Node, ActorData& ParentData);
-	TSharedPtr< IDatasmithActorElement > BuildInstance(FCTNode& Node, ActorData& ParentData);
-	TSharedPtr< IDatasmithActorElement > BuildComponent(FCTNode& Node, ActorData& ParentData);
-	TSharedPtr< IDatasmithActorElement > BuildBody(FCTNode& NodeId, ActorData& ParentData);
+	TSharedPtr< IDatasmithActorElement > BuildNode(FSceneNodeDescription& Node, ActorData& ParentData);
+	TSharedPtr< IDatasmithActorElement > BuildInstance(FSceneNodeDescription& Node, ActorData& ParentData);
+	TSharedPtr< IDatasmithActorElement > BuildComponent(FSceneNodeDescription& Node, ActorData& ParentData);
+	TSharedPtr< IDatasmithActorElement > BuildBody(FSceneNodeDescription& NodeId, ActorData& ParentData);
 
 	void AddMetaData(TSharedPtr< IDatasmithActorElement > ActorElement, TMap<FString, FString>& InstanceNodeAttributeSetMap, TMap<FString, FString>& ReferenceNodeAttributeSetMap);
-	void AddChildren(TSharedPtr< IDatasmithActorElement > Actor, FCTNode& Node, ActorData& ParentData);
+	void AddChildren(TSharedPtr< IDatasmithActorElement > Actor, FSceneNodeDescription& Node, ActorData& ParentData);
 	void LinkActor(TSharedPtr< IDatasmithActorElement > ParentActor, TSharedPtr< IDatasmithActorElement > Actor);
 
 	TSharedPtr< IDatasmithUEPbrMaterialElement > GetDefaultMaterial();
 	TSharedPtr<IDatasmithMaterialIDElement> FindOrAddMaterial(uint32 MaterialUuid);
 
 	TSharedPtr< IDatasmithActorElement > CreateActor(const TCHAR* ActorUUID, const TCHAR* ActorLabel);
-	TSharedPtr< IDatasmithMeshElement > FindOrAddMeshElement(FCTNode& Node, FString& InLabel);
+	TSharedPtr< IDatasmithMeshElement > FindOrAddMeshElement(FSceneNodeDescription& Node, FString& InLabel);
 
 	void GetNodeUUIDAndName(TMap<FString, FString>& InInstanceNodeMetaDataMap, TMap<FString, FString>& InReferenceNodeMetaDataMap, const TCHAR* InParentUEUUID, FString& OutUEUUID, FString& OutName);
 
 
 protected:
-	TMap<FString, FString>& CADFileToRawDataFile;
-	TMap< TSharedPtr< IDatasmithMeshElement >, uint32 >& MeshElementToCTBodyUuidMap;
+	TMap<FString, FString>& CADFileToSceneGraphDescriptionFile;
+	TMap< TSharedPtr< IDatasmithMeshElement >, uint32 >& MeshElementToCADBodyUuidMap;
 	const FString& CachePath;
 	TSharedRef<IDatasmithScene> DatasmithScene;
 	const FDatasmithSceneSource& Source;
 	const CADLibrary::FImportParameters& ImportParameters;
 
-	TMap<FString, FCTRawDataFile*> CADFileToRawData;
+	TMap<FString, FCADSceneGraphDescriptionFile*> CADFileToSceneGraphDescription;
 
-	TArray<FCTRawDataFile> RawDataArray; //#ueent_CAD
+	TArray<FCADSceneGraphDescriptionFile> ArrayOfSceneGraphDescription; //#ueent_CAD
 
 	TMap< uint32, TSharedPtr< IDatasmithMeshElement > > BodyUuidToMeshElementMap;
 
