Fix import 3dxml + import CAD without UV

diff --git a/Engine/Plugins/Enterprise/DatasmithCADImporter/Source/CADInterfaces/Private/CoreTechFileParser.cpp b/Engine/Plugins/Enterprise/DatasmithCADImporter/Source/CADInterfaces/Private/CoreTechFileParser.cpp
--- a/Engine/Plugins/Enterprise/DatasmithCADImporter/Source/CADInterfaces/Private/CoreTechFileParser.cpp
+++ b/Engine/Plugins/Enterprise/DatasmithCADImporter/Source/CADInterfaces/Private/CoreTechFileParser.cpp
@@ -234,6 +234,7 @@ uint32 GetFaceTessellation(CT_OBJECT_ID FaceID, TArray<FTessellationData>& FaceT
 	Tessellation.VertexCount = VertexCount;
 	Tessellation.NormalCount = NormalCount;
 	Tessellation.IndexCount = IndexCount;
+	Tessellation.TexCoordCount = TexCoordArray ? VertexCount : 0;
 	Tessellation.SizeOfVertexType = GetSize(VertexType);
 	Tessellation.SizeOfTexCoordType = GetSize(TexCoordType);
 	Tessellation.SizeOfNormalType = GetSize(NormalType);
@@ -242,12 +243,16 @@ uint32 GetFaceTessellation(CT_OBJECT_ID FaceID, TArray<FTessellationData>& FaceT
 	Tessellation.VertexArray.Empty(Tessellation.VertexCount * Tessellation.SizeOfVertexType);
 	Tessellation.NormalArray.Empty(Tessellation.NormalCount * Tessellation.SizeOfNormalType);
 	Tessellation.IndexArray.Empty(Tessellation.IndexCount * Tessellation.SizeOfIndexType);
-	Tessellation.TexCoordArray.Empty(Tessellation.VertexCount * Tessellation.SizeOfTexCoordType);
+	Tessellation.TexCoordArray.Empty(Tessellation.TexCoordCount * Tessellation.SizeOfTexCoordType);
 
 	Tessellation.VertexArray.Append((uint8*) VertexArray, 3 * Tessellation.VertexCount * Tessellation.SizeOfVertexType);
 	Tessellation.NormalArray.Append((uint8*)NormalArray, 3 * Tessellation.NormalCount * Tessellation.SizeOfNormalType);
 	Tessellation.IndexArray.Append((uint8*)IndexArray, Tessellation.IndexCount * Tessellation.SizeOfIndexType);
-	Tessellation.TexCoordArray.Append((uint8*)TexCoordArray, 2 * Tessellation.VertexCount * Tessellation.SizeOfTexCoordType);
+
+	if (TexCoordArray)
+	{
+		Tessellation.TexCoordArray.Append((uint8*)TexCoordArray, 2 * Tessellation.TexCoordCount * Tessellation.SizeOfTexCoordType);
+	}
 
 	OutRawDataSize = Tessellation.VertexArray.Num() + Tessellation.NormalArray.Num() + Tessellation.IndexArray.Num() + Tessellation.TexCoordArray.Num();
 
@@ -864,9 +869,11 @@ CT_FLAGS FCoreTechFileParser::SetCoreTechImportOption(const FString& MainFileExt
 		Flags |= CT_LOAD_FLAG_SEARCH_NEW_TOPOLOGY | CT_LOAD_FLAG_COMPLETE_TOPOLOGY;
 	}
 
-	Flags |= CT_LOAD_FLAGS_V5_READ_GEOM_SET;
-	
-	Flags &= ~CT_LOAD_FLAGS_LOAD_EXTERNAL_REF;
+	// 3dxml file is zipped files, it's full managed by Kernel_io. We cannot read it in sequential mode
+	if (MainFileExt != TEXT("3dxml"))
+	{
+		Flags &= ~CT_LOAD_FLAGS_LOAD_EXTERNAL_REF;
+	}
 
 	return Flags;
 }
diff --git a/Engine/Plugins/Enterprise/DatasmithCADImporter/Source/CADTools/Private/CADData.cpp b/Engine/Plugins/Enterprise/DatasmithCADImporter/Source/CADTools/Private/CADData.cpp
--- a/Engine/Plugins/Enterprise/DatasmithCADImporter/Source/CADTools/Private/CADData.cpp
+++ b/Engine/Plugins/Enterprise/DatasmithCADImporter/Source/CADTools/Private/CADData.cpp
@@ -69,7 +69,6 @@ FCTRawGeomFile::FCTRawGeomFile(FString& InFileName, TMap< uint32, FBody* >& InBo
 
 	for(FBody& Body : CTBodySet)
 	{
-		//CTIdToBodyUuidMap.Emplace(Body.GetBodyId(), Body.GetBodyUuid());
 		BodyUuidToBodyMap.Add(Body.GetBodyUuid(), &Body);
 	}
 }
@@ -109,6 +108,7 @@ void WriteTessellationInRawData(FTessellationData& Tessellation, TArray<uint8>&
 	INFO[LineNormalType] = Tessellation.SizeOfNormalType;
 	INFO[LineIndexCount] = Tessellation.IndexCount;
 	INFO[LineIndexType] = Tessellation.SizeOfIndexType;
+	INFO[LineTexCoordCount] = Tessellation.TexCoordCount;
 	INFO[LineTexCoordType] = Tessellation.SizeOfTexCoordType;
 
 	GlobalRawData.Append((uint8*)INFO, sizeof(INFO));
@@ -170,12 +170,13 @@ bool ReadTessellationSetFromFile(TArray<uint8>& RawData, uint32& OutNbBodies, TA
 		Tessellation.SizeOfNormalType = INFO[LineNormalType];
 		Tessellation.IndexCount = INFO[LineIndexCount];
 		Tessellation.SizeOfIndexType = INFO[LineIndexType];
+		Tessellation.TexCoordCount = INFO[LineTexCoordCount];
 		Tessellation.SizeOfTexCoordType = INFO[LineTexCoordType];
 
 		uint32 VertexSize = Tessellation.VertexCount * Tessellation.SizeOfVertexType * 3;
 		uint32 NormalSize = Tessellation.NormalCount * Tessellation.SizeOfNormalType * 3;
 		uint32 IndexSize = Tessellation.IndexCount * Tessellation.SizeOfIndexType;
-		uint32 TexCoordSize = Tessellation.VertexCount * Tessellation.SizeOfTexCoordType * 2;
+		uint32 TexCoordSize = Tessellation.TexCoordCount * Tessellation.SizeOfTexCoordType * 2;
 		uint32 DataByteSize = VertexSize + NormalSize + IndexSize + TexCoordSize + InfoSize;
 
 
diff --git a/Engine/Plugins/Enterprise/DatasmithCADImporter/Source/CADTools/Public/CADData.h b/Engine/Plugins/Enterprise/DatasmithCADImporter/Source/CADTools/Public/CADData.h
--- a/Engine/Plugins/Enterprise/DatasmithCADImporter/Source/CADTools/Public/CADData.h
+++ b/Engine/Plugins/Enterprise/DatasmithCADImporter/Source/CADTools/Public/CADData.h
@@ -25,6 +25,7 @@ enum ECTTessellationLine
 	LineIndexCount,
 	LineIndexType,
 	LineTexCoordType,
+	LineTexCoordCount,
 	LastLine
 };
 
@@ -59,6 +60,7 @@ struct CADTOOLS_API FTessellationData
 	uint32        VertexCount = 0;
 	uint32        NormalCount = 0;
 	uint32        IndexCount = 0;
+	uint32        TexCoordCount = 0;
 
 	uint32    StartVertexIndex = 0;
 
