DatasmithCADImporte/CADLibrary remove NotForLicensees folders

diff --git a/Engine/Plugins/Enterprise/DatasmithCADImporter/Source/CADLibrary/CADLibrary.Build.cs b/Engine/Plugins/Enterprise/DatasmithCADImporter/Source/CADLibrary/CADLibrary.Build.cs
--- a/Engine/Plugins/Enterprise/DatasmithCADImporter/Source/CADLibrary/CADLibrary.Build.cs
+++ b/Engine/Plugins/Enterprise/DatasmithCADImporter/Source/CADLibrary/CADLibrary.Build.cs
@@ -20,16 +20,13 @@ public class CADLibrary : ModuleRules
 			}
 		);
 
-		// Module is effectively empty without non redistributable content
-		bool bHasNonRedistSources = Directory.Exists(Path.Combine(ModuleDirectory, "Public", "NotForLicensees"));
-
 		// Support for Windows only
 		bool bIsPlateformSupported = Target.Platform == UnrealTargetPlatform.Win64;
 
 		// CAD library is only available if CoreTech is available too
 		bool bHasCoretech = System.Type.GetType("CoreTech") != null;
 
-		if (bHasNonRedistSources && bIsPlateformSupported && bHasCoretech)
+		if (bIsPlateformSupported && bHasCoretech)
 		{
 			PublicDefinitions.Add("CAD_LIBRARY");
 
@@ -42,11 +39,5 @@ public class CADLibrary : ModuleRules
 				}
 			);
 		}
-
-		if (bHasNonRedistSources)
-		{
-			PublicIncludePaths.Add(ModuleDirectory + "/Public/NotForLicensees");
-			PrivateIncludePaths.Add(ModuleDirectory + "/Private/NotForLicensees");
-		}
 	}
 }
diff --git a/Engine/Plugins/Enterprise/DatasmithCADImporter/Source/CADLibrary/Private/CoreTechHelper.cpp b/Engine/Plugins/Enterprise/DatasmithCADImporter/Source/CADLibrary/Private/CoreTechHelper.cpp
new file mode 100644
--- /dev/null
+++ b/Engine/Plugins/Enterprise/DatasmithCADImporter/Source/CADLibrary/Private/CoreTechHelper.cpp
@@ -0,0 +1,835 @@
+// Copyright 1998-2019 Epic Games, Inc. All Rights Reserved.
+#include "CoreTechHelper.h"
+
+
+#ifdef CAD_LIBRARY
+#include "CoreTechTypes.h"
+
+#include "CADData.h"
+#include "CoreTechFileParser.h"
+#include "DatasmithMaterialElements.h"
+#include "DatasmithSceneFactory.h"
+#include "DatasmithUtils.h"
+#include "Math.h"
+#include "MeshDescription.h"
+#include "MeshOperator.h"
+#include "Misc/FileHelper.h"
+#include "StaticMeshAttributes.h"
+#include "StaticMeshOperations.h"
+#include "CADData.h"
+
+typedef uint32 TriangleIndex[3];
+
+namespace CADLibrary
+{
+// Ref. GPureMeshInterface
+
+struct FVertexData
+{
+	float Z;
+	int32 Index;
+	FVector Coordinates;
+	bool bIsMerged;
+	FVertexID VertexID;
+	FVertexID SymVertexID;
+
+	/** Default constructor. */
+	FVertexData() {}
+
+	/** Initialization constructor. */
+	FVertexData(int32 InIndex, const FVector& V)
+	{
+		Z = 0.30f * V.X + 0.33f * V.Y + 0.37f * V.Z;
+		Index = InIndex;
+		Coordinates = V;
+		bIsMerged = false;
+		VertexID = FVertexID::Invalid;
+	}
+};
+	
+struct FCompareVertexZ
+{
+	FORCEINLINE bool operator()(FVertexData const& A, FVertexData const& B) const { return A.Z < B.Z; }
+};
+
+// Verify the 3 input indices are not defining a degenerated triangle and fill up the corresponding FVertexIDs
+bool IsTriangleDegenerated(const int32_t* Indices, const TArray<FVertexID>& RemapVertexPosition, FVertexID VertexIDs[3])
+{
+	if (Indices[0] == Indices[1] || Indices[0] == Indices[2] || Indices[1] == Indices[2])
+	{
+		return true;
+	}
+
+	for (int32 Corner = 0; Corner < 3; ++Corner)
+	{
+		VertexIDs[Corner] = RemapVertexPosition[Indices[Corner]];
+	}
+
+	return (VertexIDs[0] == VertexIDs[1] || VertexIDs[0] == VertexIDs[2] || VertexIDs[1] == VertexIDs[2]);
+}
+
+void FillVertexPosition(const FImportParameters& ImportParams, const FMeshParameters& MeshParameters, int32 TriangleCount, TArray<FTessellationData>& FaceTessellationSet, FMeshDescription& MeshDescription)
+{
+	TVertexAttributesRef<FVector> VertexPositions = MeshDescription.VertexAttributes().GetAttributesRef<FVector>(MeshAttribute::Vertex::Position);
+
+	// Create a list of vertex Z/index pairs
+	TArray<FVertexData> VertexDataSet;
+	VertexDataSet.Reserve(TriangleCount * 3);
+
+	FVector Position;
+	uint32 GlobalVertexCount = 0;
+	for (FTessellationData& CTTessellation : FaceTessellationSet)
+	{
+		CTTessellation.StartVertexIndex = GlobalVertexCount;
+		for (uint32 VertexIndex = 0; VertexIndex < CTTessellation.VertexCount; ++VertexIndex, ++GlobalVertexCount)
+		{
+			CopyValue(CTTessellation.VertexArray.GetData(), VertexIndex * 3, CTTessellation.SizeOfVertexType, true, Position);
+			Position *= ImportParams.ScaleFactor;  // convert Position unit into cm according to scaleFactor
+			VertexDataSet.Emplace(GlobalVertexCount, Position);
+		}
+	}
+	VertexDataSet.SetNum(GlobalVertexCount);
+
+	// Sort the vertices by z value
+	VertexDataSet.Sort(FCompareVertexZ());
+
+	TArray<uint32> NewIndexOf;
+	NewIndexOf.SetNumZeroed(GlobalVertexCount);
+
+	TArray<int32> IndexOfCoincidentNode;
+	IndexOfCoincidentNode.SetNumZeroed(GlobalVertexCount);
+
+	int32 VertexCount = 0;
+	// Search for duplicates, quickly!
+	for (int32 i = 0; i < VertexDataSet.Num(); i++)
+	{
+		NewIndexOf[VertexDataSet[i].Index] = i;
+		if (VertexDataSet[i].bIsMerged)
+		{
+			continue;
+		}
+
+		VertexDataSet[i].bIsMerged = true;
+		int32 Index_i = VertexDataSet[i].Index;
+		IndexOfCoincidentNode[Index_i] = Index_i;
+
+		// only need to search forward, since we add pairs both ways
+		for (int32 j = i + 1; j < VertexDataSet.Num(); j++)
+		{
+			if (FMath::Abs(VertexDataSet[j].Z - VertexDataSet[i].Z) > KINDA_SMALL_NUMBER)
+			{
+				break; // can't be any more duplicated
+			}
+
+			const FVector& PositionA = VertexDataSet[i].Coordinates;
+			const FVector& PositionB = VertexDataSet[j].Coordinates;
+
+			if (PositionA.Equals(PositionB, KINDA_SMALL_NUMBER))
+			{
+				VertexDataSet[j].bIsMerged = true;
+				IndexOfCoincidentNode[VertexDataSet[j].Index] = Index_i;
+			}
+		}
+		VertexCount++;
+	}
+
+
+	// if Symmetric mesh, the symmetric side of the mesh have to be generated
+	FMatrix SymmetricMatrix;
+	bool bIsSymmetricMesh = MeshParameters.bIsSymmetric;
+	if (bIsSymmetricMesh)
+	{
+		SymmetricMatrix = FDatasmithUtils::GetSymmetricMatrix(MeshParameters.SymmetricOrigin, MeshParameters.SymmetricNormal);
+	}
+
+
+	// Make MeshDescription.VertexPositions and VertexID
+	MeshDescription.ReserveNewVertices(VertexCount);
+	uint32 GlobalVertexIndex = 0;
+
+	for (uint32 VertexIndex = 0; VertexIndex < GlobalVertexCount; ++VertexIndex)
+	{
+		uint32 RealIndex = VertexDataSet[VertexIndex].Index;
+		if (IndexOfCoincidentNode[RealIndex] != RealIndex)
+		{
+			continue;
+		}
+
+		FVertexID VertexID = MeshDescription.CreateVertex();
+		VertexPositions[VertexID] = FDatasmithUtils::ConvertVector((FDatasmithUtils::EModelCoordSystem) ImportParams.ModelCoordSys, VertexDataSet[VertexIndex].Coordinates);
+		VertexDataSet[VertexIndex].VertexID = VertexID;
+	}
+
+	if (bIsSymmetricMesh)
+	{
+		for (uint32 VertexIndex = 0; VertexIndex < GlobalVertexCount; ++VertexIndex)
+		{
+			uint32 RealIndex = VertexDataSet[VertexIndex].Index;
+			if (IndexOfCoincidentNode[RealIndex] != RealIndex)
+			{
+				continue;
+			}
+
+			FVertexID VertexID = MeshDescription.CreateVertex();
+			VertexPositions[VertexID] = FDatasmithUtils::ConvertVector((FDatasmithUtils::EModelCoordSystem) ImportParams.ModelCoordSys, VertexDataSet[VertexIndex].Coordinates);
+			VertexPositions[VertexID] = SymmetricMatrix.TransformPosition(VertexPositions[VertexID]);
+
+			VertexDataSet[VertexIndex].SymVertexID = VertexID;
+		}
+	}
+
+	// For each face, for each vertex set VertexId
+	GlobalVertexIndex = 0;
+	for (FTessellationData& CTTessellation : FaceTessellationSet)
+	{
+		CTTessellation.VertexIdSet.SetNum(CTTessellation.VertexCount);
+		for (uint32 VertexIndex = 0; VertexIndex < CTTessellation.VertexCount; ++VertexIndex, ++GlobalVertexIndex)
+		{
+			uint32 NewIndex = NewIndexOf[IndexOfCoincidentNode[GlobalVertexIndex]];
+			CTTessellation.VertexIdSet[VertexIndex] = VertexDataSet[NewIndex].VertexID.GetValue();
+		}
+	}
+
+	if (bIsSymmetricMesh)
+	{
+
+		GlobalVertexIndex = 0;
+		for (FTessellationData& CTTessellation : FaceTessellationSet)
+		{
+			CTTessellation.SymVertexIdSet.SetNum(CTTessellation.VertexCount);
+			for (uint32 VertexIndex = 0; VertexIndex < CTTessellation.VertexCount; ++VertexIndex, ++GlobalVertexIndex)
+			{
+				uint32 NewIndex = NewIndexOf[IndexOfCoincidentNode[GlobalVertexIndex]];
+				CTTessellation.SymVertexIdSet[VertexIndex] = VertexDataSet[NewIndex].SymVertexID.GetValue();;
+			}
+		}
+	}
+}
+
+void UpdatePolygonGroup(TMap<uint32, FPolygonGroupID>& MaterialToPolygonGroupMapping, TPolygonGroupAttributesRef<FName>& PolygonGroupImportedMaterialSlotNames, FMeshDescription& MeshDescription)
+{
+	if (MeshDescription.PolygonGroups().Num() > 0)
+	{
+		for (FPolygonGroupID PolygonGroupID : MeshDescription.PolygonGroups().GetElementIDs())
+		{
+			FName initialSlotName = PolygonGroupImportedMaterialSlotNames[PolygonGroupID];
+			uint32 MaterialHash = StaticCast<uint32>(FCString::Atoi64(*initialSlotName.ToString()));
+			MaterialToPolygonGroupMapping.Add(MaterialHash, PolygonGroupID);
+		}
+	}
+
+	for (auto& Material : MaterialToPolygonGroupMapping)
+	{
+		if (Material.Value == FPolygonGroupID::Invalid)
+		{
+			uint32 MaterialHash = Material.Key;
+			FName ImportedSlotName = *LexToString<uint32>(MaterialHash);
+
+			FPolygonGroupID PolyGroupID = MeshDescription.CreatePolygonGroup();
+			PolygonGroupImportedMaterialSlotNames[PolyGroupID] = ImportedSlotName;
+			Material.Value = PolyGroupID;
+		}
+	}
+}
+
+bool FillMesh(const FMeshParameters& MeshParameters, const FImportParameters& ImportParams, TArray<FTessellationData>& FaceTessellations, TMap<uint32, FPolygonGroupID> MaterialToPolygonGroupMapping, FMeshDescription& MeshDescription)
+{
+	const int32 UVChannel = 0;
+	const int32 TriangleCount = 3;
+	const TriangleIndex Clockwise = { 0, 1, 2 };
+	const TriangleIndex CounterClockwise = { 0, 2, 1 };
+
+
+	TArray<FVertexInstanceID> TriangleVertexInstanceIDs;
+	TriangleVertexInstanceIDs.SetNum(TriangleCount);
+
+	TArray<FVertexInstanceID> MeshVertexInstanceIDs;
+	TArray<uint32> CTFaceIndex;  // new CT face index to remove degenerated face
+
+	// Gather all array data
+	FStaticMeshAttributes Attributes(MeshDescription);
+	TVertexInstanceAttributesRef<FVector> VertexInstanceNormals = Attributes.GetVertexInstanceNormals();
+	TVertexInstanceAttributesRef<FVector> VertexInstanceTangents = Attributes.GetVertexInstanceTangents();
+	TVertexInstanceAttributesRef<float> VertexInstanceBinormalSigns = Attributes.GetVertexInstanceBinormalSigns();
+	TVertexInstanceAttributesRef<FVector4> VertexInstanceColors = Attributes.GetVertexInstanceColors();
+	TVertexInstanceAttributesRef<FVector2D> VertexInstanceUVs = Attributes.GetVertexInstanceUVs();
+	TPolygonGroupAttributesRef<FName> PolygonGroupImportedMaterialSlotNames = Attributes.GetPolygonGroupMaterialSlotNames();
+
+	if ( !VertexInstanceNormals.IsValid() || !VertexInstanceTangents.IsValid() || !VertexInstanceBinormalSigns.IsValid() || !VertexInstanceColors.IsValid() || !VertexInstanceUVs.IsValid() || !PolygonGroupImportedMaterialSlotNames.IsValid() )
+	{
+		return false;
+	}
+
+	VertexInstanceUVs.SetNumIndices(1);
+
+	int32 NbStep = 1;
+	if (MeshParameters.bIsSymmetric)
+	{
+		NbStep = 2;
+	}
+
+	for (int32 Step = 0; Step < NbStep; ++Step)
+	{
+		// Swap mesh if needed
+		const TriangleIndex& Orientation = (!MeshParameters.bNeedSwapOrientation == (bool)Step) ? CounterClockwise : Clockwise;
+
+		for (FTessellationData& Tessellation : FaceTessellations)
+		{
+			// Get the polygonGroup
+			const FPolygonGroupID* PolygonGroupID = MaterialToPolygonGroupMapping.Find(Tessellation.MaterialHash);
+			if (PolygonGroupID == nullptr)
+			{
+				continue;
+			}
+
+			//int32 TriangleCount = IndicesCount / 3;
+			int32 VertexIDs[3];
+			FVector Temp3D = { 0, 0, 0 };
+			FVector2D TexCoord2D = { 0, 0 };
+			int32_t IndicesVertex[3];
+
+			MeshVertexInstanceIDs.SetNum(Tessellation.IndexCount);
+			CTFaceIndex.Reserve(Tessellation.IndexCount);
+			CTFaceIndex.SetNum(0);
+
+			TArray<int32>& VertexIdSet = (Step == 0) ? Tessellation.VertexIdSet : Tessellation.SymVertexIdSet;
+
+			// build each valid face i.e. 3 different indexes
+			for (uint32 Index = 0, NewIndex = 0; Index < Tessellation.IndexCount; Index += 3)
+			{
+				CopyValue(Tessellation.IndexArray.GetData(), Index, Tessellation.SizeOfIndexType, IndicesVertex);
+				VertexIDs[Orientation[0]] = VertexIdSet[IndicesVertex[0]];
+				VertexIDs[Orientation[1]] = VertexIdSet[IndicesVertex[1]];
+				VertexIDs[Orientation[2]] = VertexIdSet[IndicesVertex[2]];
+
+				// Verify the 3 input indices are not defining a degenerated triangle
+				if (VertexIDs[0] == VertexIDs[1] || VertexIDs[0] == VertexIDs[2] || VertexIDs[1] == VertexIDs[2])
+				{
+					continue;
+				}
+
+				CTFaceIndex.Add(IndicesVertex[0]);
+				CTFaceIndex.Add(IndicesVertex[1]);
+				CTFaceIndex.Add(IndicesVertex[2]);
+
+				TriangleVertexInstanceIDs[0] = MeshVertexInstanceIDs[NewIndex++] = MeshDescription.CreateVertexInstance((FVertexID) VertexIDs[0]);
+				TriangleVertexInstanceIDs[1] = MeshVertexInstanceIDs[NewIndex++] = MeshDescription.CreateVertexInstance((FVertexID) VertexIDs[1]);
+				TriangleVertexInstanceIDs[2] = MeshVertexInstanceIDs[NewIndex++] = MeshDescription.CreateVertexInstance((FVertexID) VertexIDs[2]);
+
+				// Add the triangle as a polygon to the mesh description
+				const FPolygonID NewPolygonID = MeshDescription.CreatePolygon(*PolygonGroupID, TriangleVertexInstanceIDs);
+			}
+
+			// finalization of the mesh by setting colors, tangents, bi-normals, UV
+			// for (uint32 Index = 0; Index <: CTFaceIndex)  // pour toutes les index des faces non degeneres de CT
+
+			for (int32 IndexFace = 0; IndexFace < CTFaceIndex.Num(); IndexFace += 3)
+			{
+				for (int32 Index = 0; Index < TriangleCount; Index++)
+				{
+					FVertexInstanceID VertexInstanceID = MeshVertexInstanceIDs[IndexFace + Orientation[Index]];
+
+					VertexInstanceColors[VertexInstanceID] = FLinearColor::White;
+					VertexInstanceTangents[VertexInstanceID] = FVector(ForceInitToZero);
+					VertexInstanceBinormalSigns[VertexInstanceID] = 0.0f;
+				}
+			}
+
+			if (!Tessellation.TexCoordArray.Num()) 
+			{
+				for (int32 IndexFace = 0; IndexFace < CTFaceIndex.Num(); IndexFace += 3)
+				{
+					for (int32 Index = 0; Index < TriangleCount; Index++)
+					{
+						FVertexInstanceID VertexInstanceID = MeshVertexInstanceIDs[IndexFace + Orientation[Index]];
+						CopyValue(Tessellation.TexCoordArray.GetData(), CTFaceIndex[IndexFace + Index] * 2, Tessellation.SizeOfTexCoordType, false, Temp3D);
+						VertexInstanceUVs.Set(VertexInstanceID, UVChannel, FVector2D(Temp3D));
+					}
+				}
+			}
+
+			if (Tessellation.NormalCount == 1)
+			{
+				CopyValue(Tessellation.NormalArray.GetData(), 0, Tessellation.SizeOfNormalType, true, Temp3D);
+				Temp3D = FDatasmithUtils::ConvertVector((FDatasmithUtils::EModelCoordSystem) ImportParams.ModelCoordSys, Temp3D).GetSafeNormal();
+				for (int32 Index = 0; Index < CTFaceIndex.Num(); Index++)
+				{
+					FVertexInstanceID VertexInstanceID = MeshVertexInstanceIDs[Index];
+					VertexInstanceNormals[VertexInstanceID] = Temp3D;
+				}
+			}
+			else
+			{
+				for (int32 IndexFace = 0; IndexFace < CTFaceIndex.Num(); IndexFace += 3)
+				{
+					for (int32 Index = 0; Index < 3; Index++)
+					{
+						FVertexInstanceID VertexInstanceID = MeshVertexInstanceIDs[IndexFace + Orientation[Index]];
+						CopyValue(Tessellation.NormalArray.GetData(), CTFaceIndex[IndexFace + Index] * 3, Tessellation.SizeOfNormalType, true, Temp3D);
+						VertexInstanceNormals[VertexInstanceID] = FDatasmithUtils::ConvertVector((FDatasmithUtils::EModelCoordSystem) ImportParams.ModelCoordSys, Temp3D).GetSafeNormal();
+					}
+				}
+			}
+
+			// compute normals
+			if (Step)
+			{
+				FMatrix SymmetricMatrix;
+				SymmetricMatrix = FDatasmithUtils::GetSymmetricMatrix(MeshParameters.SymmetricOrigin, MeshParameters.SymmetricNormal);
+				for (int32 Index = 0; Index < CTFaceIndex.Num(); Index ++)
+				{
+					VertexInstanceNormals[MeshVertexInstanceIDs[Index]] = SymmetricMatrix.TransformVector(VertexInstanceNormals[MeshVertexInstanceIDs[Index]]);;
+				}
+			}
+
+			if (MeshParameters.bNeedSwapOrientation)
+			{
+				for (int32 Index = 0; Index < CTFaceIndex.Num(); Index++)
+				{
+					VertexInstanceNormals[MeshVertexInstanceIDs[Index]] = VertexInstanceNormals[MeshVertexInstanceIDs[Index]] * -1.f;
+				}
+			}
+		}
+	}
+
+
+	return true;
+}
+
+bool ConvertCTBodySetToMeshDescription(const FImportParameters& ImportParams, const FMeshParameters& MeshParameters, TArray<CT_OBJECT_ID>& BodySet, FMeshDescription& MeshDescription)
+{
+	TArray<FTessellationData> FaceTessellationSet;
+	TMap<uint32, uint32> MaterialIdToMaterialHash;
+
+	int32 BodyRawSize = 0;
+	uint32 TriangleCount = GetBodiesTessellations(BodySet, (TArray<FTessellationData>&) FaceTessellationSet, MaterialIdToMaterialHash, BodyRawSize, ImportParams.ScaleFactor);
+
+	return ConvertCTBodySetToMeshDescription(ImportParams, MeshParameters, TriangleCount, FaceTessellationSet, MaterialIdToMaterialHash, MeshDescription);
+}
+
+bool ConvertCTBodySetToMeshDescription(const FImportParameters& ImportParams, const FMeshParameters& MeshParameters, uint32 TriangleCount, TArray<FTessellationData>& FaceTessellationSet, TMap<uint32, uint32>& MaterialIdToMaterialHash, FMeshDescription& MeshDescription)
+{
+	// Ref. CreateMesh(UDatasmithCADImportOptions* CADOptions, CTMesh& Mesh)
+	MeshDescription.EdgeAttributes().RegisterAttribute<bool>(MeshAttribute::Edge::IsUVSeam, 1, false);
+
+	// in a closed big mesh VertexCount ~ TriangleCount / 2, EdgeCount ~ 1.5* TriangleCount
+	MeshDescription.ReserveNewVertexInstances(TriangleCount*3);
+	MeshDescription.ReserveNewPolygons(TriangleCount);
+	MeshDescription.ReserveNewEdges(TriangleCount*3);
+
+	// CoreTech is generating position duplicates. make sure to remove them before filling the mesh description
+	TArray<FVertexID> RemapVertexPosition;
+	FillVertexPosition(ImportParams, MeshParameters, TriangleCount, FaceTessellationSet, MeshDescription);
+
+	TMap<uint32, FPolygonGroupID> MaterialToPolygonGroupMapping;
+	for (const FTessellationData& FaceTessellation : FaceTessellationSet)
+	{
+		MaterialToPolygonGroupMapping.Add(FaceTessellation.MaterialHash, FPolygonGroupID::Invalid);
+	}
+
+	// Add the mesh's materials as polygon groups
+	TPolygonGroupAttributesRef<FName> PolygonGroupImportedMaterialSlotNames = MeshDescription.PolygonGroupAttributes().GetAttributesRef<FName>(MeshAttribute::PolygonGroup::ImportedMaterialSlotName);
+	UpdatePolygonGroup(MaterialToPolygonGroupMapping, PolygonGroupImportedMaterialSlotNames, MeshDescription);
+
+	if (!FillMesh(MeshParameters, ImportParams, FaceTessellationSet, MaterialToPolygonGroupMapping, MeshDescription))
+	{
+		return false;
+	}
+
+	// Orient mesh
+	MeshOperator::OrientMesh(MeshDescription);
+
+	// Build edge meta data
+	FStaticMeshOperations::DetermineEdgeHardnessesFromVertexInstanceNormals(MeshDescription);
+
+	return MeshDescription.Polygons().Num() > 0;
+}
+
+
+
+double Distance(CT_COORDINATE Point1, CT_COORDINATE Point2)
+{
+	return sqrt((Point2.xyz[0] - Point1.xyz[0]) * (Point2.xyz[0] - Point1.xyz[0]) + (Point2.xyz[1] - Point1.xyz[1]) * (Point2.xyz[1] - Point1.xyz[1]) + (Point2.xyz[2] - Point1.xyz[2]) * (Point2.xyz[2] - Point1.xyz[2]));
+}
+
+template<typename UVType>
+void ScaleUV(CT_OBJECT_ID FaceID, FTessellationData& Tessellation, UVType Scale)
+{
+	UVType VMin, VMax, UMin, UMax;
+	VMin = UMin = HUGE_VALF;
+	VMax = UMax = -HUGE_VALF;
+	UVType* UVSet = (UVType*)Tessellation.TexCoordArray;
+	for (uint32 Index = 0, UVCoord = 0; Index < Tessellation.VertexCount; ++Index, UVCoord += 2)
+	{
+		UMin = FMath::Min(UVSet[UVCoord + 0], UMin);
+		UMax = FMath::Max(UVSet[UVCoord + 0], UMax);
+		VMin = FMath::Min(UVSet[UVCoord + 1], VMin);
+		VMax = FMath::Max(UVSet[UVCoord + 1], VMax);
+	}
+
+	double PuMin, PuMax, PvMin, PvMax;
+	PuMin = PvMin = HUGE_VALF;
+	PuMax = PvMax = -HUGE_VALF;
+
+	// fast UV min max 
+	CT_FACE_IO::AskUVminmax(FaceID, PuMin, PuMax, PvMin, PvMax);
+
+	const uint32 NbIsoCurves = 7;
+
+	// Compute Point grid on the restricted surface defined by [PuMin, PuMax], [PvMin, PvMax]
+	CT_OBJECT_ID SurfaceID; 
+	CT_ORIENTATION Orientation;
+	CT_FACE_IO::AskSurface(FaceID, SurfaceID, Orientation);
+
+	CT_OBJECT_TYPE SurfaceType;
+	CT_SURFACE_IO::AskType(SurfaceID, SurfaceType);
+
+	UVType DeltaU = (PuMax - PuMin) / (NbIsoCurves - 1);
+	UVType DeltaV = (PvMax - PvMin) / (NbIsoCurves - 1);
+	UVType U = PuMin, V = PvMin;
+
+	CT_COORDINATE NodeMatrix[121];
+
+	for (int32 IndexI = 0; IndexI < NbIsoCurves; IndexI++)
+	{
+		for (int32 IndexJ = 0; IndexJ < NbIsoCurves; IndexJ++)
+		{
+			CT_SURFACE_IO::Evaluate(SurfaceID, U, V, NodeMatrix[IndexI*NbIsoCurves +IndexJ]);
+			V += DeltaV;
+		}
+		U += DeltaU;
+		V = PvMin;
+	}
+
+	// Compute length of 7 iso V line
+	UVType LengthU[NbIsoCurves];
+	UVType LengthUMin = HUGE_VAL;
+	UVType LengthUMax = 0;
+	UVType LengthUMed = 0;
+
+	for (int32 IndexJ = 0; IndexJ < NbIsoCurves; IndexJ++)
+	{
+		LengthU[IndexJ] = 0;
+		for (int32 IndexI = 0; IndexI < (NbIsoCurves - 1); IndexI++)
+		{
+			LengthU[IndexJ] += Distance(NodeMatrix[IndexI * NbIsoCurves + IndexJ], NodeMatrix[(IndexI + 1) * NbIsoCurves + IndexJ]);
+		}
+		LengthUMed += LengthU[IndexJ];
+		LengthUMin = FMath::Min(LengthU[IndexJ], LengthUMin);
+		LengthUMax = FMath::Max(LengthU[IndexJ], LengthUMax);
+	}
+	LengthUMed /= NbIsoCurves;
+	LengthUMed = LengthUMed * 2 / 3 + LengthUMax / 3;
+
+	// Compute length of 7 iso U line
+	UVType LengthV[NbIsoCurves];
+	UVType LengthVMin = HUGE_VAL;
+	UVType LengthVMax = 0;
+	UVType LengthVMed = 0;
+
+	for (int32 IndexI = 0; IndexI < NbIsoCurves; IndexI++)
+	{
+		LengthV[IndexI] = 0;
+		for (int32 IndexJ = 0; IndexJ < (NbIsoCurves - 1); IndexJ++)
+		{
+			LengthV[IndexI] += Distance(NodeMatrix[IndexI * NbIsoCurves + IndexJ], NodeMatrix[IndexI * NbIsoCurves + IndexJ + 1]);
+		}
+		LengthVMed += LengthV[IndexI];
+		LengthVMin = FMath::Min(LengthV[IndexI], LengthVMin);
+		LengthVMax = FMath::Max(LengthV[IndexI], LengthVMax);
+	}
+	LengthVMed /= NbIsoCurves;
+	LengthVMed = LengthVMed * 2 / 3 + LengthVMax / 3;
+
+	switch (SurfaceType)
+	{
+	case CT_CONE_TYPE:
+	case CT_CYLINDER_TYPE:
+	case CT_SPHERE_TYPE:
+		Swap(LengthUMed, LengthVMed);
+		break;
+	case CT_S_REVOL_TYPE:
+	case CT_TORUS_TYPE:
+		// Need swap ?
+		// Swap(LengthUMed, LengthVMed);
+		break;
+	case CT_S_NURBS_TYPE:
+	case CT_PLANE_TYPE:
+	case CT_S_OFFSET_TYPE:
+	case CT_S_RULED_TYPE:
+	case CT_TABULATED_RULED_TYPE:
+	case CT_S_LINEARTRANSFO_TYPE:
+	case CT_S_NONLINEARTRANSFO_TYPE:
+	case CT_S_BLEND_TYPE:
+	default:
+		break;
+	}
+
+	// scale the UV map
+	// 0.1 define UV in cm and not in mm
+	UVType VScale = Scale * LengthVMed * 1 / (VMax - VMin) / 100;
+	UVType UScale = Scale * LengthUMed * 1 / (UMax - UMin) / 100;
+
+	for (uint32 Index = 0, UVCoord = 0; Index < Tessellation.VertexCount; ++Index, UVCoord += 2)
+	{
+		UVSet[UVCoord + 0] *= UScale;
+		UVSet[UVCoord + 1] *= VScale;
+	}
+}
+
+void SetFaceMainMaterial(FObjectDisplayDataId& InFaceMaterial, FObjectDisplayDataId& InBodyMaterial, TMap<uint32, uint32> MaterialIdToMaterialHashMap, FTessellationData& OutFaceTessellations)
+{
+	uint32 MaterialId = 0; // either CT Material ID or CT Color/Alpha ID (ColorId & Alpha << 24)
+	uint32 MaterialHash = 0;
+	if (InFaceMaterial.MaterialId > 0)
+	{
+		MaterialId = InFaceMaterial.MaterialId;
+		if (uint32* MHash = MaterialIdToMaterialHashMap.Find(InFaceMaterial.MaterialId))
+		{
+			MaterialHash = *MHash;
+		}
+		else
+		{
+			MaterialHash = BuildMaterialHash(InFaceMaterial.MaterialId);
+		}
+	}
+	else if (InFaceMaterial.ColorHId > 0)
+	{
+		MaterialId = InFaceMaterial.ColorHId;
+		if (uint32* MHash = MaterialIdToMaterialHashMap.Find(InFaceMaterial.ColorHId))
+		{
+			MaterialHash = *MHash;
+		}
+		else
+		{
+			MaterialHash = BuildColorHash(InFaceMaterial.ColorHId);
+		}
+	}
+	else if (InBodyMaterial.MaterialId > 0)
+	{
+		MaterialId = InBodyMaterial.MaterialId;
+		if (uint32* MHash = MaterialIdToMaterialHashMap.Find(InBodyMaterial.MaterialId))
+		{
+			MaterialHash = *MHash;
+		}
+		else 
+		{
+			MaterialHash = BuildMaterialHash(InBodyMaterial.MaterialId);
+		}
+	}
+	else if (InBodyMaterial.ColorHId > 0)
+	{
+		MaterialId = InBodyMaterial.ColorHId;
+		if (uint32* MHash = MaterialIdToMaterialHashMap.Find(InBodyMaterial.ColorHId))
+		{
+			MaterialHash = *MHash;
+		}
+		else
+		{
+			MaterialHash = BuildColorHash(InBodyMaterial.ColorHId);
+		}
+	}
+
+	// set output
+	MaterialIdToMaterialHashMap.Add(MaterialId, MaterialHash);
+	OutFaceTessellations.MaterialId = MaterialId;
+	OutFaceTessellations.MaterialHash = MaterialHash;
+}
+
+TSharedPtr<IDatasmithUEPbrMaterialElement> CreateDefaultUEPbrMaterial()
+{
+	// Take the Material diffuse color and connect it to the BaseColor of a UEPbrMaterial
+	TSharedRef<IDatasmithUEPbrMaterialElement> MaterialElement = FDatasmithSceneFactory::CreateUEPbrMaterial(TEXT("0"));
+	MaterialElement->SetLabel(TEXT("DefaultCADImportMaterial"));
+
+	FLinearColor LinearColor = FLinearColor::FromPow22Color(FColor(200, 200, 200, 255));
+	IDatasmithMaterialExpressionColor* ColorExpression = MaterialElement->AddMaterialExpression<IDatasmithMaterialExpressionColor>();
+	ColorExpression->SetName(TEXT("Diffuse Color"));
+	ColorExpression->GetColor() = LinearColor;
+	MaterialElement->GetBaseColor().SetExpression(ColorExpression);
+
+	return MaterialElement;
+}
+
+TSharedPtr<IDatasmithUEPbrMaterialElement> CreateUEPbrMaterialFromColor(const FColor& InColor)
+{
+	FString Name = FString::FromInt(BuildColorHash(InColor));
+	FString Label = FString::Printf(TEXT("color_%02x%02x%02x%02x"), InColor.R, InColor.G, InColor.B, InColor.A);
+
+	// Take the Material diffuse color and connect it to the BaseColor of a UEPbrMaterial
+	TSharedRef<IDatasmithUEPbrMaterialElement> MaterialElement = FDatasmithSceneFactory::CreateUEPbrMaterial(*Name);
+	MaterialElement->SetLabel(*Label);
+
+	FLinearColor LinearColor = FLinearColor::FromPow22Color(InColor);
+
+	IDatasmithMaterialExpressionColor* ColorExpression = MaterialElement->AddMaterialExpression<IDatasmithMaterialExpressionColor>();
+	ColorExpression->SetName(TEXT("Diffuse Color"));
+	ColorExpression->GetColor() = LinearColor;
+
+	MaterialElement->GetBaseColor().SetExpression(ColorExpression);
+
+	if (LinearColor.A < 1.0f)
+	{
+		MaterialElement->SetBlendMode(/*EBlendMode::BLEND_Translucent*/2);
+
+		IDatasmithMaterialExpressionScalar* Scalar = MaterialElement->AddMaterialExpression<IDatasmithMaterialExpressionScalar>();
+		Scalar->GetScalar() = LinearColor.A;
+		Scalar->SetName(TEXT("Opacity Level"));
+
+		MaterialElement->GetOpacity().SetExpression(Scalar);
+		MaterialElement->SetParentLabel(TEXT("CAD Transparent Color"));
+	}
+	else
+	{
+		MaterialElement->SetParentLabel(TEXT("CAD Color"));
+	}
+
+	return MaterialElement;
+}
+
+TSharedPtr<IDatasmithUEPbrMaterialElement> CreateUEPbrMaterialFromMaterial(FCADMaterial& InMaterial, TSharedRef<IDatasmithScene> Scene)
+{
+	FString Name = FString::FromInt(BuildMaterialHash(InMaterial));
+
+	// Take the Material diffuse color and connect it to the BaseColor of a UEPbrMaterial
+	TSharedRef<IDatasmithUEPbrMaterialElement> MaterialElement = FDatasmithSceneFactory::CreateUEPbrMaterial(*Name);
+	FString MaterialLabel(InMaterial.MaterialName);
+	if (MaterialLabel.IsEmpty())
+	{
+		MaterialLabel = TEXT("Material");
+	}
+	MaterialElement->SetLabel(*MaterialLabel);
+
+	// Set a diffuse color if there's nothing in the BaseColor
+	if (MaterialElement->GetBaseColor().GetExpression() == nullptr)
+	{
+		FLinearColor LinearColor = FLinearColor::FromPow22Color(InMaterial.Diffuse);
+
+		IDatasmithMaterialExpressionColor* ColorExpression = MaterialElement->AddMaterialExpression<IDatasmithMaterialExpressionColor>();
+		ColorExpression->SetName(TEXT("Diffuse Color"));
+		ColorExpression->GetColor() = LinearColor;
+
+		MaterialElement->GetBaseColor().SetExpression(ColorExpression);
+	}
+
+	if (InMaterial.Transparency > 0.0f)
+	{
+		MaterialElement->SetBlendMode(/*EBlendMode::BLEND_Translucent*/2);
+		IDatasmithMaterialExpressionScalar* Scalar = MaterialElement->AddMaterialExpression<IDatasmithMaterialExpressionScalar>();
+		Scalar->GetScalar() = InMaterial.Transparency;
+		Scalar->SetName(TEXT("Opacity Level"));
+		MaterialElement->GetOpacity().SetExpression(Scalar);
+		MaterialElement->SetParentLabel(TEXT("CAD Transparent Material"));
+	}
+	else 
+	{
+		MaterialElement->SetParentLabel(TEXT("CAD Material"));
+	}
+
+	// Set a Emissive color 
+	if (MaterialElement->GetEmissiveColor().GetExpression() == nullptr)
+	{
+		// Doc CT => TODO
+		//GLfloat Specular[4] = { specular.rgb[0] / 255., specular.rgb[1] / 255., specular.rgb[2] / 255., 1. - transparency };
+		//GLfloat Shininess[1] = { (float)(128 * shininess) };
+		//glMaterialfv(GL_FRONT_AND_BACK, GL_SPECULAR, Specular);
+		//glMaterialfv(GL_FRONT_AND_BACK, GL_SHININESS, Shininess);
+
+		FLinearColor LinearColor = FLinearColor::FromPow22Color(InMaterial.Specular);
+
+		IDatasmithMaterialExpressionColor* ColorExpression = MaterialElement->AddMaterialExpression<IDatasmithMaterialExpressionColor>();
+		ColorExpression->SetName(TEXT("Specular Color"));
+		ColorExpression->GetColor() = LinearColor;
+
+		MaterialElement->GetEmissiveColor().SetExpression(ColorExpression);
+	}
+
+	// Simple conversion of shininess and reflectivity to PBR roughness and metallic values; model could be improved to properly blend the values
+	if (!FMath::IsNearlyZero(InMaterial.Shininess))
+	{
+		IDatasmithMaterialExpressionScalar* Scalar = static_cast<IDatasmithMaterialExpressionScalar*>(MaterialElement->AddMaterialExpression(EDatasmithMaterialExpressionType::ConstantScalar));
+		Scalar->GetScalar() = 1.f - InMaterial.Shininess;
+		Scalar->SetName(TEXT("Shininess Level"));
+		MaterialElement->GetRoughness().SetExpression(Scalar);
+	}
+
+	if (!FMath::IsNearlyZero(InMaterial.Reflexion))
+	{
+		IDatasmithMaterialExpressionScalar* Scalar = static_cast<IDatasmithMaterialExpressionScalar*>(MaterialElement->AddMaterialExpression(EDatasmithMaterialExpressionType::ConstantScalar));
+		Scalar->GetScalar() = InMaterial.Reflexion;
+		Scalar->SetName(TEXT("Reflexion Level"));
+		MaterialElement->GetMetallic().SetExpression(Scalar);
+	}
+
+	return MaterialElement;
+}
+
+uint32 GetSize(CT_TESS_DATA_TYPE type)
+{
+	switch (type)
+	{
+	case CT_TESS_USE_DEFAULT:
+		return sizeof(uint32);
+	case CT_TESS_UBYTE:
+		return sizeof(uint8_t);
+	case CT_TESS_BYTE:
+		return sizeof(int8_t);
+	case CT_TESS_USHORT:
+		return sizeof(int16_t);
+	case CT_TESS_SHORT:
+		return sizeof(uint16_t);
+	case CT_TESS_UINT:
+		return sizeof(uint32);
+	case CT_TESS_INT:
+		return sizeof(int32);
+	case CT_TESS_ULONG:
+		return sizeof(uint64);
+	case CT_TESS_LONG:
+		return sizeof(int64);
+	case CT_TESS_FLOAT:
+		return sizeof(float);
+	case CT_TESS_DOUBLE:
+		return sizeof(double);
+	}
+	return 0;
+}
+
+CT_IO_ERROR Tessellate(CT_OBJECT_ID MainObjectId, const FImportParameters& ImportParams, FMeshDescription& MeshDesc, FMeshParameters& MeshParameters)
+{
+	CheckedCTError Result;
+
+	CT_LIST_IO Objects;
+	Result = CT_COMPONENT_IO::AskChildren(MainObjectId, Objects);
+	if (!Result)
+		return Result;
+
+	SetCoreTechTessellationState(ImportParams);
+
+	TArray<FTessellationData> FaceTessellations;
+	TArray<uint32> FaceColors;
+
+	size_t VertexCount = 0;
+	size_t IndexCount = 0;
+
+	TArray<CT_OBJECT_ID> BodySet;
+	BodySet.Reserve(Objects.Count());
+	while (CT_OBJECT_ID CurrentObj = Objects.IteratorIter())
+	{
+		BodySet.Add(CurrentObj);
+	}
+
+	bool bPreferBodyData = true;
+	bool bTessellated = ConvertCTBodySetToMeshDescription(ImportParams, MeshParameters, BodySet, MeshDesc);
+
+	CheckedCTError ConversionResult;
+	if (!bTessellated)
+	{
+		ConversionResult.RaiseOtherError("Error during mesh conversion");
+	}
+
+	return ConversionResult;
+}
+
+}
+
+#endif // CAD_LIBRARY
\ No newline at end of file
diff --git a/Engine/Plugins/Enterprise/DatasmithCADImporter/Source/CADLibrary/Private/CoreTechMeshLoader.cpp b/Engine/Plugins/Enterprise/DatasmithCADImporter/Source/CADLibrary/Private/CoreTechMeshLoader.cpp
new file mode 100644
--- /dev/null
+++ b/Engine/Plugins/Enterprise/DatasmithCADImporter/Source/CADLibrary/Private/CoreTechMeshLoader.cpp
@@ -0,0 +1,309 @@
+// Copyright 1998-2019 Epic Games, Inc. All Rights Reserved.
+#include "CoreTechMeshLoader.h"
+
+#ifdef CAD_LIBRARY
+#include "CoreTechTypes.h"
+
+namespace CADLibrary
+{
+
+bool CoreTechMeshLoader::LoadFile(const FString& FileName, FMeshDescription& MeshDescription, const FImportParameters& ImportParameters, FMeshParameters& MeshParameters)
+{
+	CoreTechSessionBase Session(TEXT("CoreTechMeshLoader::LoadFile"), ImportParameters.MetricUnit);
+	if (!Session.IsSessionValid())
+	{
+		return false;
+	}
+
+	CT_FLAGS LoadingFlags = CT_LOAD_FLAGS_USE_DEFAULT | CT_LOAD_FLAGS_READ_META_DATA;
+	CT_IO_ERROR Result = IO_OK;
+
+	CT_OBJECT_ID MainObjectID;
+	Result = CTKIO_LoadFile(*FileName, MainObjectID, LoadingFlags);
+
+	// Something wrong happened during the load, abort
+	if (Result != IO_OK)
+	{
+		return false;
+	}
+
+	Result = Repair(MainObjectID, ImportParameters.StitchingTechnique);
+	Result = SetCoreTechTessellationState(ImportParameters);
+
+	bool bNeedSwapOrientation = ImportParameters.StitchingTechnique != EStitchingTechnique::StitchingSew ? MeshParameters.bNeedSwapOrientation : false;
+
+	// Parse CoreTech model's hierarchy to collect body objects.
+	ExtractComponent(MainObjectID);
+
+	// Nothing useful in the file, abort
+	if (BodySet.Num() == 0)
+	{
+		return false;
+	}
+
+	ConvertCTBodySetToMeshDescription(ImportParameters, MeshParameters, BodySet, MeshDescription);
+
+	return true;
+}
+
+void CoreTechMeshLoader::ExtractComponent( CT_OBJECT_ID ObjectID )
+{
+	if( !CT_OBJECT_IO::IsA(ObjectID, CT_COMPONENT_TYPE) )
+	{
+		return;
+	}
+
+	if( ! CT_OBJECT_IO::IsVisible( ObjectID ))
+	{
+		return;
+	}
+
+	CT_LIST_IO ObjectList;
+	CT_IO_ERROR Error = CT_COMPONENT_IO::AskChildren(ObjectID, ObjectList);
+	if (Error != IO_OK || ObjectList.IsEmpty())
+	{
+		return;
+	}
+
+	ObjectList.IteratorInitialize();
+
+	CT_OBJECT_ID ChildID;
+	while( (ChildID = ObjectList.IteratorIter()) != 0 )
+	{
+		if (!CT_OBJECT_IO::IsVisible(ChildID))
+		{
+			continue;
+		}
+
+		if (CT_OBJECT_IO::IsA( ChildID, CT_INSTANCE_TYPE ))
+		{
+			ExtractInstance( ChildID );
+		}
+		else if (CT_OBJECT_IO::IsA( ChildID, CT_SOLID_TYPE ))
+		{
+			ExtractSolid( ChildID );
+		}
+		else if (CT_OBJECT_IO::IsA( ChildID, CT_BODY_TYPE ))
+		{
+			ExtractBody( ChildID );
+		}
+		else if (CT_OBJECT_IO::IsA( ChildID, CT_CURVE_TYPE ))
+		{
+			ExtractCurve( ChildID );
+		}
+		else if (CT_OBJECT_IO::IsA( ChildID, CT_POINT_TYPE ) )
+		{
+			ExtractPoint( ChildID );
+		}
+	}
+}
+
+void CoreTechMeshLoader::ExtractInstance( CT_OBJECT_ID ObjectID )
+{
+	if (!CT_OBJECT_IO::IsA(ObjectID, CT_INSTANCE_TYPE))
+	{
+		return;
+	}
+
+	CT_DOUBLE Matrix[16] =
+	{
+		1., 0., 0., 0.,
+		0., 1., 0., 0.,
+		0., 0., 1., 0.,
+		0., 0., 0., 1.
+	};
+
+	CT_IO_ERROR Error = CT_INSTANCE_IO::AskTransformation( ObjectID, Matrix, CT_MATRIX_ALL );
+
+	CT_OBJECT_ID ChildID;
+	Error = CT_INSTANCE_IO::AskComponentInstanceLevel( ObjectID, ChildID );
+
+	if( Error == IO_OK && ChildID != 0 )
+	{
+		// There is geometry altered by assembly level operations
+	}
+	else
+	{
+		// Get the default children
+		Error = CT_INSTANCE_IO::AskChild( ObjectID, ChildID );
+	}
+
+	if (ChildID == 0)
+	{
+		return;
+	}
+
+	if (CT_OBJECT_IO::IsA( ChildID, CT_UNLOADED_PART_TYPE ))
+	{
+		return;
+	}
+
+	//CHECK_ERROR_NO_LOG(Error);
+	if (Error == IO_OK && ChildID != 0)
+	{
+		ExtractComponent( ChildID);
+	}
+}
+
+void CoreTechMeshLoader::ExtractSolid( CT_OBJECT_ID ObjectID )
+{
+	if (!CT_OBJECT_IO::IsA(ObjectID, CT_SOLID_TYPE))
+	{
+		return;
+	}
+
+	CT_LIST_IO BodyList;
+	CT_IO_ERROR Error = CT_SOLID_IO::AskResultBodies( ObjectID, BodyList );
+	if (Error != IO_OK)
+	{
+		return;
+	}
+
+	if( BodyList.Count() > 0 )
+	{
+		BodyList.IteratorInitialize();
+
+		CT_OBJECT_ID BodyID;
+		while ( (BodyID = BodyList.IteratorIter()) != 0 )
+		{
+			ExtractBody(BodyID);
+		}
+	}
+	else
+	{
+		const CT_UINT32 NodeCount = CT_SOLID_IO::AskNodesCount( ObjectID );
+
+		for( CT_UINT32 NodeIndex = 0; NodeIndex < NodeCount; ++NodeIndex )
+		{
+			CT_OBJECT_ID NodeID;
+			if( CT_SOLID_IO::AskIthNode( ObjectID, NodeIndex, NodeID) == IO_OK )
+			{
+				ExtractBranch( NodeID );
+			}
+		}
+	}
+}
+
+void CoreTechMeshLoader::ExtractBranch( CT_OBJECT_ID ObjectID )
+{
+	if (!CT_OBJECT_IO::IsA(ObjectID, CT_HBRANCH_TYPE))
+	{
+		return;
+	}
+
+	CT_UINT32 NodeCount = CT_HBRANCH_IO::AskNodesCount( ObjectID );
+
+	for( CT_UINT32 NodeIndex = 0; NodeIndex < NodeCount; ++NodeIndex )
+	{
+		CT_OBJECT_ID NodeID;
+		if( CT_HBRANCH_IO::AskIthNode( ObjectID, NodeIndex, NodeID ) == IO_OK )
+		{
+			if( CT_OBJECT_IO::IsA( NodeID, CT_HBRANCH_TYPE ) )
+			{
+				ExtractBranch( NodeID );
+			}
+			else
+			{
+				ExtractLeaf( NodeID );
+			}
+		}
+	}
+}
+
+void CoreTechMeshLoader::ExtractLeaf( CT_OBJECT_ID ObjectID )
+{
+	if (!CT_OBJECT_IO::IsA(ObjectID, CT_HLEAF_TYPE))
+	{
+		return;
+	}
+
+	if( CT_HLEAF_IO::IsOfType( ObjectID, "GEOMETRY" ) )
+	{
+		CT_OBJECT_ID  GeomID = 0;
+
+		CT_IO_ERROR Error = CT_HLEAF_IO::AskGeometry( ObjectID, GeomID );
+
+		if( Error == IO_OK && GeomID != 0 )
+		{
+			if (CT_OBJECT_IO::IsA(GeomID, CT_BODY_TYPE))
+			{
+				ExtractBody( GeomID );
+			}
+			else if( CT_OBJECT_IO::IsA( GeomID, CT_SHELL_TYPE ) )
+			{
+				ExtractShell( GeomID );
+			}
+			else if( CT_OBJECT_IO::IsA( GeomID, CT_FACE_TYPE ) )
+			{
+				ExtractFace( GeomID );
+			}
+			else if( CT_OBJECT_IO::IsA( GeomID, CT_CURVE_TYPE ) )
+			{
+				ExtractCurve( GeomID );
+			}
+			else if( CT_OBJECT_IO::IsA( GeomID, CT_POINT_TYPE ) )
+			{
+				ExtractPoint( GeomID );
+			}
+		}
+	}
+}
+
+void CoreTechMeshLoader::ExtractBody( CT_OBJECT_ID ObjectID )
+{
+	if (!CT_OBJECT_IO::IsA(ObjectID, CT_BODY_TYPE))
+	{
+		return;
+	}
+
+	BodySet.Add(ObjectID);
+}
+
+void CoreTechMeshLoader::ExtractShell( CT_OBJECT_ID ObjectID )
+{
+	if (!CT_OBJECT_IO::IsA(ObjectID, CT_SHELL_TYPE))
+	{
+		return;
+	}
+}
+
+void CoreTechMeshLoader::ExtractFace( CT_OBJECT_ID ObjectID )
+{
+	if (!CT_OBJECT_IO::IsA(ObjectID, CT_FACE_TYPE))
+	{
+		return;
+	}
+
+}
+
+void CoreTechMeshLoader::ExtractLoop( CT_OBJECT_ID ObjectID )
+{
+	if (!CT_OBJECT_IO::IsA(ObjectID, CT_LOOP_TYPE))
+	{
+		return;
+	}
+}
+
+void CoreTechMeshLoader::ExtractCoedge( CT_OBJECT_ID ObjectID )
+{
+	if (!CT_OBJECT_IO::IsA(ObjectID, CT_COEDGE_TYPE))
+	{
+		return;
+	}
+}
+
+void CoreTechMeshLoader::ExtractCurve( CT_OBJECT_ID ObjectID )
+{
+	if (!CT_OBJECT_IO::IsA(ObjectID, CT_CURVE_TYPE))
+	{
+		return;
+	}
+}
+
+void CoreTechMeshLoader::ExtractPoint( CT_OBJECT_ID ObjectID )
+{
+}
+
+} // ns CADLibrary
+
+#endif // CAD_LIBRARY
diff --git a/Engine/Plugins/Enterprise/DatasmithCADImporter/Source/CADLibrary/Private/MeshEditingWrapper.cpp b/Engine/Plugins/Enterprise/DatasmithCADImporter/Source/CADLibrary/Private/MeshEditingWrapper.cpp
new file mode 100644
--- /dev/null
+++ b/Engine/Plugins/Enterprise/DatasmithCADImporter/Source/CADLibrary/Private/MeshEditingWrapper.cpp
@@ -0,0 +1,304 @@
+// Copyright 1998-2019 Epic Games, Inc. All Rights Reserved.
+
+
+#include "MeshEditingWrapper.h"
+
+#include "MeshAttributes.h"
+#include "StaticMeshAttributes.h"
+
+
+
+
+FMeshEditingWrapper::FMeshEditingWrapper(FMeshDescription& InMeshDescription)
+	: MeshDescription(InMeshDescription)
+{
+	UpdateMeshWrapper();
+}
+
+void FMeshEditingWrapper::UpdateMeshWrapper()
+{
+	VertexMetaData.SetNumZeroed(MeshDescription.Vertices().GetArraySize());
+	EdgeMetaData.SetNumZeroed(MeshDescription.Edges().GetArraySize());
+	TriangleMetaData.SetNumZeroed(MeshDescription.Triangles().GetArraySize());
+
+	// If there are more than one polygon groups, edges between groups are considered feature lines
+	for (const FEdgeID EdgeID : MeshDescription.Edges().GetElementIDs())
+	{
+		DefineEdgeTopology(EdgeID);
+	}
+
+	for (const FVertexID VertexID : MeshDescription.Vertices().GetElementIDs())
+	{
+		DefineVertexTopologyApproximation(VertexID);
+	}
+}
+
+
+void FMeshEditingWrapper::DefineEdgeTopology(FEdgeID EdgeID)
+{
+	if (EdgeID == FEdgeID::Invalid)
+	{
+		return;
+	}
+
+	const TArray<FTriangleID>& EdgeConnectedPolygons = MeshDescription.GetEdgeConnectedTriangles(EdgeID);
+
+	switch (EdgeConnectedPolygons.Num())
+	{
+	case 0:
+		EdgeMetaData[EdgeID.GetValue()].Category = (uint16)MeshCategory::EElementCategory::ElementCategoryLine;
+		break;
+	case 1:
+		EdgeMetaData[EdgeID.GetValue()].Category = (uint16)MeshCategory::EElementCategory::ElementCategoryBorder;
+		break;
+	case 2:
+		EdgeMetaData[EdgeID.GetValue()].Category = (uint16)MeshCategory::EElementCategory::ElementCategorySurface;
+		break;
+	default:
+		EdgeMetaData[EdgeID.GetValue()].Category = (uint16)MeshCategory::EElementCategory::ElementCategoryNonManifold;
+		break;
+	}
+}
+
+void FMeshEditingWrapper::DefineVertexTopologyApproximation(FVertexID VertexID)
+{
+	if (VertexID == FVertexID::Invalid)
+	{
+		return;
+	}
+
+	const TArray<FEdgeID>& VertexConnectedEdgeIDs = MeshDescription.GetVertexConnectedEdges(VertexID);
+
+	switch (VertexConnectedEdgeIDs.Num())
+	{
+	case 0:
+		VertexMetaData[VertexID.GetValue()].Category = (uint16)MeshCategory::EElementCategory::ElementCategoryFree;
+		break;
+
+	case 1:
+		VertexMetaData[VertexID.GetValue()].Category = (uint16)MeshCategory::EElementCategory::ElementCategoryBorder;
+		break;
+	default:
+	{
+		int CountPerCategory[(uint16)MeshCategory::EElementCategory::ElementCategoryMax] = { 0 };
+
+		for (const FEdgeID EdgeID : VertexConnectedEdgeIDs)
+		{
+			CountPerCategory[EdgeMetaData[EdgeID.GetValue()].Category]++;
+		}
+
+		if (CountPerCategory[(uint16)MeshCategory::EElementCategory::ElementCategoryNonManifold] > 0)
+		{
+			VertexMetaData[VertexID.GetValue()].Category = (uint16)MeshCategory::EElementCategory::ElementCategoryNonManifold;
+		}
+		else if (CountPerCategory[(uint16)MeshCategory::EElementCategory::ElementCategoryBorder] > 0)
+		{
+			VertexMetaData[VertexID.GetValue()].Category = (uint16)MeshCategory::EElementCategory::ElementCategoryNonSurface;
+
+		}
+		else if (CountPerCategory[(uint16)MeshCategory::EElementCategory::ElementCategoryLine] > 0)
+		{
+			if (CountPerCategory[(uint16)MeshCategory::EElementCategory::ElementCategoryLine] == 2 && VertexConnectedEdgeIDs.Num() == 2)
+			{
+				VertexMetaData[VertexID.GetValue()].Category = (uint16)MeshCategory::EElementCategory::ElementCategoryLine;
+			}
+			else
+			{
+				VertexMetaData[VertexID.GetValue()].Category = (uint16)MeshCategory::EElementCategory::ElementCategoryNonManifold;
+			}
+		}
+		else if (CountPerCategory[(uint16)MeshCategory::EElementCategory::ElementCategorySurface] > 0)
+		{
+			const int32 EdgeCount = VertexConnectedEdgeIDs.Num();
+			const FEdgeID FirstEdgeID(VertexConnectedEdgeIDs[0]);
+			FEdgeID EdgeID(FirstEdgeID);
+
+			FTriangleID TriangleID(FPolygonID::Invalid);
+			int32 TriangleCount = 0;
+
+			do
+			{
+				const TArray<FTriangleID>& EdgeConnectedPolygons = MeshDescription.GetEdgeConnectedTriangles(EdgeID);
+				// Border edge no more triangles to process, exit from the loop
+				if (EdgeConnectedPolygons.Num() < 2)
+				{
+					break;
+				}
+
+				TriangleID = (TriangleID == EdgeConnectedPolygons[0]) ? EdgeConnectedPolygons[1] : EdgeConnectedPolygons[0];
+				++TriangleCount;
+
+				TStaticArray<FEdgeID, 3> TriangleEdges;
+				MeshDescription.GetTriangleEdges(TriangleID, TriangleEdges);
+
+				for (int32 Corner = 0; Corner < 3; ++Corner)
+				{
+					if (EdgeID != TriangleEdges[Corner])
+					{
+						const FVertexID EdgeVertexID0 = MeshDescription.GetEdgeVertex(TriangleEdges[Corner], 0);
+						const FVertexID EdgeVertexID1 = MeshDescription.GetEdgeVertex(TriangleEdges[Corner], 1);
+						if (EdgeVertexID0 == VertexID || EdgeVertexID1 == VertexID)
+						{
+							EdgeID = TriangleEdges[Corner];
+							break;
+						}
+					}
+				}
+			} while (EdgeID != FirstEdgeID);
+
+			VertexMetaData[VertexID.GetValue()].Category = TriangleCount == EdgeCount ? (uint16)MeshCategory::EElementCategory::ElementCategorySurface : (uint16)MeshCategory::EElementCategory::ElementCategoryNonManifold;
+		}
+		else
+		{
+			// we don`t need to know exactly if the node is a border or non manifold node. Non surface is enough
+			VertexMetaData[VertexID.GetValue()].Category = (uint16)MeshCategory::EElementCategory::ElementCategoryNonSurface;
+		}
+	}
+	break;
+	}
+}
+
+FMeshEditingWrapper::~FMeshEditingWrapper()
+{
+
+}
+
+void FMeshEditingWrapper::ResetTriangleMarkerRecursively(FTriangleID Triangle)
+{
+	MeshEditingWrapperUtils::ResetElementMarker(TriangleMetaData[Triangle.GetValue()], ELEMENT_FIRST_MARKER_MASK);
+	TStaticArray<FEdgeID, 3> EdgeSet;
+	MeshDescription.GetTriangleEdges(Triangle, EdgeSet);
+	for (int32 IEdge = 0; IEdge < 3; IEdge++)
+	{
+		MeshEditingWrapperUtils::ResetElementMarker(TriangleMetaData[EdgeSet[IEdge].GetValue()], ELEMENT_FIRST_MARKER_MASK);
+	}
+}
+
+bool FMeshEditingWrapper::IsTriangleMarked(FTriangleID Triangle) const
+{
+	return MeshEditingWrapperUtils::IsElementMarkerSet(TriangleMetaData[Triangle.GetValue()], ELEMENT_FIRST_MARKER_MASK);
+}
+
+void FMeshEditingWrapper::SetTriangleMarked(FTriangleID Triangle)
+{
+	return MeshEditingWrapperUtils::SetElementMarker(TriangleMetaData[Triangle.GetValue()], ELEMENT_FIRST_MARKER_MASK);
+}
+
+MeshCategory::EElementCategory FMeshEditingWrapper::GetEdgeCategory(FEdgeID Edge) const
+{
+	return (MeshCategory::EElementCategory) EdgeMetaData[Edge.GetValue()].Category;
+}
+
+bool FMeshEditingWrapper::IsEdgeOfCategory(FEdgeID Edge, MeshCategory::EElementCategory Category) const
+{
+	return EdgeMetaData[Edge.GetValue()].Category == Category;
+}
+
+bool FMeshEditingWrapper::IsVertexOfCategory(FVertexInstanceID VertexInstanceID, MeshCategory::EElementCategory Category) const
+{
+	FVertexID VertexID = MeshDescription.GetVertexInstanceVertex(VertexInstanceID);
+	return VertexMetaData[VertexID.GetValue()].Category == Category;
+}
+
+// Triangle
+void FMeshEditingWrapper::GetTriangleBoundingBox(FTriangleID Triangle, FVector& MinCorner, FVector& MaxCorner, FVertexInstanceID HighestVertex[3], FVertexInstanceID LowestVertex[3]) const
+{
+	TArrayView < const FVertexInstanceID> Vertices = MeshDescription.GetTriangleVertexInstances(Triangle);
+	for (int32 Index = 0; Index < 3; Index++)
+	{
+		GetVertexBoundingBox(Vertices[Index], MinCorner, MaxCorner, HighestVertex, LowestVertex);
+	}
+}
+
+void FMeshEditingWrapper::GetVertexBoundingBox(FVertexInstanceID VertexInstanceID, FVector& MinCorner, FVector& MaxCorner, FVertexInstanceID HighestVertex[3], FVertexInstanceID LowestVertex[3]) const
+{
+	FVertexID VertexID = MeshDescription.GetVertexInstanceVertex(VertexInstanceID);
+	const FVector VertexPosition = MeshDescription.VertexAttributes().GetAttribute<FVector>(VertexID, MeshAttribute::Vertex::Position);
+
+	if (MaxCorner[0] < VertexPosition[0])
+	{
+		HighestVertex[0] = VertexInstanceID;
+		MaxCorner[0] = VertexPosition[0];
+	}
+	if (MaxCorner[1] < VertexPosition[1])
+	{
+		HighestVertex[1] = VertexInstanceID;
+		MaxCorner[1] = VertexPosition[1];
+	}
+	if (MaxCorner[2] < VertexPosition[2])
+	{
+		HighestVertex[2] = VertexInstanceID;
+		MaxCorner[2] = VertexPosition[2];
+	}
+
+	if (MinCorner[0] > VertexPosition[0])
+	{
+		LowestVertex[0] = VertexInstanceID;
+		MinCorner[0] = VertexPosition[0];
+	}
+	if (MinCorner[1] > VertexPosition[1])
+	{
+		LowestVertex[1] = VertexInstanceID;
+		MinCorner[1] = VertexPosition[1];
+	}
+	if (MinCorner[2] > VertexPosition[2])
+	{
+		LowestVertex[2] = VertexInstanceID;
+		MinCorner[2] = VertexPosition[2];
+	}
+}
+
+void FMeshEditingWrapper::SwapTriangleOrientation(FTriangleID Triangle)
+{
+	FTriangleArray& TraingleSet = MeshDescription.Triangles();
+	FVertexInstanceID Vertex0ID = TraingleSet[FTriangleID(Triangle)].GetVertexInstanceID(0);
+	FVertexInstanceID Vertex1ID = TraingleSet[FTriangleID(Triangle)].GetVertexInstanceID(1);
+	TraingleSet[Triangle].SetVertexInstanceID(0, Vertex1ID);
+	TraingleSet[Triangle].SetVertexInstanceID(1, Vertex0ID);
+
+	for (int32 IVertex = 0; IVertex < 3; IVertex++)
+	{
+		SwapVertexNormal(TraingleSet[FTriangleID(Triangle)].GetVertexInstanceID(IVertex));
+	}
+}
+
+// Edge
+const FTriangleID FMeshEditingWrapper::GetOtherTriangleAtEdge(FEdgeID EdgeID, FTriangleID Triangle) const
+{
+	const TArray<FTriangleID>& EdgeConnectedPolygons = MeshDescription.GetEdgeConnectedTriangles(EdgeID);
+	return EdgeConnectedPolygons.Num() < 2 ? FTriangleID::Invalid : (EdgeConnectedPolygons[0] == Triangle ? EdgeConnectedPolygons[1] : EdgeConnectedPolygons[0]);
+}
+
+bool FMeshEditingWrapper::GetEdgeDirectionInTriangle(FEdgeID EdgeID, int32 TriangleIndex) const
+{
+	const TArray<FTriangleID>& EdgeConnectedTriangles = MeshDescription.GetEdgeConnectedTriangles(EdgeID);
+
+	if (EdgeConnectedTriangles.Num() > TriangleIndex)
+	{
+		const FTriangleID TriangleID = EdgeConnectedTriangles[TriangleIndex];
+		const TArrayView<const FVertexInstanceID>& VertexInstanceIDs = MeshDescription.GetTriangleVertexInstances(TriangleID);
+
+		TStaticArray<FEdgeID, 3> TriangleEdges = MeshDescription.GetTriangleEdges(TriangleID);
+
+		for (int32 Corner = 0; Corner < 3; ++Corner)
+		{
+			if (TriangleEdges[Corner] == EdgeID)
+			{
+				const FVertexID VertexID = MeshDescription.GetVertexInstanceVertex(VertexInstanceIDs[Corner]);
+				return MeshDescription.GetEdgeVertex(EdgeID, 0) == VertexID ? true : false;
+			}
+		}
+	}
+	return true;
+}
+
+// Vertex
+void FMeshEditingWrapper::SwapVertexNormal(FVertexInstanceID VertexInstanceID)
+{
+	FStaticMeshAttributes StaticMeshAttributes( MeshDescription );
+	FVector Normal = StaticMeshAttributes.GetVertexInstanceNormals()[ VertexInstanceID ];
+	Normal *= -1;
+
+	StaticMeshAttributes.GetVertexInstanceNormals()[ VertexInstanceID ] = Normal;
+}
+
diff --git a/Engine/Plugins/Enterprise/DatasmithCADImporter/Source/CADLibrary/Private/MeshEditingWrapper.h b/Engine/Plugins/Enterprise/DatasmithCADImporter/Source/CADLibrary/Private/MeshEditingWrapper.h
new file mode 100644
--- /dev/null
+++ b/Engine/Plugins/Enterprise/DatasmithCADImporter/Source/CADLibrary/Private/MeshEditingWrapper.h
@@ -0,0 +1,145 @@
+// Copyright 1998-2019 Epic Games, Inc. All Rights Reserved.
+
+#pragma once
+
+
+#include "CoreMinimal.h"
+
+#include "MeshTypes.h"
+#include "MeshDescription.h"
+
+
+#define EDGE_IS_HARD      0x01
+#define EDGE_IS_UV_SEAM   0x02
+
+#define ELEMENT_STATUS_MASK			0x01
+#define ELEMENT_FIRST_MARKER_MASK	0x02
+#define ELEMENT_SECOND_MARKER_MASK	0x04
+
+#define ELEMENT_CRITICAL_ZONE_MASK		0x01
+#define ELEMENT_PARTITION_BORDER_MASK	0x02
+
+struct FElementMetaData
+{
+	uint16 Category:4;
+	uint16 Markers:4;
+	uint16 Extras:4;
+};
+
+namespace MeshCategory
+{
+	enum EElementCategory : uint8_t
+	{
+		ElementCategoryUnused = 0,
+		ElementCategoryFree = 1,
+		ElementCategoryLine = 2,
+		ElementCategorySurface = 3,
+		ElementCategoryBorder = 4,
+		ElementCategoryNonManifold = 5,
+		ElementCategoryNonSurface = 6,
+		ElementCategoryMax = 7
+	};
+}
+
+namespace MeshEditingWrapperUtils
+{
+	const FName Debug( "Debug" );
+	const FName FeatureLine( "FeatureLine" );
+	const FName EdgeLength( "EdgeLength" );
+
+	FORCEINLINE bool IsElementMarkerSet(const FElementMetaData& ElementMetaData, const uint16 Mask)
+	{
+		return (ElementMetaData.Markers & Mask) == Mask;
+	}
+
+	FORCEINLINE void SetElementMarker(FElementMetaData& ElementMetaData, const uint16 Mask)
+	{
+		ElementMetaData.Markers |= Mask;
+	}
+
+	FORCEINLINE void ResetElementMarker(FElementMetaData& ElementMetaData, const uint16 Mask)
+	{
+		ElementMetaData.Markers &= ~Mask;
+	}
+
+	FORCEINLINE bool IsElementExtraSet(const FElementMetaData& ElementMetaData, const uint16 Mask)
+	{
+		return (ElementMetaData.Extras & Mask) == Mask;
+	}
+
+	FORCEINLINE void SetElementExtra(FElementMetaData& ElementMetaData, const bool Value, const uint16 Mask)
+	{
+		if (Value)
+		{
+			ElementMetaData.Extras |= Mask;
+		}
+		else
+		{
+			ElementMetaData.Extras &= ~Mask;
+		}
+	}
+
+	FORCEINLINE void ResetElementExtra(FElementMetaData& ElementMetaData, const uint16 Mask)
+	{
+		ElementMetaData.Extras &= ~Mask;
+	}
+
+	FORCEINLINE void ResetElementData(FElementMetaData& ElementMetaData)
+	{
+		ElementMetaData.Category = (uint16)MeshCategory::EElementCategory::ElementCategoryUnused;
+		ElementMetaData.Markers = 0;
+		ElementMetaData.Extras = 0;
+	}
+}
+
+class FMeshEditingWrapper
+{
+
+private:
+	FMeshDescription& MeshDescription;
+
+	TArray<FElementMetaData> VertexMetaData;
+	TArray<FElementMetaData> EdgeMetaData;
+	TArray<FElementMetaData> TriangleMetaData;
+
+private:
+
+
+public:
+	FMeshEditingWrapper(FMeshDescription& InMeshDescription);
+	 ~FMeshEditingWrapper();
+
+	 void UpdateMeshWrapper();
+
+	 void ResetTriangleMarkerRecursively(FTriangleID Triangle);
+	 bool IsTriangleMarked(FTriangleID Triangle) const;
+	 void SetTriangleMarked(FTriangleID Triangle);
+
+	 MeshCategory::EElementCategory GetEdgeCategory(FEdgeID Edge) const;
+	 bool IsEdgeOfCategory(FEdgeID Edge, MeshCategory::EElementCategory Category) const;
+	 bool IsVertexOfCategory(FVertexInstanceID Vertex, MeshCategory::EElementCategory Category) const;
+
+	 // Triangle
+	 void GetTriangleBoundingBox(FTriangleID Triangle, FVector& MinCorner, FVector& MaxCorner, FVertexInstanceID HighestVertex[3], FVertexInstanceID LowestVertex[3]) const;
+
+	 void SwapTriangleOrientation(FTriangleID Triangle);
+
+	 // Edge
+	 const FTriangleID GetOtherTriangleAtEdge(FEdgeID Edge, FTriangleID Triangle) const;
+
+	 /**
+	  * @param Edge           The edge defined by its FEdgeID to get direction in its connected triangle 
+	  * @param TriangleIndex  The triangle connected to the edge, 0 for the first, 1 for the second, ...
+	  *
+	  */
+	 bool GetEdgeDirectionInTriangle(FEdgeID Edge, int32 TriangleIndex) const;
+	 void DefineEdgeTopology(FEdgeID EdgeID);
+
+	 // Vertex
+	 void SwapVertexNormal(FVertexInstanceID VertexID);
+	 void GetVertexBoundingBox(FVertexInstanceID Vertex, FVector& MinCorner, FVector& MaxCorner, FVertexInstanceID HighestVertex[3], FVertexInstanceID LowestVertex[3]) const;
+
+	 void DefineVertexTopologyApproximation(FVertexID VertexID);
+};
+
+
diff --git a/Engine/Plugins/Enterprise/DatasmithCADImporter/Source/CADLibrary/Private/MeshOperator.cpp b/Engine/Plugins/Enterprise/DatasmithCADImporter/Source/CADLibrary/Private/MeshOperator.cpp
new file mode 100644
--- /dev/null
+++ b/Engine/Plugins/Enterprise/DatasmithCADImporter/Source/CADLibrary/Private/MeshOperator.cpp
@@ -0,0 +1,193 @@
+// Copyright 1998-2019 Epic Games, Inc. All Rights Reserved.
+
+#include "MeshOperator.h"
+
+#include "MeshEditingWrapper.h"
+
+#include "Containers/Array.h"
+#include "Containers/Queue.h"
+#include "MeshAttributes.h"
+#include "MeshElementArray.h"
+#include "StaticMeshAttributes.h"
+
+using namespace MeshOperator;
+using namespace MeshCategory;
+
+bool MeshOperator::OrientMesh(FMeshDescription& MeshDescription)
+{
+	FMeshEditingWrapper MeshWrapper(MeshDescription);
+
+	TQueue<FTriangleID> Front; 
+	FTriangleID AdjacentTriangle;
+	TStaticArray<FEdgeID, 3> EdgeSet;
+	FEdgeID Edge;
+
+	TArrayView < const FVertexInstanceID> Vertices;
+	int32 NbSwappedTriangles = 0;
+
+	FVector MaxCorner;
+	FVector MinCorner;
+	FVertexInstanceID HighestVertex[3];
+	FVertexInstanceID LowestVertex[3];
+
+	uint32 NbTriangles = MeshDescription.Triangles().Num();
+
+	TArray<FTriangleID> ConnectedTriangles;
+	ConnectedTriangles.Reserve(NbTriangles);
+
+	for (FTriangleID Triangle : MeshDescription.Triangles().GetElementIDs())
+	{
+		if (MeshWrapper.IsTriangleMarked(Triangle))
+		{
+			continue;
+		}
+
+		HighestVertex[0] = FVertexInstanceID::Invalid;
+		HighestVertex[1] = FVertexInstanceID::Invalid;
+		HighestVertex[2] = FVertexInstanceID::Invalid;
+		LowestVertex[0] = FVertexInstanceID::Invalid;
+		LowestVertex[1] = FVertexInstanceID::Invalid;
+		LowestVertex[2] = FVertexInstanceID::Invalid;
+
+		MaxCorner[0] = -MAX_flt;
+		MaxCorner[1] = -MAX_flt;
+		MaxCorner[2] = -MAX_flt;
+		MinCorner[0] = MAX_flt;
+		MinCorner[1] = MAX_flt;
+		MinCorner[2] = MAX_flt;
+
+		MeshWrapper.SetTriangleMarked(Triangle);
+
+		MeshWrapper.GetTriangleBoundingBox(Triangle, MinCorner, MaxCorner, HighestVertex, LowestVertex);
+
+		Front.Enqueue(Triangle);
+		ConnectedTriangles.Add(Triangle);
+
+		int32 NbConnectedFaces = 1;
+		int32 NbBorderEdges = 0;
+		int32 NbSurfaceEdges = 0;
+		while (!Front.IsEmpty())
+		{
+			Front.Dequeue(Triangle);
+
+			MeshDescription.GetTriangleEdges(Triangle, EdgeSet);
+
+			for (int32 IEdge = 0; IEdge < 3; IEdge++)
+			{
+				Edge = EdgeSet[IEdge];
+
+				if (!MeshWrapper.IsEdgeOfCategory(Edge, EElementCategory::ElementCategorySurface))
+				{
+					NbBorderEdges++;
+					continue;
+				}
+				NbSurfaceEdges++;
+
+				AdjacentTriangle = MeshWrapper.GetOtherTriangleAtEdge(Edge, Triangle);
+				if (MeshWrapper.IsTriangleMarked(AdjacentTriangle))
+				{
+					continue;
+				}
+				NbConnectedFaces++;
+				ConnectedTriangles.Add(AdjacentTriangle);
+
+				if(MeshWrapper.GetEdgeDirectionInTriangle(Edge, 0) == MeshWrapper.GetEdgeDirectionInTriangle(Edge, 1))
+				{
+					MeshWrapper.SwapTriangleOrientation(AdjacentTriangle);
+					NbSwappedTriangles++;
+				}
+				MeshWrapper.SetTriangleMarked(AdjacentTriangle);
+				MeshWrapper.GetTriangleBoundingBox(AdjacentTriangle, MinCorner, MaxCorner, HighestVertex, LowestVertex);
+
+				Front.Enqueue(AdjacentTriangle);
+			}
+		}
+
+		// Check if the mesh orientation need to be swapped
+		int NbInverted = 0;
+		int NbNotInverted = 0;
+		if (NbBorderEdges == 0 || NbBorderEdges * 40 < NbSurfaceEdges)  // NbBorderEdges * 40 < nbSurfaceEdges => basic rule to define if a mesh is a surface mesh or if the mesh is a volume mesh with gaps 
+		{
+			// case of volume mesh
+			// A vertex can have many normal (one by vertex instance) e.g. the corner of a box that have 3 normal and could be the highest 
+			// vertex of a mesh. At the highest vertex, a folding of the mesh can give a vertex with two opposite normal.
+			// The normal most parallel to the axis is preferred
+			// To avoid mistake, we check for each highest vertex and trust the majority 
+
+			if (HighestVertex[0] != FVertexInstanceID::Invalid)
+			{
+				FStaticMeshAttributes StaticMeshAttributes( MeshDescription );
+
+				for (int32 VertexIndex = 0; VertexIndex < 3; VertexIndex++)
+				{
+					if (MeshWrapper.IsVertexOfCategory(HighestVertex[VertexIndex], EElementCategory::ElementCategorySurface))
+					{
+						FVertexID VertexID = MeshDescription.GetVertexInstanceVertex(HighestVertex[VertexIndex]);
+						const TArray<FVertexInstanceID>& CoincidentVertexInstanceIdSet = MeshDescription.GetVertexVertexInstances(VertexID);
+						float MaxComponent = 0;
+						for (const FVertexInstanceID& VertexInstanceID : CoincidentVertexInstanceIdSet)
+						{
+							FVector Normal = StaticMeshAttributes.GetVertexInstanceNormals()[ VertexInstanceID ];
+							if (FMath::Abs(MaxComponent) < FMath::Abs(Normal[VertexIndex]))
+							{
+								MaxComponent = Normal[VertexIndex];
+							}
+						}
+
+						if (0 > MaxComponent)
+						{
+							NbInverted++;
+						}
+						else
+						{
+							NbNotInverted++;
+						}
+					}
+
+					if (MeshWrapper.IsVertexOfCategory(LowestVertex[VertexIndex], EElementCategory::ElementCategorySurface))
+					{
+						FVertexID VertexID = MeshDescription.GetVertexInstanceVertex(LowestVertex[VertexIndex]);
+						const TArray<FVertexInstanceID>& CoincidentVertexInstanceIdSet = MeshDescription.GetVertexVertexInstances(VertexID);
+						float MaxComponent = 0;
+						for (const FVertexInstanceID& VertexInstanceID : CoincidentVertexInstanceIdSet)
+						{
+							FVector Normal = MeshDescription.VertexInstanceAttributes().GetAttribute<FVector>(VertexInstanceID, MeshAttribute::VertexInstance::Normal);
+							if (FMath::Abs(MaxComponent) < FMath::Abs(Normal[VertexIndex]))
+							{
+								MaxComponent = Normal[VertexIndex];
+							}
+						}
+
+						if (0 < MaxComponent)
+						{
+							NbInverted++;
+						}
+						else
+						{
+							NbNotInverted++;
+						}
+					}
+				}
+			}
+		}
+		else if (NbSwappedTriangles*2 > NbConnectedFaces)
+		{
+			// case of surface mesh
+			// this means that more triangles of surface shape have been swapped than no swapped, the good orientation has been reversed.
+			// The mesh need to be re swapped
+			NbInverted++;
+		}
+		
+		// if needed swap all the mesh
+		if (NbInverted > NbNotInverted)
+		{
+			for (auto Tri : ConnectedTriangles)
+			{
+				MeshWrapper.SwapTriangleOrientation(Tri);
+			}
+		}
+		ConnectedTriangles.Empty(NbTriangles);
+	}
+
+	return true;
+}
\ No newline at end of file
diff --git a/Engine/Plugins/Enterprise/DatasmithCADImporter/Source/CADLibrary/Private/MeshOperator.h b/Engine/Plugins/Enterprise/DatasmithCADImporter/Source/CADLibrary/Private/MeshOperator.h
new file mode 100644
--- /dev/null
+++ b/Engine/Plugins/Enterprise/DatasmithCADImporter/Source/CADLibrary/Private/MeshOperator.h
@@ -0,0 +1,13 @@
+// Copyright 1998-2019 Epic Games, Inc. All Rights Reserved.
+#pragma once
+
+#include "CoreMinimal.h"
+
+#include "MeshTypes.h"
+#include "MeshDescription.h"
+
+
+namespace MeshOperator
+{
+	bool OrientMesh(FMeshDescription& MeshDescription);
+}
\ No newline at end of file
diff --git a/Engine/Plugins/Enterprise/DatasmithCADImporter/Source/CADLibrary/Public/CoreTechHelper.h b/Engine/Plugins/Enterprise/DatasmithCADImporter/Source/CADLibrary/Public/CoreTechHelper.h
new file mode 100644
--- /dev/null
+++ b/Engine/Plugins/Enterprise/DatasmithCADImporter/Source/CADLibrary/Public/CoreTechHelper.h
@@ -0,0 +1,65 @@
+// Copyright 1998-2019 Epic Games, Inc. All Rights Reserved.
+#pragma once
+
+#ifdef CAD_LIBRARY
+
+#include "CADData.h"
+#include "CADOptions.h"
+#include "CoreTechTypes.h"
+#include "DatasmithUtils.h"
+#include "Math/Matrix.h"
+#include "Math/Vector.h"
+#include "MeshTypes.h"
+
+struct CTMesh;
+struct FMeshDescription;
+class IDatasmithMaterialIDElement;
+class IDatasmithUEPbrMaterialElement;
+class IDatasmithScene;
+
+#define LAST_CT_MATERIAL_ID 0x00ffffff
+
+namespace CADLibrary
+{
+	struct DbgState
+	{
+		CT_OBJECT_ID      main_object;
+		CT_DOUBLE         max_face_sag;     // Maximum sag for face(in current unit)
+		CT_DOUBLE         max_face_length;  // Maximum length for face(in current unit)
+		CT_DOUBLE         max_face_angle;   // Maximum angle for face(in degree)
+		CT_DOUBLE         max_curve_sag;    // Maximum sag for curve(in current unit)
+		CT_DOUBLE         max_curve_length; // Maximum length for curve(in current unit)
+		CT_DOUBLE         max_curve_angle;  // Maximum angle for curve(in degree)S
+		CT_LOGICAL        high_quality;     // High quality flag
+		CT_TESS_DATA_TYPE vertex_type;      // Vertex type
+		CT_TESS_DATA_TYPE normal_type;      // Normal type
+		CT_TESS_DATA_TYPE texture_type;     // Texture type
+		double            model_size;       // Model size (in current unit)
+		double            tolerance;        // Tolerance (in current unit)
+		double            unit;             // Unit (in meter)
+
+		DbgState()
+		{
+			CTKIO_AskMainObject(main_object);
+			CTKIO_AskTesselationParameters(max_face_sag, max_face_length, max_face_angle, max_curve_sag, max_curve_length, max_curve_angle, high_quality, vertex_type, normal_type, texture_type);
+			CTKIO_AskModelSize(model_size);
+			CTKIO_AskTolerance(tolerance);
+			CTKIO_AskUnit(unit);
+		}
+	};
+
+	uint32 GetSize(CT_TESS_DATA_TYPE type);
+
+	CADLIBRARY_API bool ConvertCTBodySetToMeshDescription(const FImportParameters& ImportParams, const FMeshParameters& MeshParameters, TArray<CT_OBJECT_ID>& BodySet, FMeshDescription& OutMeshDescription);
+	CADLIBRARY_API bool ConvertCTBodySetToMeshDescription(const FImportParameters& ImportParams, const FMeshParameters& MeshParameters, uint32 TriangleCount, TArray<FTessellationData>& FacePTessellationSet, TMap<uint32, uint32>& MaterialIdToMaterialHash, FMeshDescription& MeshDescription);
+	CADLIBRARY_API CT_IO_ERROR Tessellate(CT_OBJECT_ID MainObjectId, const FImportParameters& ImportParams, FMeshDescription& Mesh, FMeshParameters& MeshParameters);
+
+	using TColorMap = TMap<uint32, uint32>;
+
+	CADLIBRARY_API TSharedPtr<IDatasmithUEPbrMaterialElement> CreateDefaultUEPbrMaterial();
+	CADLIBRARY_API TSharedPtr<IDatasmithUEPbrMaterialElement> CreateUEPbrMaterialFromColor(const FColor& InColor);
+	CADLIBRARY_API TSharedPtr<IDatasmithUEPbrMaterialElement> CreateUEPbrMaterialFromMaterial(FCADMaterial& InMaterial, TSharedRef<IDatasmithScene> Scene);
+
+} // namespace CADLibrary
+
+#endif // CAD_LIBRARY
diff --git a/Engine/Plugins/Enterprise/DatasmithCADImporter/Source/CADLibrary/Public/CoreTechMeshLoader.h b/Engine/Plugins/Enterprise/DatasmithCADImporter/Source/CADLibrary/Public/CoreTechMeshLoader.h
new file mode 100644
--- /dev/null
+++ b/Engine/Plugins/Enterprise/DatasmithCADImporter/Source/CADLibrary/Public/CoreTechMeshLoader.h
@@ -0,0 +1,50 @@
+// Copyright 1998-2019 Epic Games, Inc. All Rights Reserved.
+
+#pragma once
+
+#ifdef CAD_LIBRARY
+#include "CoreTechTypes.h"
+#include "CoreTechHelper.h"
+
+#include <vector>
+
+class GPMMesh;
+class GPUVSetData;
+class GPMFace;
+class GPMatrixH;
+
+
+namespace CADLibrary
+{
+
+	class CADLIBRARY_API CoreTechMeshLoader
+	{
+	public:
+		/**
+		 * Load a .ct file and populate the CoreTech mesh with the loaded geometry
+		 * @note All input are assumed to be valid
+		 */
+		bool LoadFile(const FString& FileName, FMeshDescription& MeshDescription, const FImportParameters& ImportParameters, FMeshParameters& MeshParameters);
+
+	private:
+		void ExtractComponent( CT_OBJECT_ID ObjectID );
+		void ExtractInstance( CT_OBJECT_ID ObjectID );
+		void ExtractSolid( CT_OBJECT_ID ObjectID );
+		void ExtractBranch( CT_OBJECT_ID ObjectID );
+		void ExtractLeaf( CT_OBJECT_ID ObjectID );
+		void ExtractBody( CT_OBJECT_ID ObjectID );
+		void ExtractShell( CT_OBJECT_ID ObjectID );
+		void ExtractFace( CT_OBJECT_ID ObjectID );
+		void ExtractLoop( CT_OBJECT_ID ObjectID );
+		void ExtractCoedge( CT_OBJECT_ID ObjectID );
+		void ExtractCurve( CT_OBJECT_ID ObjectID );
+		void ExtractPoint( CT_OBJECT_ID ObjectID );
+
+	private:
+		TColorMap ColorMap;
+		TArray<CT_OBJECT_ID> BodySet;
+	};
+
+} // ns CADLibrary
+
+#endif // CAD_LIBRARY
