Improve naming of CoreTechFileParser class naming

diff --git a/Engine/Plugins/Enterprise/DatasmithCADImporter/Source/CADTools/Private/CADData.cpp b/Engine/Plugins/Enterprise/DatasmithCADImporter/Source/CADTools/Private/CADData.cpp
--- a/Engine/Plugins/Enterprise/DatasmithCADImporter/Source/CADTools/Private/CADData.cpp
+++ b/Engine/Plugins/Enterprise/DatasmithCADImporter/Source/CADTools/Private/CADData.cpp
@@ -49,7 +49,7 @@ int32 BuildMaterialHash(const FCADMaterial& Material)
 }
 
 
-FBody::FBody(uint32 InBodyID, int32 FaceNum)
+FBodyMesh::FBodyMesh(uint32 InBodyID, int32 FaceNum)
 	: TriangleCount(0)
 	, BodyID(InBodyID)
 {
@@ -57,17 +57,17 @@ FBody::FBody(uint32 InBodyID, int32 FaceNum)
 }
 
 
-FCTRawGeomFile::FCTRawGeomFile(FString& InFileName, TMap< uint32, FBody* >& InBodyUuidToCTBodyMap)
+FCADMeshFile::FCADMeshFile(FString& InFileName, TMap< uint32, FBodyMesh* >& InBodyUuidToCTBodyMap)
 	: BodyUuidToBodyMap(InBodyUuidToCTBodyMap)
 {
 	uint32 NbBody = 0;
 
-	TArray<uint8> RawData;
-	FFileHelper::LoadFileToArray(RawData, *InFileName);
+	TArray<uint8> MeshData;
+	FFileHelper::LoadFileToArray(MeshData, *InFileName);
 
-	ReadTessellationSetFromFile(RawData, NbBody, CTBodySet);
+	ReadTessellationSetFromFile(MeshData, NbBody, CTBodySet);
 
-	for(FBody& Body : CTBodySet)
+	for(FBodyMesh& Body : CTBodySet)
 	{
 		BodyUuidToBodyMap.Add(Body.GetBodyUuid(), &Body);
 	}
@@ -118,7 +118,7 @@ void WriteTessellationInRawData(FTessellationData& Tessellation, TArray<uint8>&
 	GlobalRawData.Append(Tessellation.TexCoordArray);
 }
 
-bool ReadTessellationSetFromFile(TArray<uint8>& RawData, uint32& OutNbBodies, TArray<FBody>& CTBodySet)
+bool ReadTessellationSetFromFile(TArray<uint8>& RawData, uint32& OutNbBodies, TArray<FBodyMesh>& CTBodySet)
 {
 	uint32 INFO[LastLine];
 	const uint32 InfoSize = sizeof(INFO);
@@ -127,7 +127,7 @@ bool ReadTessellationSetFromFile(TArray<uint8>& RawData, uint32& OutNbBodies, TA
 
 	CTBodySet.Reserve(OutNbBodies);
 
-	FBody* CurrentBody = nullptr;
+	FBodyMesh* CurrentBody = nullptr;
 
 	int32 CurrentBodyId = -1;
 	int32 CurrentBodyIndex = -1;
diff --git a/Engine/Plugins/Enterprise/DatasmithCADImporter/Source/CADTools/Public/CADData.h b/Engine/Plugins/Enterprise/DatasmithCADImporter/Source/CADTools/Public/CADData.h
--- a/Engine/Plugins/Enterprise/DatasmithCADImporter/Source/CADTools/Public/CADData.h
+++ b/Engine/Plugins/Enterprise/DatasmithCADImporter/Source/CADTools/Public/CADData.h
@@ -80,10 +80,10 @@ struct CADTOOLS_API FTessellationData
 	TArray<int32> SymVertexIdSet; // StaticMesh FVertexID for sym part
 };
 
-class CADTOOLS_API FBody
+class CADTOOLS_API FBodyMesh
 {
 public:
-	FBody(uint32 BodyID, int32 FaceNum);
+	FBodyMesh(uint32 BodyID, int32 FaceNum);
 
 	TArray<FTessellationData>& GetTessellationSet()
 	{
@@ -118,7 +118,6 @@ public:
 protected:
 
 	/**
-	 * Array of FCTTessellation
 	 * FCTTessellation is an elementary structure of Mesh
 	 */
 	TArray<FTessellationData> FaceTessellationSet;
@@ -132,10 +131,10 @@ protected:
 	uint32 BodyUuid;
 };
 
-class CADTOOLS_API FCTRawGeomFile
+class CADTOOLS_API FCADMeshFile
 {
 public:
-	FCTRawGeomFile(FString& InFileName, TMap< uint32, FBody* >& InBodyUuidToCTBodyMap);
+	FCADMeshFile(FString& InFileName, TMap< uint32, FBodyMesh* >& InBodyUuidToCTBodyMap);
 
 	FString& GetFileName()
 	{
@@ -144,8 +143,8 @@ public:
 
 protected:
 	FString FileName;
-	TArray<FBody> CTBodySet;
-	TMap<uint32, FBody* >& BodyUuidToBodyMap;
+	TArray<FBodyMesh> CTBodySet;
+	TMap<uint32, FBodyMesh* >& BodyUuidToBodyMap;
 };
 
 CADTOOLS_API uint32 BuildFastColorHash(uint32 ColorId, uint8 Alpha);
@@ -155,7 +154,7 @@ CADTOOLS_API int32 BuildColorHash(const FColor& Color);
 CADTOOLS_API int32 BuildMaterialHash(const FCADMaterial& Material);
 
 CADTOOLS_API void WriteTessellationInRawData(FTessellationData& Tessellation, TArray<uint8>& GlobalRawData);
-CADTOOLS_API bool ReadTessellationSetFromFile(TArray<uint8>& RawData, uint32& OutNbBody, TArray<FBody>& FaceTessellations);
+CADTOOLS_API bool ReadTessellationSetFromFile(TArray<uint8>& RawData, uint32& OutNbBody, TArray<FBodyMesh>& FaceTessellations);
 
 inline void CopyValue(const uint8* Source, int Offset, uint8 Size, bool bIs3D, FVector& Dest, const FMatrix* Matrix = nullptr)
 {
diff --git a/Engine/Plugins/Enterprise/DatasmithCADImporter/Source/DatasmithCADTranslator/Private/DatasmithCADTranslator.cpp b/Engine/Plugins/Enterprise/DatasmithCADImporter/Source/DatasmithCADTranslator/Private/DatasmithCADTranslator.cpp
--- a/Engine/Plugins/Enterprise/DatasmithCADImporter/Source/DatasmithCADTranslator/Private/DatasmithCADTranslator.cpp
+++ b/Engine/Plugins/Enterprise/DatasmithCADImporter/Source/DatasmithCADTranslator/Private/DatasmithCADTranslator.cpp
@@ -12,7 +12,6 @@
 
 
 FDatasmithCADTranslator::FDatasmithCADTranslator()
-	: MeshBuilder(CADFileToUE4GeomMap, MeshElementToCADBRepUuidMap)
 {
 }
 
@@ -86,6 +85,8 @@ bool FDatasmithCADTranslator::LoadScene(TSharedRef<IDatasmithScene> DatasmithSce
 
 	FString CachePath = FPaths::ConvertRelativePathToFull(FDatasmithCADTranslatorModule::Get().GetCacheDir());
 
+	TMap<FString, FString> CADFileToUE4FileMap;
+
 	int32 NumCores = FPlatformMisc::NumberOfCores();
 	DatasmithDispatcher::FDatasmithDispatcher Dispatcher(CachePath, NumCores, CADFileToUE4FileMap, CADFileToUE4GeomMap);
 
@@ -98,24 +99,30 @@ bool FDatasmithCADTranslator::LoadScene(TSharedRef<IDatasmithScene> DatasmithSce
 
 	FDatasmithSceneGraphBuilder SceneGraphBuilder(CADFileToUE4FileMap, MeshElementToCADBRepUuidMap,  CachePath, DatasmithScene, GetSource(), ImportParameters);
 	SceneGraphBuilder.Build();
-	SceneGraphBuilder.Clear();
 
-	MeshBuilder.Init(CachePath);
-	MeshBuilder.SetImportParameters(ImportParameters);
-	MeshBuilder.LoadRawDataGeom();
+	MeshBuilderPtr = MakeUnique<FDatasmithMeshBuilder>(CADFileToUE4GeomMap, MeshElementToCADBRepUuidMap, CachePath, ImportParameters);
 
 	return true;
 }
 
 void FDatasmithCADTranslator::UnloadScene()
 {
+	MeshBuilderPtr = nullptr;
+
+	CADFileToUE4GeomMap.Empty();
+	MeshElementToCADBRepUuidMap.Empty();
 }
 
 bool FDatasmithCADTranslator::LoadStaticMesh(const TSharedRef<IDatasmithMeshElement> MeshElement, FDatasmithMeshElementPayload& OutMeshPayload)
 {
+	if (!MeshBuilderPtr.IsValid())
+	{
+		return false;
+	}
+
 	CADLibrary::FMeshParameters MeshParameters;
 
-	if (TOptional< FMeshDescription > Mesh = MeshBuilder.GetMeshDescription(MeshElement, MeshParameters))
+	if (TOptional< FMeshDescription > Mesh = MeshBuilderPtr->GetMeshDescription(MeshElement, MeshParameters))
 	{
 		OutMeshPayload.LodMeshes.Add(MoveTemp(Mesh.GetValue()));
 
diff --git a/Engine/Plugins/Enterprise/DatasmithCADImporter/Source/DatasmithCADTranslator/Private/DatasmithCADTranslator.h b/Engine/Plugins/Enterprise/DatasmithCADImporter/Source/DatasmithCADTranslator/Private/DatasmithCADTranslator.h
--- a/Engine/Plugins/Enterprise/DatasmithCADImporter/Source/DatasmithCADTranslator/Private/DatasmithCADTranslator.h
+++ b/Engine/Plugins/Enterprise/DatasmithCADImporter/Source/DatasmithCADTranslator/Private/DatasmithCADTranslator.h
@@ -39,13 +39,12 @@ public:
 	virtual void SetSceneImportOptions(TArray<TStrongObjectPtr<UObject>>& Options) override;
 
 private:
-	TMap<FString, FString> CADFileToUE4FileMap;
 	TMap<FString, FString> CADFileToUE4GeomMap;
 	TMap< TSharedPtr< IDatasmithMeshElement >, uint32 > MeshElementToCADBRepUuidMap;
 
 	CADLibrary::FImportParameters ImportParameters;
 
-	FDatasmithMeshBuilder MeshBuilder;
+	TUniquePtr<FDatasmithMeshBuilder> MeshBuilderPtr;
 #endif
 };
 
diff --git a/Engine/Plugins/Enterprise/DatasmithCADImporter/Source/DatasmithCADTranslator/Private/DatasmithMeshBuilder.cpp b/Engine/Plugins/Enterprise/DatasmithCADImporter/Source/DatasmithCADTranslator/Private/DatasmithMeshBuilder.cpp
--- a/Engine/Plugins/Enterprise/DatasmithCADImporter/Source/DatasmithCADTranslator/Private/DatasmithMeshBuilder.cpp
+++ b/Engine/Plugins/Enterprise/DatasmithCADImporter/Source/DatasmithCADTranslator/Private/DatasmithMeshBuilder.cpp
@@ -15,35 +15,27 @@
 
 using namespace CADLibrary;
 
-FDatasmithMeshBuilder::FDatasmithMeshBuilder(TMap<FString, FString>& InCADFileToUE4GeomMap, TMap< TSharedPtr< IDatasmithMeshElement >, uint32 >& InMeshElementToCTBodyUuidMap)
-	: CADFileToUE4GeomMap(InCADFileToUE4GeomMap)
+FDatasmithMeshBuilder::FDatasmithMeshBuilder(TMap<FString, FString>& InCADFileToUE4GeomMap, TMap< TSharedPtr< IDatasmithMeshElement >, uint32 >& InMeshElementToCTBodyUuidMap, const FString& InCachePath, const CADLibrary::FImportParameters& InImportParameters)
+	: CADFileToMeshFileMap(InCADFileToUE4GeomMap)
 	, MeshElementToBodyUuidMap(InMeshElementToCTBodyUuidMap)
+	, CachePath(InCachePath)
+	, ImportParameters(InImportParameters)
 {
+	LoadMeshFiles();
 }
 
-void FDatasmithMeshBuilder::Init(const FString& InCachePath)
+void FDatasmithMeshBuilder::LoadMeshFiles()
 {
-	CachePath = InCachePath;
-}
-
-void FDatasmithMeshBuilder::Clear()
-{
-	RawDataArray.Empty();
-	BodyUuidToCADBRepMap.Empty();
-}
-
-void FDatasmithMeshBuilder::LoadRawDataGeom()
-{
-	RawDataArray.Reserve(CADFileToUE4GeomMap.Num());
-	for (const auto& FilePair : CADFileToUE4GeomMap)
+	Meshes.Reserve(CADFileToMeshFileMap.Num());
+	for (const auto& FilePair : CADFileToMeshFileMap)
 	{
-		FString RawDataFile = FPaths::Combine(CachePath, TEXT("mesh"), FilePair.Value + TEXT(".gm"));
-		if (!IFileManager::Get().FileExists(*RawDataFile))
+		FString MeshFile = FPaths::Combine(CachePath, TEXT("mesh"), FilePair.Value + TEXT(".gm"));
+		if (!IFileManager::Get().FileExists(*MeshFile))
 		{
 			continue;
 		}
 
-		int32 index = RawDataArray.Emplace(RawDataFile, BodyUuidToCADBRepMap);
+		Meshes.Emplace(MeshFile, BodyUuidToMeshMap);
 	}
 }
 
@@ -56,13 +48,13 @@ TOptional<FMeshDescription> FDatasmithMeshBuilder::GetMeshDescription(TSharedRef
 		return TOptional<FMeshDescription>();
 	}
 
-	FBody*const* PPBody = BodyUuidToCADBRepMap.Find(*BodyUuid);
+	FBodyMesh*const* PPBody = BodyUuidToMeshMap.Find(*BodyUuid);
 	if(PPBody == nullptr || *PPBody == nullptr)
 	{
 		return TOptional<FMeshDescription>();
 	}
 
-	FBody& Body = **PPBody;
+	FBodyMesh& Body = **PPBody;
 
 	TMap<uint32, uint32> MaterialIdToMaterialHash;
 
diff --git a/Engine/Plugins/Enterprise/DatasmithCADImporter/Source/DatasmithCADTranslator/Private/DatasmithMeshBuilder.h b/Engine/Plugins/Enterprise/DatasmithCADImporter/Source/DatasmithCADTranslator/Private/DatasmithMeshBuilder.h
--- a/Engine/Plugins/Enterprise/DatasmithCADImporter/Source/DatasmithCADTranslator/Private/DatasmithMeshBuilder.h
+++ b/Engine/Plugins/Enterprise/DatasmithCADImporter/Source/DatasmithCADTranslator/Private/DatasmithMeshBuilder.h
@@ -15,33 +15,19 @@ class IDatasmithMeshElement;
 class FDatasmithMeshBuilder
 {
 public:
-	FDatasmithMeshBuilder(TMap<FString, FString>& InCADFileToUE4GeomMap, TMap< TSharedPtr< IDatasmithMeshElement >, uint32 >& InMeshElementToCTBodyUuidMap);
-	
-	void Init(const FString& InCachePath);
+	FDatasmithMeshBuilder(TMap<FString, FString>& InCADFileToMeshFileMap, TMap< TSharedPtr< IDatasmithMeshElement >, uint32 >& InMeshElementToCADBodyUuidMap, const FString& InCachePath, const CADLibrary::FImportParameters& InImportParameters);
 
 	TOptional<FMeshDescription> GetMeshDescription(TSharedRef<IDatasmithMeshElement> OutMeshElement, CADLibrary::FMeshParameters& OutMeshParameters);
-	void LoadRawDataGeom();
-	void Clear();
-
-	void SetImportParameters(const CADLibrary::FImportParameters& InImportParameters)
-	{
-		ImportParameters = InImportParameters;
-	}
 
 protected:
-	//const FDatasmithSceneSource& Source;
 	FString CachePath;
 
-	TArray<CADLibrary::FCTRawGeomFile> RawDataArray;
-
-	/** Map linking Cad file to RawGeom file (*.gm) */
-	TMap<FString, FString>& CADFileToUE4GeomMap;
+	void LoadMeshFiles();
 
-	/** Datasmith mesh elements to BodyUuid */
+	TArray<CADLibrary::FCADMeshFile> Meshes;
+	TMap<FString, FString>& CADFileToMeshFileMap;
 	TMap< TSharedPtr< IDatasmithMeshElement >, uint32 >& MeshElementToBodyUuidMap;
-
-	/** BodyUuid to CTBody */
-	TMap< uint32, CADLibrary::FBody* > BodyUuidToCADBRepMap;
+	TMap< uint32, CADLibrary::FBodyMesh* > BodyUuidToMeshMap;
 
 	CADLibrary::FImportParameters ImportParameters;
 };
diff --git a/Engine/Plugins/Enterprise/DatasmithCADImporter/Source/DatasmithCADTranslator/Private/DatasmithSceneGraphBuilder.cpp b/Engine/Plugins/Enterprise/DatasmithCADImporter/Source/DatasmithCADTranslator/Private/DatasmithSceneGraphBuilder.cpp
--- a/Engine/Plugins/Enterprise/DatasmithCADImporter/Source/DatasmithCADTranslator/Private/DatasmithSceneGraphBuilder.cpp
+++ b/Engine/Plugins/Enterprise/DatasmithCADImporter/Source/DatasmithCADTranslator/Private/DatasmithSceneGraphBuilder.cpp
@@ -441,16 +441,15 @@ bool FDatasmithSceneGraphBuilder::Build()
 	TSharedPtr< IDatasmithActorElement > RootActor = BuildNode(*RootNode, Data);
 	DatasmithScene->AddActor(RootActor);
 
-	return true;
-}
 
-void FDatasmithSceneGraphBuilder::Clear()
-{
 	CADFileToSceneGraphDescription.Empty();
 	ArrayOfSceneGraphDescription.Empty();
 	BodyUuidToMeshElementMap.Empty();
+
+	return true;
 }
 
+
 void FDatasmithSceneGraphBuilder::LoadSceneGraphDescriptionFiles()
 {
 	ArrayOfSceneGraphDescription.Reserve(CADFileToSceneGraphDescriptionFile.Num());
diff --git a/Engine/Plugins/Enterprise/DatasmithCADImporter/Source/DatasmithCADTranslator/Private/DatasmithSceneGraphBuilder.h b/Engine/Plugins/Enterprise/DatasmithCADImporter/Source/DatasmithCADTranslator/Private/DatasmithSceneGraphBuilder.h
--- a/Engine/Plugins/Enterprise/DatasmithCADImporter/Source/DatasmithCADTranslator/Private/DatasmithSceneGraphBuilder.h
+++ b/Engine/Plugins/Enterprise/DatasmithCADImporter/Source/DatasmithCADTranslator/Private/DatasmithSceneGraphBuilder.h
@@ -146,8 +146,6 @@ public:
 
 	bool Build();
 
-	void Clear();
-
 protected:
 	void LoadSceneGraphDescriptionFiles();
 
