Incremental UE4Edito.98513976.log:66
Fix warning

diff --git a/Engine/Plugins/Enterprise/DatasmithCADImporter/Source/DatasmithCADTranslator/Private/DatasmithMeshBuilder.cpp b/Engine/Plugins/Enterprise/DatasmithCADImporter/Source/DatasmithCADTranslator/Private/DatasmithMeshBuilder.cpp
--- a/Engine/Plugins/Enterprise/DatasmithCADImporter/Source/DatasmithCADTranslator/Private/DatasmithMeshBuilder.cpp
+++ b/Engine/Plugins/Enterprise/DatasmithCADImporter/Source/DatasmithCADTranslator/Private/DatasmithMeshBuilder.cpp
@@ -16,9 +16,9 @@
 using namespace CADLibrary;
 
 FDatasmithMeshBuilder::FDatasmithMeshBuilder(TMap<FString, FString>& InCADFileToUE4GeomMap, TMap< TSharedPtr< IDatasmithMeshElement >, uint32 >& InMeshElementToCTBodyUuidMap, const FString& InCachePath, const CADLibrary::FImportParameters& InImportParameters)
-	: CADFileToMeshFileMap(InCADFileToUE4GeomMap)
+	: CachePath(InCachePath)
+	, CADFileToMeshFileMap(InCADFileToUE4GeomMap)
 	, MeshElementToBodyUuidMap(InMeshElementToCTBodyUuidMap)
-	, CachePath(InCachePath)
 	, ImportParameters(InImportParameters)
 {
 	LoadMeshFiles();
