Active kernel_IO option to load meta data from CAD file
Remove Metadata used for the import and not consistency due to sequential import

diff --git a/Engine/Plugins/Enterprise/DatasmithCADImporter/Source/CADInterfaces/Private/CoreTechFileParser.cpp b/Engine/Plugins/Enterprise/DatasmithCADImporter/Source/CADInterfaces/Private/CoreTechFileParser.cpp
--- a/Engine/Plugins/Enterprise/DatasmithCADImporter/Source/CADInterfaces/Private/CoreTechFileParser.cpp
+++ b/Engine/Plugins/Enterprise/DatasmithCADImporter/Source/CADInterfaces/Private/CoreTechFileParser.cpp
@@ -852,10 +852,7 @@ CT_FLAGS FCoreTechFileParser::SetCoreTechImportOption(const FString& MainFileExt
 	// Set import option
 	CT_FLAGS Flags = CT_LOAD_FLAGS_USE_DEFAULT;
 
-	if (MainFileExt == TEXT("jt"))
-	{
-		Flags |= CT_LOAD_FLAGS_READ_META_DATA;
-	}
+	Flags |= CT_LOAD_FLAGS_READ_META_DATA;
 
 	if (MainFileExt == TEXT("catpart") || MainFileExt == TEXT("catproduct") || MainFileExt == TEXT("cgr"))
 	{
diff --git a/Engine/Plugins/Enterprise/DatasmithCADImporter/Source/DatasmithCADTranslator/Private/DatasmithSceneGraphBuilder.cpp b/Engine/Plugins/Enterprise/DatasmithCADImporter/Source/DatasmithCADTranslator/Private/DatasmithSceneGraphBuilder.cpp
--- a/Engine/Plugins/Enterprise/DatasmithCADImporter/Source/DatasmithCADTranslator/Private/DatasmithSceneGraphBuilder.cpp
+++ b/Engine/Plugins/Enterprise/DatasmithCADImporter/Source/DatasmithCADTranslator/Private/DatasmithSceneGraphBuilder.cpp
@@ -839,7 +839,9 @@ void FDatasmithSceneGraphBuilder::AddMetaData(TSharedPtr< IDatasmithActorElement
 		UnwantedAttributes.Add(TEXT("OriginalIdStr"));
 		UnwantedAttributes.Add(TEXT("ShowAttribute"));
 		UnwantedAttributes.Add(TEXT("Identification"));
-
+		UnwantedAttributes.Add(TEXT("MaterialId"));
+		UnwantedAttributes.Add(TEXT("ColorUEId"));
+		UnwantedAttributes.Add(TEXT("ColorId"));
 		return UnwantedAttributes;
 	};
 
