Datasmith CAD import solidworks - materials do not get assigned correctly
To prepare a good management of Color & Material according to CAD format,  two new import parameters are defined: EDisplayPreference, EDisplayDataPropagationMode.
In this first version, for solidworks file, EDisplayPreference = ColorOnly and EDisplayDataPropagationMode=BodyOnly
BodyOnly meens that material are not apply on MeshActor

diff --git a/Engine/Plugins/Enterprise/DatasmithCADImporter/Source/CADTools/Public/CADOptions.h b/Engine/Plugins/Enterprise/DatasmithCADImporter/Source/CADTools/Public/CADOptions.h
--- a/Engine/Plugins/Enterprise/DatasmithCADImporter/Source/CADTools/Public/CADOptions.h
+++ b/Engine/Plugins/Enterprise/DatasmithCADImporter/Source/CADTools/Public/CADOptions.h
@@ -21,6 +21,21 @@ namespace CADLibrary
 		YUp_RightHanded,
 	};
 
+	enum class EDisplayPreference : uint8
+	{
+		ColorPrefered,
+		MaterialPrefered,
+		ColorOnly,
+		MaterialOnly,
+	};
+
+	enum class EDisplayDataPropagationMode : uint8
+	{
+		TopDown,
+		BottomUp,
+		BodyOnly,
+	};
+
 	struct FImportParameters
 	{
 		double MetricUnit = 0.001;
@@ -30,6 +45,8 @@ namespace CADLibrary
 		double MaxNormalAngle = 20.0;
 		EStitchingTechnique StitchingTechnique = EStitchingTechnique::StitchingNone;
 		EModelCoordSystem ModelCoordSys = EModelCoordSystem::ZUp_RightHanded;
+		EDisplayDataPropagationMode Propagation = EDisplayDataPropagationMode::TopDown;
+		EDisplayPreference DisplayPreference = EDisplayPreference::MaterialPrefered;
 	};
 
 	struct FMeshParameters
diff --git a/Engine/Plugins/Enterprise/DatasmithCADImporter/Source/DatasmithCADTranslator/Private/DatasmithCADTranslator.cpp b/Engine/Plugins/Enterprise/DatasmithCADImporter/Source/DatasmithCADTranslator/Private/DatasmithCADTranslator.cpp
--- a/Engine/Plugins/Enterprise/DatasmithCADImporter/Source/DatasmithCADTranslator/Private/DatasmithCADTranslator.cpp
+++ b/Engine/Plugins/Enterprise/DatasmithCADImporter/Source/DatasmithCADTranslator/Private/DatasmithCADTranslator.cpp
@@ -81,6 +81,8 @@ bool FDatasmithCADTranslator::LoadScene(TSharedRef<IDatasmithScene> DatasmithSce
 	if (FileExtension == TEXT("sldprt") || FileExtension == TEXT("sldasm"))
 	{
 		ImportParameters.ModelCoordSys = CADLibrary::EModelCoordSystem::YUp_RightHanded;
+		ImportParameters.DisplayPreference = CADLibrary::EDisplayPreference::ColorOnly;
+		ImportParameters.Propagation = CADLibrary::EDisplayDataPropagationMode::BodyOnly;
 	}
 
 	FString CachePath = FPaths::ConvertRelativePathToFull(FDatasmithCADTranslatorModule::Get().GetCacheDir());
diff --git a/Engine/Plugins/Enterprise/DatasmithCADImporter/Source/DatasmithCADTranslator/Private/DatasmithSceneGraphBuilder.cpp b/Engine/Plugins/Enterprise/DatasmithCADImporter/Source/DatasmithCADTranslator/Private/DatasmithSceneGraphBuilder.cpp
--- a/Engine/Plugins/Enterprise/DatasmithCADImporter/Source/DatasmithCADTranslator/Private/DatasmithSceneGraphBuilder.cpp
+++ b/Engine/Plugins/Enterprise/DatasmithCADImporter/Source/DatasmithCADTranslator/Private/DatasmithSceneGraphBuilder.cpp
@@ -691,7 +691,7 @@ TSharedPtr< IDatasmithActorElement > FDatasmithSceneGraphBuilder::BuildBody(FSce
 	uint32 MaterialUuid = 0;
 	MaterialUuid = ParentData.MaterialUuid ? ParentData.MaterialUuid : ParentData.ColorUuid;
 
-	if (MaterialUuid)
+	if (MaterialUuid && ImportParameters.Propagation != CADLibrary::EDisplayDataPropagationMode::BodyOnly)
 	{
 		TSharedPtr< IDatasmithMaterialIDElement > PartMaterialIDElement = FindOrAddMaterial(MaterialUuid);
 		const TCHAR* MaterialIDElementName = PartMaterialIDElement->GetName();
@@ -703,7 +703,6 @@ TSharedPtr< IDatasmithActorElement > FDatasmithSceneGraphBuilder::BuildBody(FSce
 			ActorElement->AddMaterialOverride(MaterialIDElement);
 		}
 	}
-
 	return ActorElement;
 }
 
