Fix CAD Iges file import
the extension was compared with ".igs" instead of "igs"

diff --git a/Engine/Plugins/Enterprise/DatasmithCADImporter/Source/CADInterfaces/Private/CoreTechFileParser.cpp b/Engine/Plugins/Enterprise/DatasmithCADImporter/Source/CADInterfaces/Private/CoreTechFileParser.cpp
--- a/Engine/Plugins/Enterprise/DatasmithCADImporter/Source/CADInterfaces/Private/CoreTechFileParser.cpp
+++ b/Engine/Plugins/Enterprise/DatasmithCADImporter/Source/CADInterfaces/Private/CoreTechFileParser.cpp
@@ -861,7 +861,7 @@ CT_FLAGS FCoreTechFileParser::SetCoreTechImportOption(const FString& MainFileExt
 
 	// All the BRep topology is not available in IGES import
 	// Ask Kernel IO to complete or create missing topology
-	if (MainFileExt == TEXT(".igs") || MainFileExt == TEXT("iges"))
+	if (MainFileExt == TEXT("igs") || MainFileExt == TEXT("iges"))
 	{
 		Flags |= CT_LOAD_FLAG_SEARCH_NEW_TOPOLOGY | CT_LOAD_FLAG_COMPLETE_TOPOLOGY;
 	}
