CAD materials: fix specular, emissive and Roughness

diff --git a/Engine/Plugins/Enterprise/DatasmithCADImporter/Source/CADLibrary/Private/CoreTechHelper.cpp b/Engine/Plugins/Enterprise/DatasmithCADImporter/Source/CADLibrary/Private/CoreTechHelper.cpp
--- a/Engine/Plugins/Enterprise/DatasmithCADImporter/Source/CADLibrary/Private/CoreTechHelper.cpp
+++ b/Engine/Plugins/Enterprise/DatasmithCADImporter/Source/CADLibrary/Private/CoreTechHelper.cpp
@@ -727,30 +727,24 @@ TSharedPtr<IDatasmithUEPbrMaterialElement> CreateUEPbrMaterialFromMaterial(FCADM
 	}
 
 	// Set a Emissive color 
-	if (MaterialElement->GetEmissiveColor().GetExpression() == nullptr)
+	if (MaterialElement->GetSpecular().GetExpression() == nullptr)
 	{
-		// Doc CT => TODO
-		//GLfloat Specular[4] = { specular.rgb[0] / 255., specular.rgb[1] / 255., specular.rgb[2] / 255., 1. - transparency };
-		//GLfloat Shininess[1] = { (float)(128 * shininess) };
-		//glMaterialfv(GL_FRONT_AND_BACK, GL_SPECULAR, Specular);
-		//glMaterialfv(GL_FRONT_AND_BACK, GL_SHININESS, Shininess);
-
 		FLinearColor LinearColor = FLinearColor::FromPow22Color(InMaterial.Specular);
 
 		IDatasmithMaterialExpressionColor* ColorExpression = MaterialElement->AddMaterialExpression<IDatasmithMaterialExpressionColor>();
 		ColorExpression->SetName(TEXT("Specular Color"));
 		ColorExpression->GetColor() = LinearColor;
 
-		MaterialElement->GetEmissiveColor().SetExpression(ColorExpression);
+		MaterialElement->GetSpecular().SetExpression(ColorExpression);
 	}
 
 	// Simple conversion of shininess and reflectivity to PBR roughness and metallic values; model could be improved to properly blend the values
 	if (!FMath::IsNearlyZero(InMaterial.Shininess))
 	{
-		IDatasmithMaterialExpressionScalar* Scalar = static_cast<IDatasmithMaterialExpressionScalar*>(MaterialElement->AddMaterialExpression(EDatasmithMaterialExpressionType::ConstantScalar));
-		Scalar->GetScalar() = 1.f - InMaterial.Shininess;
-		Scalar->SetName(TEXT("Shininess Level"));
-		MaterialElement->GetRoughness().SetExpression(Scalar);
+		IDatasmithMaterialExpressionScalar* RoughnessScalar = static_cast<IDatasmithMaterialExpressionScalar*>(MaterialElement->AddMaterialExpression(EDatasmithMaterialExpressionType::ConstantScalar));
+		RoughnessScalar->GetScalar() = 1.f - InMaterial.Shininess;
+		RoughnessScalar->SetName(TEXT("Roughness Level"));
+		MaterialElement->GetRoughness().SetExpression(RoughnessScalar);
 	}
 
 	if (!FMath::IsNearlyZero(InMaterial.Reflexion))
