Alias : Static mesh imported without material from wire file with Datasmith
Material wasn't  add on body with many patches
Fix conversion of wire file name into char*
Fix add face into a KernelIO Body

diff --git a/Engine/Plugins/Enterprise/DatasmithCADImporter/Source/DatasmithWireTranslator/Private/AliasCoretechWrapper.cpp b/Engine/Plugins/Enterprise/DatasmithCADImporter/Source/DatasmithWireTranslator/Private/AliasCoretechWrapper.cpp
--- a/Engine/Plugins/Enterprise/DatasmithCADImporter/Source/DatasmithWireTranslator/Private/AliasCoretechWrapper.cpp
+++ b/Engine/Plugins/Enterprise/DatasmithCADImporter/Source/DatasmithWireTranslator/Private/AliasCoretechWrapper.cpp
@@ -319,7 +319,7 @@ void FAliasCoretechWrapper::AddFace(AlSurface& Surface, CT_LIST_IO& FaceList, bo
 	CT_OBJECT_ID FaceID;
 	CT_ORIENTATION FaceOrient = bOrientation ? CT_ORIENTATION::CT_FORWARD : CT_ORIENTATION::CT_REVERSE;
 	CT_IO_ERROR Result = CT_FACE_IO::Create(FaceID, NurbsId, FaceOrient, Boundaries);
-	if (Result != IO_OK)
+	if (Result == IO_OK)
 	{
 		FaceList.PushBack(FaceID);
 	}
diff --git a/Engine/Plugins/Enterprise/DatasmithCADImporter/Source/DatasmithWireTranslator/Private/AliasCoretechWrapper.h b/Engine/Plugins/Enterprise/DatasmithCADImporter/Source/DatasmithWireTranslator/Private/AliasCoretechWrapper.h
--- a/Engine/Plugins/Enterprise/DatasmithCADImporter/Source/DatasmithWireTranslator/Private/AliasCoretechWrapper.h
+++ b/Engine/Plugins/Enterprise/DatasmithCADImporter/Source/DatasmithWireTranslator/Private/AliasCoretechWrapper.h
@@ -10,7 +10,6 @@
 // Additional scale factor required when tessellating NURBS as Alias has extremely small geometry,
 // originally tessellating to triangles with area in the order of 10^-5, failing the
 // FourSquaredTriangleArea > SMALL_NUMBER test in DatasmithMeshHelper.cpp::IsMeshValid
-#define ALIAS_BUILD_SCALE 100
 
 class AlDagNode;
 class AlShell;
@@ -36,7 +35,9 @@ public:
 	 * eg. For a file in inches, arg should be 0.0254
 	 */
 	FAliasCoretechWrapper(const TCHAR* InOwner)
-		: CTSession(InOwner, 0.01, ALIAS_BUILD_SCALE) // Alias files are in cm units
+		: CTSession(InOwner, 0.001, 1) 
+// We prefere to tell to kernelIO that Nurbs are in mm (default unit of kernelIO) to don't have side effect. 
+// Scale == 1 because in fact Alias work in cm so we do not need to scale mesh parameters
 	{
 	}
 
diff --git a/Engine/Plugins/Enterprise/DatasmithCADImporter/Source/DatasmithWireTranslator/Private/DatasmithWireTranslator.cpp b/Engine/Plugins/Enterprise/DatasmithCADImporter/Source/DatasmithWireTranslator/Private/DatasmithWireTranslator.cpp
--- a/Engine/Plugins/Enterprise/DatasmithCADImporter/Source/DatasmithWireTranslator/Private/DatasmithWireTranslator.cpp
+++ b/Engine/Plugins/Enterprise/DatasmithCADImporter/Source/DatasmithWireTranslator/Private/DatasmithWireTranslator.cpp
@@ -206,8 +206,7 @@ bool FWireTranslatorImpl::Read()
 	// Initialize Alias.
 	AlUniverse::initialize();
 
-	char* fileName = TCHAR_TO_UTF8(*SceneFullPath);
-	if (AlUniverse::retrieve(fileName) != sSuccess)
+	if (AlUniverse::retrieve(TCHAR_TO_UTF8(*SceneFullPath)) != sSuccess)
 	{
 		return false;
 	}
@@ -1615,12 +1614,14 @@ bool FWireTranslatorImpl::RecurseDagForLeaves(AlDagNode* FirstDagNode, const FDa
 			AlShellNode *ShellNode = DagNode->asShellNodePtr();
 			AlShell *Shell = ShellNode->shell();
 			uint32 NbPatch = getNumOfPatch(*Shell);
+
+			if (AlShader* Shader = Shell->firstShader())
+			{
+				ShaderName = Shader->name();
+			}
+
 			if (NbPatch == 1)
 			{
-				if (AlShader * Shader = Shell->firstShader())
-				{
-					ShaderName = Shader->name();
-				}
 				AddNodeInBodySet(*DagNode, ShaderName, ShellToProcess, true, MaxSize);
 			}
 			else
