Fix build breakage introduced by 9843400: Alias translator

diff --git a/Engine/Plugins/Enterprise/DatasmithCADImporter/Source/DatasmithWireTranslator/Private/DatasmithWireTranslator.cpp b/Engine/Plugins/Enterprise/DatasmithCADImporter/Source/DatasmithWireTranslator/Private/DatasmithWireTranslator.cpp
--- a/Engine/Plugins/Enterprise/DatasmithCADImporter/Source/DatasmithWireTranslator/Private/DatasmithWireTranslator.cpp
+++ b/Engine/Plugins/Enterprise/DatasmithCADImporter/Source/DatasmithWireTranslator/Private/DatasmithWireTranslator.cpp
@@ -1393,7 +1393,7 @@ bool FWireTranslatorImpl::ProcessAlShellNode(AlDagNode& ShellNode, const FDagNod
 	AlObjectType Type = ShellNode.type();
 	if (Type == AlObjectType::kSurfaceNodeType || Type == AlObjectType::kShellNodeType)
 	{
-		ActorElement->SetScale(ActorElement->GetScale() / ALIAS_BUILD_SCALE);
+		ActorElement->SetScale(ActorElement->GetScale());
 	}
 
 	//// Apply materials on the current part
