Fixed crash with JT files

diff --git a/Engine/Plugins/Enterprise/DatasmithCADImporter/Source/CADInterfaces/Private/CoreTechFileParser.cpp b/Engine/Plugins/Enterprise/DatasmithCADImporter/Source/CADInterfaces/Private/CoreTechFileParser.cpp
--- a/Engine/Plugins/Enterprise/DatasmithCADImporter/Source/CADInterfaces/Private/CoreTechFileParser.cpp
+++ b/Engine/Plugins/Enterprise/DatasmithCADImporter/Source/CADInterfaces/Private/CoreTechFileParser.cpp
@@ -852,7 +852,11 @@ CT_FLAGS FCoreTechFileParser::SetCoreTechImportOption(const FString& MainFileExt
 	// Set import option
 	CT_FLAGS Flags = CT_LOAD_FLAGS_USE_DEFAULT;
 
-	Flags |= CT_LOAD_FLAGS_READ_META_DATA;
+	// Do not read meta-data from JT files with CoreTech. It crashes...
+	if (MainFileExt != TEXT("jt"))
+	{
+		Flags |= CT_LOAD_FLAGS_READ_META_DATA;
+	}
 
 	if (MainFileExt == TEXT("catpart") || MainFileExt == TEXT("catproduct") || MainFileExt == TEXT("cgr"))
 	{
