Stiching technique is irrelevant in operator

diff --git a/Engine/Plugins/Enterprise/DatasmithCADImporter/Source/DatasmithCoreTechExtension/Private/DataprepTessellationOperation.cpp b/Engine/Plugins/Enterprise/DatasmithCADImporter/Source/DatasmithCoreTechExtension/Private/DataprepTessellationOperation.cpp
--- a/Engine/Plugins/Enterprise/DatasmithCADImporter/Source/DatasmithCoreTechExtension/Private/DataprepTessellationOperation.cpp
+++ b/Engine/Plugins/Enterprise/DatasmithCADImporter/Source/DatasmithCoreTechExtension/Private/DataprepTessellationOperation.cpp
@@ -24,6 +24,20 @@
 
 DEFINE_LOG_CATEGORY_STATIC(LogCADLibrary, Log, All)
 
+void UDataprepTessellationOperation::PostLoad()
+{
+	if( HasAnyFlags(RF_WasLoaded) && TessellationSettings_DEPRECATED.ChordTolerance != -MAX_FLT)
+	{
+		ChordTolerance = TessellationSettings_DEPRECATED.ChordTolerance;
+		MaxEdgeLength = TessellationSettings_DEPRECATED.MaxEdgeLength;
+		NormalTolerance = TessellationSettings_DEPRECATED.NormalTolerance;
+		// Mark TessellationSettings_DEPRECATED as non usable
+		TessellationSettings_DEPRECATED.ChordTolerance = -MAX_FLT;
+
+		MarkPackageDirty();
+	}
+}
+
 void UDataprepTessellationOperation::OnExecution_Implementation(const FDataprepContext& InContext)
 {
 #ifdef CAD_LIBRARY
@@ -53,6 +67,8 @@ void UDataprepTessellationOperation::OnExecution_Implementation(const FDataprepC
 
 	if(	!IsCancelled() && SelectedMeshes.Num() > 0)
 	{
+		FDatasmithTessellationOptions TessellationSettings( ChordTolerance, MaxEdgeLength, NormalTolerance );
+
 		TSharedPtr<FDataprepWorkReporter> Task = CreateTask( LOCTEXT( "LogCADLibrary_Tessellating", "Tessellating meshes ..." ), (float)SelectedMeshes.Num() );
 
 		TArray<UObject*> ModifiedStaticMeshes;
diff --git a/Engine/Plugins/Enterprise/DatasmithCADImporter/Source/DatasmithCoreTechExtension/Private/DataprepTessellationOperation.h b/Engine/Plugins/Enterprise/DatasmithCADImporter/Source/DatasmithCoreTechExtension/Private/DataprepTessellationOperation.h
--- a/Engine/Plugins/Enterprise/DatasmithCADImporter/Source/DatasmithCoreTechExtension/Private/DataprepTessellationOperation.h
+++ b/Engine/Plugins/Enterprise/DatasmithCADImporter/Source/DatasmithCoreTechExtension/Private/DataprepTessellationOperation.h
@@ -23,9 +23,49 @@ class UDataprepTessellationOperation : public UDataprepOperation
 {
 	GENERATED_BODY()
 
+	UDataprepTessellationOperation()
+		: ChordTolerance(0.2f)
+		, MaxEdgeLength(0.0f)
+		, NormalTolerance(20.0f)
+		// Mark TessellationSettings_DEPRECATED as non usable
+		, TessellationSettings_DEPRECATED(-MAX_FLT)
+	{
+	}
+
+	// UObject interface
+	/** Overridden for backward compatibility with 4.23. To be removed in 4.25 */
+	virtual void PostLoad() override;
+	// End of UObject interface
+
 public:
-	UPROPERTY(EditAnywhere, BlueprintReadWrite, Category = TessellationOptions, meta = (ShowOnlyInnerProperties) )
-	FDatasmithTessellationOptions TessellationSettings;
+	/**
+	 * Maximum distance between any point on a triangle generated by the tessellation process and the actual surface.
+	 * The lower the value the more triangles.
+	 * Default value is 0.2.
+	 */
+	UPROPERTY(EditAnywhere, BlueprintReadWrite, Category = "Tessellation Options", meta = (Units = cm, ToolTip = "Maximum distance between any generated triangle and the original surface. Smaller values make more triangles.", ClampMin = "0.0"))
+	float ChordTolerance;
+
+	/**
+	 * Maximum length of edges of triangles generated by the tessellation process.
+	 * The length is in scene/model unit. The smaller the more triangles are generated.
+	 * Value of 0 means no constraint on length of edges
+	 * Default value is 0.
+	 */
+	UPROPERTY(EditAnywhere, BlueprintReadWrite, AdvancedDisplay, Category = "Tessellation Options", meta = (Units = cm, DisplayName = "Max Edge Length", ToolTip = "Maximum length of any edge in the generated triangles. Smaller values make more triangles.", ClampMin = "0.0"))
+	float MaxEdgeLength;
+
+	/**
+	 * Maximum angle between the normal of two triangles generated by the tessellation process.
+	 * The angle is expressed in degree. The smaller the more triangles are generated.
+	 * Default value is 20 degrees.
+	 */
+	UPROPERTY(EditAnywhere, BlueprintReadWrite, AdvancedDisplay, Category = "Tessellation Options", meta = (Units = deg, ToolTip = "Maximum angle between adjacent triangles. Smaller values make more triangles.", ClampMin = "0.0", ClampMax = "90.0"))
+	float NormalTolerance;
+
+	/** Version 4.24 : Deprecated - not serialized anymore */
+	UPROPERTY()
+	FDatasmithTessellationOptions TessellationSettings_DEPRECATED;
 
 	//~ Begin UDataprepOperation Interface
 public:
