Added the missing call Super::PostLoad.

#rnx

diff --git a/Engine/Plugins/Enterprise/DatasmithCADImporter/Source/DatasmithCoreTechExtension/Private/DataprepTessellationOperation.cpp b/Engine/Plugins/Enterprise/DatasmithCADImporter/Source/DatasmithCoreTechExtension/Private/DataprepTessellationOperation.cpp
--- a/Engine/Plugins/Enterprise/DatasmithCADImporter/Source/DatasmithCoreTechExtension/Private/DataprepTessellationOperation.cpp
+++ b/Engine/Plugins/Enterprise/DatasmithCADImporter/Source/DatasmithCoreTechExtension/Private/DataprepTessellationOperation.cpp
@@ -36,6 +36,8 @@ void UDataprepTessellationOperation::PostLoad()
 
 		MarkPackageDirty();
 	}
+
+	Super::PostLoad();
 }
 
 void UDataprepTessellationOperation::OnExecution_Implementation(const FDataprepContext& InContext)
