Wire: Remove check for IsMeshValid, allowing small triangle meshes to be imported #jira UE-82590 #rb JeanLuc.Corenthin

diff --git a/Engine/Plugins/Enterprise/DatasmithCADImporter/Source/DatasmithWireTranslator/Private/AliasCoretechWrapper.h b/Engine/Plugins/Enterprise/DatasmithCADImporter/Source/DatasmithWireTranslator/Private/AliasCoretechWrapper.h
--- a/Engine/Plugins/Enterprise/DatasmithCADImporter/Source/DatasmithWireTranslator/Private/AliasCoretechWrapper.h
+++ b/Engine/Plugins/Enterprise/DatasmithCADImporter/Source/DatasmithWireTranslator/Private/AliasCoretechWrapper.h
@@ -7,10 +7,6 @@
 
 #include "CTSession.h"
 
-// Additional scale factor required when tessellating NURBS as Alias has extremely small geometry,
-// originally tessellating to triangles with area in the order of 10^-5, failing the
-// FourSquaredTriangleArea > SMALL_NUMBER test in DatasmithMeshHelper.cpp::IsMeshValid
-
 class AlDagNode;
 class AlShell;
 class AlSurface;
diff --git a/Engine/Plugins/Enterprise/DatasmithCADImporter/Source/DatasmithWireTranslator/Private/DatasmithWireTranslator.cpp b/Engine/Plugins/Enterprise/DatasmithCADImporter/Source/DatasmithWireTranslator/Private/DatasmithWireTranslator.cpp
--- a/Engine/Plugins/Enterprise/DatasmithCADImporter/Source/DatasmithWireTranslator/Private/DatasmithWireTranslator.cpp
+++ b/Engine/Plugins/Enterprise/DatasmithCADImporter/Source/DatasmithWireTranslator/Private/DatasmithWireTranslator.cpp
@@ -1387,16 +1387,7 @@ bool FWireTranslatorImpl::ProcessAlShellNode(AlDagNode& ShellNode, const FDagNod
 
 	SetActorTransform(ShellInfo.ActorElement, ShellNode);
 
-	// If the mesh is NURBS, it will be built with an additional ALIAS_BUILD_SCALE (see AliasCoretechWrapper constructor),
-	// so we need to compensate the same scale on the actor. Alias doesn't seem to allow parenting geometry
-	// to other geometry, so we don't have to handle consequences on any potential children of this actor by doing this
-	AlObjectType Type = ShellNode.type();
-	if (Type == AlObjectType::kSurfaceNodeType || Type == AlObjectType::kShellNodeType)
-	{
-		ActorElement->SetScale(ActorElement->GetScale());
-	}
-
-	//// Apply materials on the current part
+	// Apply materials on the current part
 	if (ShaderName)
 	{
 		TSharedPtr< IDatasmithMaterialIDElement > MaterialIDElement = ShaderNameToUEMaterialId[FString(ShaderName)];
diff --git a/Engine/Plugins/Enterprise/DatasmithImporter/Source/DatasmithImporter/Private/DatasmithMeshHelper.cpp b/Engine/Plugins/Enterprise/DatasmithImporter/Source/DatasmithImporter/Private/DatasmithMeshHelper.cpp
--- a/Engine/Plugins/Enterprise/DatasmithImporter/Source/DatasmithImporter/Private/DatasmithMeshHelper.cpp
+++ b/Engine/Plugins/Enterprise/DatasmithImporter/Source/DatasmithImporter/Private/DatasmithMeshHelper.cpp
@@ -362,7 +362,10 @@ namespace DatasmithMeshHelper
 				FVector RawNormal = (Corners[1] - Corners[2]) ^ (Corners[0] - Corners[2]);
 				RawNormal *= RawNormalScale;
 				double FourSquaredTriangleArea = RawNormal.SizeSquared();
-				if (FourSquaredTriangleArea > SMALL_NUMBER)
+
+				// We support even small triangles, but this function is still useful to
+				// see if we have at least one valid triangle in the mesh
+				if (FourSquaredTriangleArea > 0.0f)
 				{
 					return true;
 				}
