Fixed issue with mesh element unique ids in DatssmithScene
Re-enabled application of stitching options
Completed the implementation of the aggregation of meshes

diff --git a/Engine/Plugins/Enterprise/DatasmithCADImporter/Source/DatasmithWireTranslator/Private/AliasCoretechWrapper.cpp b/Engine/Plugins/Enterprise/DatasmithCADImporter/Source/DatasmithWireTranslator/Private/AliasCoretechWrapper.cpp
--- a/Engine/Plugins/Enterprise/DatasmithCADImporter/Source/DatasmithWireTranslator/Private/AliasCoretechWrapper.cpp
+++ b/Engine/Plugins/Enterprise/DatasmithCADImporter/Source/DatasmithWireTranslator/Private/AliasCoretechWrapper.cpp
@@ -105,6 +105,10 @@ namespace AliasToCoreTechUtils
 
 CT_IO_ERROR FAliasCoretechWrapper::Tessellate(FMeshDescription& Mesh, FMeshParameters& MeshParameters)
 {
+	// Apply stitching if applicable
+	TopoFixes();
+
+	// Perform tessellation
 	return CADLibrary::Tessellate(MainObjectId, ImportParams, Mesh, MeshParameters);
 }
 
diff --git a/Engine/Plugins/Enterprise/DatasmithCADImporter/Source/DatasmithWireTranslator/Private/DatasmithWireTranslator.cpp b/Engine/Plugins/Enterprise/DatasmithCADImporter/Source/DatasmithWireTranslator/Private/DatasmithWireTranslator.cpp
--- a/Engine/Plugins/Enterprise/DatasmithCADImporter/Source/DatasmithWireTranslator/Private/DatasmithWireTranslator.cpp
+++ b/Engine/Plugins/Enterprise/DatasmithCADImporter/Source/DatasmithWireTranslator/Private/DatasmithWireTranslator.cpp
@@ -64,6 +64,52 @@ public:
 	FString LayerName;
 	FString Label;
 	bool bCadData = true;
+
+	// Generates BodyData's unique id from AlDagNode objects
+	FString GetUUID(const FString& ParentUUID)
+	{
+		if(ShellSet.Num() == 0)
+		{
+			return ParentUUID;
+		}
+
+		auto GetLongPersistentID = []( AlDagNode& DagNode ) -> int64
+		{
+			union {
+				int a[2];
+				int64 b;
+			} Value;
+
+			Value.b = -1;
+
+			AlPersistentID* PersistentID = nullptr;
+			DagNode.persistentID( PersistentID );
+
+			if( PersistentID != nullptr )
+			{
+				int Dummy;
+				PersistentID->id( Value.a[0], Value.a[1], Dummy, Dummy );
+			}
+
+			return Value.b;
+		};
+
+		if(ShellSet.Num() > 1)
+		{
+			ShellSet.Sort([&](AlDagNode& A, AlDagNode& B)
+			{
+				return GetLongPersistentID( A ) < GetLongPersistentID( B );
+			});
+		}
+
+		FString Buffer;
+		for(AlDagNode* DagNode : ShellSet)
+		{
+			Buffer += FString::Printf(TEXT("%016lx"), GetLongPersistentID( *DagNode ) );
+		}
+
+		return GetUEUUIDFromAIPersistentID( ParentUUID, Buffer );
+	}
 };
 
 class FWireTranslatorImpl
@@ -127,7 +173,7 @@ private:
 	bool ProcessAlGroupNode(AlGroupNode& GroupNode, const FDagNodeInfo& ParentInfo);
 	bool ProcessAlShellNode(AlDagNode& GroupNode, const FDagNodeInfo& ParentInfo, const char* ShaderName);
 	bool ProcessBodyNode(TSharedRef<BodyData> Body, const FDagNodeInfo& ParentInfo);
-	TSharedPtr< IDatasmithMeshElement > AddMeshElement(TSharedRef<BodyData> Body, const FDagNodeInfo& ParentInfo);
+	TSharedPtr< IDatasmithMeshElement > FindOrAddMeshElement(TSharedRef<BodyData> Body, const FDagNodeInfo& ParentInfo);
 	TSharedPtr< IDatasmithMeshElement > FindOrAddMeshElement(AlDagNode& ShellNode, const FDagNodeInfo& ParentInfo, const char* ShaderName);
 	void GetDagNodeInfo(AlDagNode& GroupNode, const FDagNodeInfo& ParentInfo, FDagNodeInfo& CurrentNodeInfo);
 	void GetDagNodeInfo(TSharedRef<BodyData> CurrentNode, const FDagNodeInfo& ParentInfo, FDagNodeInfo& CurrentNodeInfo);
@@ -166,6 +212,7 @@ private:
 
 	/** Table of correspondence between mesh identifier and associated Datasmith mesh element */
 	TMap< uint32, TSharedPtr< IDatasmithMeshElement > > ShellUUIDToMeshElementMap;
+	TMap< FString, TSharedPtr< IDatasmithMeshElement > > BodyToMeshElementMap;
 
 	/** Datasmith mesh elements to OpenModel objects */
 	TMap< IDatasmithMeshElement*, AlDagNode* > MeshElementToAlDagNodeMap;
@@ -1241,11 +1288,10 @@ void FWireTranslatorImpl::GetDagNodeInfo(AlDagNode& CurrentNode, const FDagNodeI
 {
 	CurrentNodeInfo.Label = CurrentNode.name();
 
-	FString ThisGroupNodeID;
-	AlPersistentID* GroupNodeId = new AlPersistentID();
-
+	AlPersistentID* GroupNodeId = nullptr;
 	CurrentNode.persistentID(GroupNodeId);
-	ThisGroupNodeID = GetPersistentIDString(GroupNodeId);
+
+	FString ThisGroupNodeID( GetPersistentIDString(GroupNodeId) );
 
     // Limit length of UUID by combining hash of parent UUID and container's UUID if ParentUuid is not empty
 	CurrentNodeInfo.UEuuid = GetUEUUIDFromAIPersistentID(ParentInfo.UEuuid, ThisGroupNodeID);
@@ -1257,7 +1303,7 @@ void FWireTranslatorImpl::GetDagNodeInfo(TSharedRef<BodyData> CurrentNode, const
 	CurrentNode->Label = CurrentNodeInfo.Label;
 
 	// Limit length of UUID by combining hash of parent UUID and container's UUID if ParentUuid is not empty
-	CurrentNodeInfo.UEuuid = GetUEUUIDFromAIPersistentID(ParentInfo.UEuuid, CurrentNodeInfo.Label);
+	CurrentNodeInfo.UEuuid = GetUEUUIDFromAIPersistentID( ParentInfo.UEuuid, CurrentNode->GetUUID( ParentInfo.UEuuid ) );
 }
 
 
@@ -1296,11 +1342,15 @@ bool FWireTranslatorImpl::ProcessAlGroupNode(AlGroupNode& GroupNode, const FDagN
 	return true;
 }
 
-TSharedPtr< IDatasmithMeshElement > FWireTranslatorImpl::AddMeshElement(TSharedRef<BodyData> Body, const FDagNodeInfo& NodeInfo)
+TSharedPtr< IDatasmithMeshElement > FWireTranslatorImpl::FindOrAddMeshElement(TSharedRef<BodyData> Body, const FDagNodeInfo& NodeInfo)
 {
-	uint32 BodyUUID = GetTypeHash(NodeInfo.UEuuid);
-
-	TSharedPtr< IDatasmithMeshElement > MeshElement = FDatasmithSceneFactory::CreateMesh(*FString::Printf(TEXT("0x%08x"), BodyUUID));
+	TSharedPtr< IDatasmithMeshElement >* MeshElementPtr = BodyToMeshElementMap.Find( NodeInfo.UEuuid );
+	if (MeshElementPtr != nullptr)
+	{
+		return *MeshElementPtr;
+	}
+	
+	TSharedPtr< IDatasmithMeshElement > MeshElement = FDatasmithSceneFactory::CreateMesh( *NodeInfo.UEuuid );
 	MeshElement->SetLabel(*NodeInfo.Label);
 	MeshElement->SetLightmapSourceUV(-1);
 
@@ -1312,9 +1362,11 @@ TSharedPtr< IDatasmithMeshElement > FWireTranslatorImpl::AddMeshElement(TSharedR
 
 	DatasmithScene->AddMesh(MeshElement);
 
-	ShellUUIDToMeshElementMap.Add(BodyUUID, MeshElement);
+	ShellUUIDToMeshElementMap.Add(FCString::Atoi(*NodeInfo.UEuuid), MeshElement);
 	MeshElementToBodyMap.Add(MeshElement.Get(), Body);
 
+	BodyToMeshElementMap.Add( NodeInfo.UEuuid, MeshElement );
+
 	return MeshElement;
 }
 
@@ -1424,7 +1476,7 @@ bool FWireTranslatorImpl::ProcessBodyNode(TSharedRef<BodyData> Body, const FDagN
 	FDagNodeInfo ShellInfo;
 	GetDagNodeInfo(Body, ParentInfo, ShellInfo);
 
-	TSharedPtr< IDatasmithMeshElement > MeshElement = AddMeshElement(Body, ShellInfo);
+	TSharedPtr< IDatasmithMeshElement > MeshElement = FindOrAddMeshElement(Body, ShellInfo);
 	if (!MeshElement.IsValid())
 	{
 		return false;
@@ -1781,8 +1833,6 @@ TOptional<FMeshDescription> FWireTranslatorImpl::MeshDagNodeWithExternalMesher(A
 		MeshElement->SetFile(*FilePath);
 	}
 
-	//Result = LocalSession->TopoFixes();
-
 	FMeshDescription MeshDescription;
 	DatasmithMeshHelper::PrepareAttributeForStaticMesh(MeshDescription);
 
@@ -1800,8 +1850,7 @@ TOptional<FMeshDescription> FWireTranslatorImpl::MeshDagNodeWithExternalMesher(T
 	Result = LocalSession->AddBRep(Body->ShellSet, MeshParameters.bIsSymmetric);
 
 	//const char*Name = DagNode->name();
-	FString Filename = Body->Label;
-	Filename += TEXT(".ct");
+	FString Filename = Body->Label + TEXT(".ct");
 
 	FString FilePath = FPaths::Combine(OutputPath, Filename);
 	Result = LocalSession->SaveBrep(FilePath);
@@ -1810,8 +1859,6 @@ TOptional<FMeshDescription> FWireTranslatorImpl::MeshDagNodeWithExternalMesher(T
 		MeshElement->SetFile(*FilePath);
 	}
 
-	//Result = LocalSession->TopoFixes();
-
 	FMeshDescription MeshDescription;
 	DatasmithMeshHelper::PrepareAttributeForStaticMesh(MeshDescription);
 
@@ -1861,6 +1908,7 @@ TOptional<FMeshDescription> FWireTranslatorImpl::GetMeshOfMeshBody(TSharedRef<Bo
 {
 	FMeshDescription MeshDescription;
 	DatasmithMeshHelper::PrepareAttributeForStaticMesh(MeshDescription);
+	MeshDescription.Empty();
 	bool True = true;
 
 	for (auto DagNode : Body->ShellSet)
@@ -1869,8 +1917,7 @@ TOptional<FMeshDescription> FWireTranslatorImpl::GetMeshOfMeshBody(TSharedRef<Bo
 		AlMesh *Mesh = MeshNode->mesh();
 		if (Mesh)
 		{
-			//TransferAlMeshToMeshDescription(*Mesh, MeshDescription, MeshParameters, True);
-			// TODO
+			TransferAlMeshToMeshDescription(*Mesh, MeshDescription, MeshParameters, True, true);
 		}
 	}
 
diff --git a/Engine/Plugins/Enterprise/DatasmithCADImporter/Source/DatasmithWireTranslator/Private/OpenModelUtils.cpp b/Engine/Plugins/Enterprise/DatasmithCADImporter/Source/DatasmithWireTranslator/Private/OpenModelUtils.cpp
--- a/Engine/Plugins/Enterprise/DatasmithCADImporter/Source/DatasmithWireTranslator/Private/OpenModelUtils.cpp
+++ b/Engine/Plugins/Enterprise/DatasmithCADImporter/Source/DatasmithWireTranslator/Private/OpenModelUtils.cpp
@@ -1605,6 +1605,11 @@ bool OpenModelUtils::IsValidActor(const TSharedPtr< IDatasmithActorElement >& Ac
 
 FString OpenModelUtils::GetPersistentIDString(AlPersistentID* GroupNodeId)
 {
+	if(GroupNodeId == nullptr)
+	{
+		return FString();
+	}
+
 	int IdA, IdB, IdC, IdD;
 	GroupNodeId->id(IdA, IdB, IdC, IdD);
 	FString ThisGroupNodeID = FString::FromInt(IdA) + FString::FromInt(IdB) + FString::FromInt(IdC) + FString::FromInt(IdD);
@@ -1636,7 +1641,7 @@ FString OpenModelUtils::GetUEUUIDFromAIPersistentID(const FString& ParentUEuuid,
 
 
 
-bool OpenModelUtils::TransferAlMeshToMeshDescription(const AlMesh& AliasMesh, FMeshDescription& MeshDescription, CADLibrary::FMeshParameters& MeshParameters, bool& bHasNormal)
+bool OpenModelUtils::TransferAlMeshToMeshDescription(const AlMesh& AliasMesh, FMeshDescription& MeshDescription, CADLibrary::FMeshParameters& MeshParameters, bool& bHasNormal, bool bMerge)
 {
 	// Ref. GP3DMVisitorImpl::visitMesh
 	// Ref. FGPureMeshInterface::CreateMesh
@@ -1645,7 +1650,10 @@ bool OpenModelUtils::TransferAlMeshToMeshDescription(const AlMesh& AliasMesh, FM
 		return false;
 	}
 
-	MeshDescription.Empty();
+	if( !bMerge )
+	{
+		MeshDescription.Empty();
+	}
 
 	int32 NbStep = 1;
 	FMatrix SymmetricMatrix;
@@ -1676,6 +1684,9 @@ bool OpenModelUtils::TransferAlMeshToMeshDescription(const AlMesh& AliasMesh, FM
 	int TriangleCount = AliasMesh.numberOfTriangles();
 	const int32 VertexInstanceCount = 3 * TriangleCount;
 
+	TArray<FVertexID> VertexPositionIDs;
+	VertexPositionIDs.SetNum( VertexCount*NbStep );
+
 	// Reserve space for attributes
 	// At this point, all the faces are triangles
 	MeshDescription.ReserveNewVertices(VertexCount*NbStep);
@@ -1690,7 +1701,10 @@ bool OpenModelUtils::TransferAlMeshToMeshDescription(const AlMesh& AliasMesh, FM
 	PolygonGroupImportedMaterialSlotNames[PolyGroupId] = ImportedSlotName;
 
 	// At least one UV set must exist.
-	VertexInstanceUVs.SetNumIndices(1);
+	if( VertexInstanceUVs.GetNumIndices() == 0 )
+	{
+		VertexInstanceUVs.SetNumIndices(1);
+	}
 
 	// Get Alias mesh info
 	const float * AlVertices = AliasMesh.vertices();
@@ -1700,22 +1714,24 @@ bool OpenModelUtils::TransferAlMeshToMeshDescription(const AlMesh& AliasMesh, FM
 		// Fill the vertex array
 		if (Step == 0)
 		{
-			for (int Index = 0; Index < VertexCount; ++Index)
+			FVertexID* VertexPositionIDPtr = VertexPositionIDs.GetData();
+			for (int Index = 0; Index < VertexCount; ++Index, ++VertexPositionIDPtr)
 			{
 				const float * CurVertex = AlVertices + 3 * Index;
-				FVertexID AddedVertexId = MeshDescription.CreateVertex();
+				*VertexPositionIDPtr = MeshDescription.CreateVertex();
 				// ConvertVector_ZUp_RightHanded
-				VertexPositions[AddedVertexId] = FVector(-CurVertex[0], CurVertex[1], CurVertex[2]);
+				VertexPositions[*VertexPositionIDPtr] = FVector(-CurVertex[0], CurVertex[1], CurVertex[2]);
 			}
 		}
 		else 
 		{
-			for (int Index = 0; Index < VertexCount; ++Index)
+			FVertexID* VertexPositionIDPtr = VertexPositionIDs.GetData() + VertexCount;
+			for (int Index = 0, PositionIndex = VertexCount; Index < VertexCount; ++Index, ++VertexPositionIDPtr)
 			{
 				const float * CurVertex = AlVertices + 3 * Index;
-				FVertexID AddedVertexId = MeshDescription.CreateVertex();
+				*VertexPositionIDPtr = MeshDescription.CreateVertex();
 				// ConvertVector_ZUp_RightHanded
-				VertexPositions[AddedVertexId] = SymmetricMatrix.TransformPosition(FVector(-CurVertex[0], CurVertex[1], CurVertex[2]));
+				VertexPositions[*VertexPositionIDPtr] = SymmetricMatrix.TransformPosition(FVector(-CurVertex[0], CurVertex[1], CurVertex[2]));
 			}
 		}
 
@@ -1739,7 +1755,7 @@ bool OpenModelUtils::TransferAlMeshToMeshDescription(const AlMesh& AliasMesh, FM
 				// Create Vertex instances and set their attributes
 				for (int32 VertexIndex = 0, TIndex = 2; VertexIndex < CornerCount; ++VertexIndex, --TIndex)
 				{
-					CornerVertexIDs[VertexIndex] = FVertexID(Triangles[TIndex] + VertexCount * Step);
+					CornerVertexIDs[VertexIndex] = VertexPositionIDs[Triangles[TIndex] + VertexCount * Step];
 					CornerVertexInstanceIDs[VertexIndex] = MeshDescription.CreateVertexInstance(CornerVertexIDs[VertexIndex]);
 
 					// Set the normal
@@ -1785,7 +1801,7 @@ bool OpenModelUtils::TransferAlMeshToMeshDescription(const AlMesh& AliasMesh, FM
 				// Create Vertex instances and set their attributes
 				for (int32 VertexIndex = 0; VertexIndex < CornerCount; ++VertexIndex)
 				{
-					CornerVertexIDs[VertexIndex] = FVertexID(Triangles[VertexIndex] + VertexCount * Step);
+					CornerVertexIDs[VertexIndex] =  VertexPositionIDs[Triangles[VertexIndex] + VertexCount * Step];
 					CornerVertexInstanceIDs[VertexIndex] = MeshDescription.CreateVertexInstance(CornerVertexIDs[VertexIndex]);
 
 					// Set the normal
diff --git a/Engine/Plugins/Enterprise/DatasmithCADImporter/Source/DatasmithWireTranslator/Private/OpenModelUtils.h b/Engine/Plugins/Enterprise/DatasmithCADImporter/Source/DatasmithWireTranslator/Private/OpenModelUtils.h
--- a/Engine/Plugins/Enterprise/DatasmithCADImporter/Source/DatasmithWireTranslator/Private/OpenModelUtils.h
+++ b/Engine/Plugins/Enterprise/DatasmithCADImporter/Source/DatasmithWireTranslator/Private/OpenModelUtils.h
@@ -49,7 +49,7 @@ namespace OpenModelUtils
 	FString GetUEUUIDFromAIPersistentID(const FString& ParentUEuuid, const FString& CurrentNodePersistentID);
 
 	// Note that Alias file unit is cm like UE
-	bool TransferAlMeshToMeshDescription(const AlMesh& Mesh, FMeshDescription& MeshDescription, CADLibrary::FMeshParameters& SymmetricParameters, bool& bHasNormal);
+	bool TransferAlMeshToMeshDescription(const AlMesh& Mesh, FMeshDescription& MeshDescription, CADLibrary::FMeshParameters& SymmetricParameters, bool& bHasNormal, bool bMerge = false);
 
 	AlDagNode* TesselateDagLeaf(AlDagNode* DagLeaf, ETesselatorType TessType, double Tolerance);
 }
