OpenNurbs: Fix StaticMeshes not reimporting when material assignment changed #jira UE-81278 #rb JeanMichel.Dignard, David.Lesage, Johan.Duparc

diff --git a/Engine/Plugins/Enterprise/DatasmithCADImporter/Source/DatasmithOpenNurbsTranslator/Private/DatasmithOpenNurbsTranslator.cpp b/Engine/Plugins/Enterprise/DatasmithCADImporter/Source/DatasmithOpenNurbsTranslator/Private/DatasmithOpenNurbsTranslator.cpp
--- a/Engine/Plugins/Enterprise/DatasmithCADImporter/Source/DatasmithOpenNurbsTranslator/Private/DatasmithOpenNurbsTranslator.cpp
+++ b/Engine/Plugins/Enterprise/DatasmithCADImporter/Source/DatasmithOpenNurbsTranslator/Private/DatasmithOpenNurbsTranslator.cpp
@@ -6,6 +6,7 @@
 #ifdef CAD_LIBRARY
 #include "CoreTechParametricSurfaceExtension.h"
 #endif
+#include "DatasmithImportOptions.h"
 #include "DatasmithMaterialElements.h"
 #include "DatasmithMaterialsUtils.h"
 #include "DatasmithMesh.h"
@@ -631,6 +632,7 @@ public:
 
 	TOptional<FMeshDescription> GetMeshDescription(TSharedRef<IDatasmithMeshElement> MeshElement);
 
+	void SetBaseOptions(const FDatasmithImportBaseOptions& InBaseOptions);
 	void SetTessellationOptions(const FDatasmithTessellationOptions& Options);
 	void SetOutputPath(const FString& Path) { OutputPath = Path; }
 	double GetScalingFactor () const { return ScalingFactor; }
@@ -674,6 +676,7 @@ private:
 	FString OutputPath;
 	FDatasmithTessellationOptions TessellationOptions;
 	uint32 TessellationOptionsHash;
+	FDatasmithImportBaseOptions BaseOptions;
 
 #ifdef CAD_LIBRARY
 	TSharedPtr<FRhinoCoretechWrapper> LocalSession;
@@ -1791,7 +1794,7 @@ TSharedPtr<IDatasmithMeshElement> FOpenNurbsTranslatorImpl::GetMeshElement(const
 	MeshElement->SetLightmapSourceUV(-1);
 
 	TSharedPtr<IDatasmithBaseMaterialElement> Material = FindMaterial(Object);
-	if (Material.IsValid())
+	if (Material.IsValid() && BaseOptions.bIncludeMaterial)
 	{
 		MeshElement->SetMaterial(Material->GetName(), 0);
 	}
@@ -1803,14 +1806,17 @@ TSharedPtr<IDatasmithMeshElement> FOpenNurbsTranslatorImpl::GetMeshElement(const
 	MeshElementToObjectMap.Add(MeshElement.Get(), &Object);
 	MeshElementToTranslatorMap.Add(MeshElement.Get(), this);
 
+	uint8 IncludeMaterial = static_cast<uint8>(BaseOptions.bIncludeMaterial);
+
+	FMD5 MD5;
+	MD5.Update(&IncludeMaterial, sizeof(IncludeMaterial));
+
 	// Use the object's CRC as the mesh element hash
 	uint32 CRC = Object.ObjectPtr->DataCRC(0);
 	if (ON_Brep::Cast(Object.ObjectPtr))
 	{
 		CRC ^= TessellationOptionsHash;
 	}
-
-	FMD5 MD5;
 	MD5.Update(reinterpret_cast<const uint8*>(&CRC), sizeof CRC);
 
 	FMD5Hash Hash;
@@ -2962,6 +2968,11 @@ TOptional< FMeshDescription > FOpenNurbsTranslatorImpl::GetMeshDescription(TShar
 	return bIsValid ? MoveTemp(MeshDescription) : TOptional< FMeshDescription >();
 }
 
+void FOpenNurbsTranslatorImpl::SetBaseOptions(const FDatasmithImportBaseOptions& InBaseOptions)
+{
+	BaseOptions = InBaseOptions;
+}
+
 void FOpenNurbsTranslatorImpl::SetTessellationOptions(const FDatasmithTessellationOptions& Options)
 {
 	TessellationOptions = Options;
@@ -3021,6 +3032,7 @@ bool FDatasmithOpenNurbsTranslator::LoadScene(TSharedRef<IDatasmithScene> OutSce
 	Translator->SetOutputPath(OutputPath);
 
 	Translator->SetTessellationOptions(GetCommonTessellationOptions());
+	Translator->SetBaseOptions(BaseOptions);
 
 	ON_BinaryFile Archive(ON::archive_mode::read3dm, FileHandle);
 
@@ -3080,9 +3092,18 @@ void FDatasmithOpenNurbsTranslator::SetSceneImportOptions(TArray<TStrongObjectPt
 #ifdef USE_OPENNURBS
 	FDatasmithCoreTechTranslator::SetSceneImportOptions(Options);
 
+	for (TStrongObjectPtr<UObject>& Option : Options)
+	{
+		if (UDatasmithImportOptions* DatasmithOptions = Cast<UDatasmithImportOptions>(Option.Get()))
+		{
+			BaseOptions = DatasmithOptions->BaseOptions;
+		}
+	}
+
 	if (Translator)
 	{
 		Translator->SetTessellationOptions( GetCommonTessellationOptions() );
+		Translator->SetBaseOptions(BaseOptions);
 	}
 #endif
 }
diff --git a/Engine/Plugins/Enterprise/DatasmithCADImporter/Source/DatasmithOpenNurbsTranslator/Private/DatasmithOpenNurbsTranslator.h b/Engine/Plugins/Enterprise/DatasmithCADImporter/Source/DatasmithOpenNurbsTranslator/Private/DatasmithOpenNurbsTranslator.h
--- a/Engine/Plugins/Enterprise/DatasmithCADImporter/Source/DatasmithOpenNurbsTranslator/Private/DatasmithOpenNurbsTranslator.h
+++ b/Engine/Plugins/Enterprise/DatasmithCADImporter/Source/DatasmithOpenNurbsTranslator/Private/DatasmithOpenNurbsTranslator.h
@@ -3,6 +3,7 @@
 #pragma once
 
 #include "DatasmithCoreTechTranslator.h"
+#include "DatasmithImportOptions.h"
 
 #include "CoreMinimal.h"
 #include "UObject/ObjectMacros.h"
@@ -40,5 +41,9 @@ protected:
 private:
 
 	TSharedPtr<FOpenNurbsTranslatorImpl> Translator;
+
+	// Temporarily store this here for UE-81278 so that we can trigger the recreation of
+	// static meshes if we're reimporting with new materials that haven't been assigned yet
+	FDatasmithImportBaseOptions BaseOptions;
 };
 
diff --git a/Engine/Plugins/Enterprise/DatasmithImporter/Source/DatasmithImporter/Private/DatasmithImportContext.cpp b/Engine/Plugins/Enterprise/DatasmithImporter/Source/DatasmithImporter/Private/DatasmithImportContext.cpp
--- a/Engine/Plugins/Enterprise/DatasmithImporter/Source/DatasmithImporter/Private/DatasmithImportContext.cpp
+++ b/Engine/Plugins/Enterprise/DatasmithImporter/Source/DatasmithImporter/Private/DatasmithImportContext.cpp
@@ -163,6 +163,9 @@ FDatasmithImportContext::FDatasmithImportContext(const FString& FileName, bool b
 		{
 			ImportOptions.Add(Option.Get());
 		}
+
+		// Temporarily give Rhino translator access to BaseOptions (JIRA UE-81278)
+		AdditionalImportOptions.Add(Options);
 	}
 }
 
