Trivial versionning for CAD worker cache

diff --git a/Engine/Plugins/Enterprise/DatasmithCADImporter/Source/DatasmithCADTranslator/Private/DatasmithCADTranslatorModule.cpp b/Engine/Plugins/Enterprise/DatasmithCADImporter/Source/DatasmithCADTranslator/Private/DatasmithCADTranslatorModule.cpp
--- a/Engine/Plugins/Enterprise/DatasmithCADImporter/Source/DatasmithCADTranslator/Private/DatasmithCADTranslatorModule.cpp
+++ b/Engine/Plugins/Enterprise/DatasmithCADImporter/Source/DatasmithCADTranslator/Private/DatasmithCADTranslatorModule.cpp
@@ -14,10 +14,8 @@ IMPLEMENT_MODULE(FDatasmithCADTranslatorModule, DatasmithCADTranslator);
 void FDatasmithCADTranslatorModule::StartupModule()
 {
 	// Create temporary directory which will be used by CoreTech to store tessellation data
-	TempDir = FPaths::ConvertRelativePathToFull(FPaths::Combine(FPaths::ProjectIntermediateDir(), TEXT("CoreTechImportTemp")));
-	IFileManager::Get().MakeDirectory(*TempDir);
-
-	CacheDir = FPaths::ConvertRelativePathToFull(FPaths::Combine(FPaths::ProjectIntermediateDir(), TEXT("DatasmithCache")));
+	const TCHAR* CacheVersion = TEXT("0");
+	CacheDir = FPaths::ConvertRelativePathToFull(FPaths::Combine(FPaths::ProjectIntermediateDir(), TEXT("DatasmithCADCache"), CacheVersion));
 	IFileManager::Get().MakeDirectory(*CacheDir);
 
 	Datasmith::RegisterTranslator<FDatasmithCADTranslator>();
@@ -28,10 +26,6 @@ void FDatasmithCADTranslatorModule::ShutdownModule()
 	Datasmith::UnregisterTranslator<FDatasmithCADTranslator>();
 }
 
-FString FDatasmithCADTranslatorModule::GetTempDir() const
-{
-	return TempDir;
-}
 
 FString FDatasmithCADTranslatorModule::GetCacheDir() const
 {
diff --git a/Engine/Plugins/Enterprise/DatasmithCADImporter/Source/DatasmithCADTranslator/Public/DatasmithCADTranslatorModule.h b/Engine/Plugins/Enterprise/DatasmithCADImporter/Source/DatasmithCADTranslator/Public/DatasmithCADTranslatorModule.h
--- a/Engine/Plugins/Enterprise/DatasmithCADImporter/Source/DatasmithCADTranslator/Public/DatasmithCADTranslatorModule.h
+++ b/Engine/Plugins/Enterprise/DatasmithCADImporter/Source/DatasmithCADTranslator/Public/DatasmithCADTranslatorModule.h
@@ -29,7 +29,6 @@ public:
 
 	virtual void ShutdownModule() override;
 
-	FString GetTempDir() const;
 	FString GetCacheDir() const;
 
 private:
