Fixed issue with additional geometry appearing in Wire import

#rnx

diff --git a/Engine/Plugins/Enterprise/DatasmithCADImporter/Source/DatasmithWireTranslator/Private/DatasmithWireTranslator.cpp b/Engine/Plugins/Enterprise/DatasmithCADImporter/Source/DatasmithWireTranslator/Private/DatasmithWireTranslator.cpp
--- a/Engine/Plugins/Enterprise/DatasmithCADImporter/Source/DatasmithWireTranslator/Private/DatasmithWireTranslator.cpp
+++ b/Engine/Plugins/Enterprise/DatasmithCADImporter/Source/DatasmithWireTranslator/Private/DatasmithWireTranslator.cpp
@@ -1470,7 +1470,7 @@ bool FWireTranslatorImpl::ProcessBodyNode(TSharedRef<BodyData> Body, const FDagN
 {
 	if (Body->ShellSet.Num() == 1)
 	{
-		RecurseDagForLeavesNoMerge(Body->ShellSet[0], ParentInfo);
+		return RecurseDagForLeavesNoMerge(Body->ShellSet[0], ParentInfo);
 	}
 
 	FDagNodeInfo ShellInfo;
