CAD: Bounding box heuristic to adapt tessellation parameters for small NURBs #jira UE-82561 #rb JeanMichel.Dignard, JeanLuc.Corenthin

diff --git a/Engine/Plugins/Enterprise/DatasmithCADImporter/Source/CADLibrary/Private/CoreTechHelper.cpp b/Engine/Plugins/Enterprise/DatasmithCADImporter/Source/CADLibrary/Private/CoreTechHelper.cpp
--- a/Engine/Plugins/Enterprise/DatasmithCADImporter/Source/CADLibrary/Private/CoreTechHelper.cpp
+++ b/Engine/Plugins/Enterprise/DatasmithCADImporter/Source/CADLibrary/Private/CoreTechHelper.cpp
@@ -46,7 +46,7 @@ struct FVertexData
 		VertexID = FVertexID::Invalid;
 	}
 };
-	
+
 struct FCompareVertexZ
 {
 	FORCEINLINE bool operator()(FVertexData const& A, FVertexData const& B) const { return A.Z < B.Z; }
@@ -335,7 +335,7 @@ bool FillMesh(const FMeshParameters& MeshParameters, const FImportParameters& Im
 				}
 			}
 
-			if (!Tessellation.TexCoordArray.Num()) 
+			if (!Tessellation.TexCoordArray.Num())
 			{
 				for (int32 IndexFace = 0; IndexFace < CTFaceIndex.Num(); IndexFace += 3)
 				{
@@ -471,13 +471,13 @@ void ScaleUV(CT_OBJECT_ID FaceID, FTessellationData& Tessellation, UVType Scale)
 	PuMin = PvMin = HUGE_VALF;
 	PuMax = PvMax = -HUGE_VALF;
 
-	// fast UV min max 
+	// fast UV min max
 	CT_FACE_IO::AskUVminmax(FaceID, PuMin, PuMax, PvMin, PvMax);
 
 	const uint32 NbIsoCurves = 7;
 
 	// Compute Point grid on the restricted surface defined by [PuMin, PuMax], [PvMin, PvMax]
-	CT_OBJECT_ID SurfaceID; 
+	CT_OBJECT_ID SurfaceID;
 	CT_ORIENTATION Orientation;
 	CT_FACE_IO::AskSurface(FaceID, SurfaceID, Orientation);
 
@@ -612,7 +612,7 @@ void SetFaceMainMaterial(FObjectDisplayDataId& InFaceMaterial, FObjectDisplayDat
 		{
 			MaterialHash = *MHash;
 		}
-		else 
+		else
 		{
 			MaterialHash = BuildMaterialHash(InBodyMaterial.MaterialId);
 		}
@@ -721,12 +721,12 @@ TSharedPtr<IDatasmithUEPbrMaterialElement> CreateUEPbrMaterialFromMaterial(FCADM
 		MaterialElement->GetOpacity().SetExpression(Scalar);
 		MaterialElement->SetParentLabel(TEXT("CAD Transparent Material"));
 	}
-	else 
+	else
 	{
 		MaterialElement->SetParentLabel(TEXT("CAD Material"));
 	}
 
-	// Set a Emissive color 
+	// Set a Emissive color
 	if (MaterialElement->GetSpecular().GetExpression() == nullptr)
 	{
 		FLinearColor LinearColor = FLinearColor::FromPow22Color(InMaterial.Specular);
@@ -788,6 +788,37 @@ uint32 GetSize(CT_TESS_DATA_TYPE type)
 	return 0;
 }
 
+// Adapt the parameters to try and make sure a valid tessellation can be generated for the object
+FImportParameters GetAdjustedTessellationParameters(CT_OBJECT_ID MainObjectId, const FImportParameters& ImportParams)
+{
+	FImportParameters Result = ImportParams;
+
+	CT_COORDINATE Min;
+	CT_COORDINATE Max;
+	if (CT_OBJECT_IO::ComputeMinmax(MainObjectId, Min, Max) == CT_IO_ERROR::IO_OK)
+	{
+		double BBSizeMm = (FMath::Abs(Max.xyz[0] - Min.xyz[0]) +
+					       FMath::Abs(Max.xyz[1] - Min.xyz[1]) +
+					       FMath::Abs(Max.xyz[2] - Min.xyz[2])) / 3.0;
+
+		// Heuristic to force smaller tessellation options if the bounding box of the object is small
+		if (BBSizeMm > 0.0 && BBSizeMm < 5.0)
+		{
+			const double ClampedMaxEdgeLength = BBSizeMm / 5.0;
+
+			// Mirror some of the conversions that will happen in SetCoreTechTessellationState
+			if (Result.MaxEdgeLength < SMALL_NUMBER)
+			{
+				Result.MaxEdgeLength = DBL_MAX;
+			}
+			Result.MaxEdgeLength = (Result.MaxEdgeLength / Result.ScaleFactor < ClampedMaxEdgeLength) ?
+									Result.MaxEdgeLength : ClampedMaxEdgeLength * Result.ScaleFactor;
+		}
+	}
+
+	return Result;
+}
+
 CT_IO_ERROR Tessellate(CT_OBJECT_ID MainObjectId, const FImportParameters& ImportParams, FMeshDescription& MeshDesc, FMeshParameters& MeshParameters)
 {
 	CheckedCTError Result;
@@ -797,7 +828,9 @@ CT_IO_ERROR Tessellate(CT_OBJECT_ID MainObjectId, const FImportParameters& Impor
 	if (!Result)
 		return Result;
 
-	SetCoreTechTessellationState(ImportParams);
+	FImportParameters FixedParams = GetAdjustedTessellationParameters(MainObjectId, ImportParams);
+
+	SetCoreTechTessellationState(FixedParams);
 
 	TArray<FTessellationData> FaceTessellations;
 	TArray<uint32> FaceColors;
