Fix file naming for cache mechanism
The file name for cache mechanism of scene graph description file is a hash code based on  file stat and, if exist, file configuration (solidworks). But this hash does not include stitching technique. Or stitching can modified scene graph structure.

diff --git a/Engine/Plugins/Enterprise/DatasmithCADImporter/Source/CADInterfaces/Private/CoreTechFileParser.cpp b/Engine/Plugins/Enterprise/DatasmithCADImporter/Source/CADInterfaces/Private/CoreTechFileParser.cpp
--- a/Engine/Plugins/Enterprise/DatasmithCADImporter/Source/CADInterfaces/Private/CoreTechFileParser.cpp
+++ b/Engine/Plugins/Enterprise/DatasmithCADImporter/Source/CADInterfaces/Private/CoreTechFileParser.cpp
@@ -154,7 +154,7 @@ namespace {
 	}
 }
 
-uint32 GetFileHash(const FString& FileName, const FFileStatData& FileStatData, const FString& Config)
+uint32 GetFileHash(const FString& FileName, const FFileStatData& FileStatData, const FString& Config, const FImportParameters& ImportParam)
 {
 	int64 FileSize = FileStatData.FileSize;
 	FDateTime ModificationTime = FileStatData.ModificationTime;
@@ -162,6 +162,7 @@ uint32 GetFileHash(const FString& FileName, const FFileStatData& FileStatData, c
 	uint32 FileHash = GetTypeHash(FileName);
 	FileHash = HashCombine(FileHash, GetTypeHash(FileSize));
 	FileHash = HashCombine(FileHash, GetTypeHash(ModificationTime));
+	FileHash = HashCombine(FileHash, GetTypeHash(ImportParam.StitchingTechnique));
 	if (!Config.IsEmpty())
 	{
 		FileHash = HashCombine(FileHash, GetTypeHash(Config));
@@ -663,7 +664,7 @@ EProcessState FCoreTechFileParser::ProcessFile()
 	}
 
 	FFileStatData FileStatData = IFileManager::Get().GetStatData(*FullPath);
-	uint32 FileHash = GetFileHash(CADFile, FileStatData, FileConfiguration);
+	uint32 FileHash = GetFileHash(CADFile, FileStatData, FileConfiguration, ImportParameters);
 
 	SceneGraphFile = FString::Printf(TEXT("UEx%08x"), FileHash);
 
