Inventor files are like solidworks YUp_RightHanded and Its colors have to be manage also like Solidworks

diff --git a/Engine/Plugins/Enterprise/DatasmithCADImporter/Source/DatasmithCADTranslator/Private/DatasmithCADTranslator.cpp b/Engine/Plugins/Enterprise/DatasmithCADImporter/Source/DatasmithCADTranslator/Private/DatasmithCADTranslator.cpp
--- a/Engine/Plugins/Enterprise/DatasmithCADImporter/Source/DatasmithCADTranslator/Private/DatasmithCADTranslator.cpp
+++ b/Engine/Plugins/Enterprise/DatasmithCADImporter/Source/DatasmithCADTranslator/Private/DatasmithCADTranslator.cpp
@@ -78,7 +78,7 @@ bool FDatasmithCADTranslator::LoadScene(TSharedRef<IDatasmithScene> DatasmithSce
 	}
 
 	ImportParameters.ModelCoordSys = CADLibrary::EModelCoordSystem::ZUp_RightHanded;
-	if (FileExtension == TEXT("sldprt") || FileExtension == TEXT("sldasm"))
+	if (FileExtension == TEXT("sldprt") || FileExtension == TEXT("sldasm") || FileExtension == TEXT("iam") || FileExtension == TEXT("ipt"))
 	{
 		ImportParameters.ModelCoordSys = CADLibrary::EModelCoordSystem::YUp_RightHanded;
 		ImportParameters.DisplayPreference = CADLibrary::EDisplayPreference::ColorOnly;
