Datasmith CAD Translator
- Enable multiprocessing
- Simplify DatasmithDispatcher::ICommand: standard FArchive serialization
- New Message over sockets system (DatasmithDispatcher::FNetworkNode)
- New CommandQueue system (DatasmithDispatcher::FCommandQueue)
- Removed DatasmithDispatcherSocket class
- fix DatasmithCADWorker executable destination (replaces the postbuild move)
- DatasmithWorkerHandlers now ows their associated thread
- FDatasmithDispatcher

#rnx

diff --git a/Engine/Plugins/Enterprise/DatasmithCADImporter/Source/CADInterfaces/Private/CoreTechFileParser.cpp b/Engine/Plugins/Enterprise/DatasmithCADImporter/Source/CADInterfaces/Private/CoreTechFileParser.cpp
--- a/Engine/Plugins/Enterprise/DatasmithCADImporter/Source/CADInterfaces/Private/CoreTechFileParser.cpp
+++ b/Engine/Plugins/Enterprise/DatasmithCADImporter/Source/CADInterfaces/Private/CoreTechFileParser.cpp
@@ -1,4 +1,5 @@
 // Copyright 1998-2019 Epic Games, Inc. All Rights Reserved.
+
 #include "CoreTechFileParser.h"
 
 #ifdef CAD_INTERFACE
@@ -6,9 +7,7 @@
 
 #include "CADData.h"
 #include "CADOptions.h"
-#include "Containers/Array.h"
-#include "Containers/Map.h"
-#include "Containers/Queue.h"
+
 #include "CoreTechTypes.h"
 #include "GenericPlatform/GenericPlatformFile.h"
 #include "HAL/FileManager.h"
@@ -16,7 +15,6 @@
 #include "Misc/FileHelper.h"
 #include "Misc/Paths.h"
 
-#define SGSIZE 100000
 #define EXTREFNUM 5000
 
 namespace CADLibrary 
@@ -644,7 +642,7 @@ void GetRawDataFileExternalRef(const FString& InRawDataFile, TSet<FString>& Exte
 	}
 }
 
-EProcessState FCoreTechFileParser::ProcessFile()
+FCoreTechFileParser::EProcessResult FCoreTechFileParser::ProcessFile()
 {
 	FileConfiguration.Empty();
 
@@ -660,7 +658,7 @@ EProcessState FCoreTechFileParser::ProcessFile()
 
 	if (!IFileManager::Get().FileExists(*FullPath))
 	{
-		return FileNotFound;
+		return EProcessResult::FileNotFound;
 	}
 
 	FFileStatData FileStatData = IFileManager::Get().GetStatData(*FullPath);
@@ -696,17 +694,14 @@ EProcessState FCoreTechFileParser::ProcessFile()
 	{
 		// The file has been yet proceed, get ExternalRef
 		GetRawDataFileExternalRef(RawDataFile, ExternalRefSet);
-		return ProcessOk;
+		return EProcessResult::ProcessOk;
 	}
 
 	// Process the file
 	return ReadFileWithKernelIO();
 }
 
-
-
-
-EProcessState FCoreTechFileParser::ReadFileWithKernelIO()
+FCoreTechFileParser::EProcessResult FCoreTechFileParser::ReadFileWithKernelIO()
 {
 	CT_IO_ERROR Result = IO_OK;
 	CT_OBJECT_ID MainId = 0;
@@ -743,7 +738,7 @@ EProcessState FCoreTechFileParser::ReadFileWithKernelIO()
 		Result = CT_KERNEL_IO::UnloadModel();
 		if (Result != IO_OK)
 		{
-			return ProcessFailed;
+			return EProcessResult::ProcessFailed;
 		}
 		Result = CT_KERNEL_IO::LoadFile(*FullPath, MainId, CTImportOption | CT_LOAD_FLAGS_LOAD_EXTERNAL_REF);
 	}
@@ -751,7 +746,7 @@ EProcessState FCoreTechFileParser::ReadFileWithKernelIO()
 	if (Result != IO_OK && Result != IO_OK_MISSING_LICENSES)
 	{
 		CT_KERNEL_IO::UnloadModel();
-		return ProcessFailed;
+		return EProcessResult::ProcessFailed;
 	}
 
 	Repair(MainId, ImportParameters.StitchingTechnique);
@@ -794,7 +789,7 @@ EProcessState FCoreTechFileParser::ReadFileWithKernelIO()
 	ReadMaterial();
 
 	// Parse the file
-	bool Ret = ReadNode(MainId);
+	bool bReadNodeSucceed = ReadNode(MainId);
 	// End of parsing
 
 	if (bNeedSaveCTFile)
@@ -807,9 +802,9 @@ EProcessState FCoreTechFileParser::ReadFileWithKernelIO()
 
 	CT_KERNEL_IO::UnloadModel();
 
-	if (!Ret)
+	if (!bReadNodeSucceed)
 	{
-		return ProcessFailed;
+		return EProcessResult::ProcessFailed;
 	}
 
 	uint32 LastLine = SceneGraphDescription.Num();
@@ -845,7 +840,7 @@ EProcessState FCoreTechFileParser::ReadFileWithKernelIO()
 	
 	ExportFileSceneGraph();
 
-	return ProcessOk;
+	return EProcessResult::ProcessOk;
 }
 
 CT_FLAGS FCoreTechFileParser::SetCoreTechImportOption(const FString& MainFileExt)
diff --git a/Engine/Plugins/Enterprise/DatasmithCADImporter/Source/CADInterfaces/Public/CoreTechFileParser.h b/Engine/Plugins/Enterprise/DatasmithCADImporter/Source/CADInterfaces/Public/CoreTechFileParser.h
--- a/Engine/Plugins/Enterprise/DatasmithCADImporter/Source/CADInterfaces/Public/CoreTechFileParser.h
+++ b/Engine/Plugins/Enterprise/DatasmithCADImporter/Source/CADInterfaces/Public/CoreTechFileParser.h
@@ -15,43 +15,14 @@
 #define MAPCTIDLINE     8
 
 #ifdef CAD_INTERFACE
-#pragma warning(push)
-#pragma warning(disable:4996) // unsafe sprintf
-#pragma warning(disable:4828) // illegal character
-#include "kernel_io/kernel_io.h"
-#include "kernel_io/kernel_io_error.h"
+#include "kernel_io/attribute_io/attribute_enum.h"
+#include "kernel_io/ct_types.h"
 #include "kernel_io/kernel_io_type.h"
-#include "kernel_io/list_io/list_io.h"
-#include "kernel_io/material_io/material_io.h"
-#include "kernel_io/object_io/asm_io/component_io/component_io.h"
-#include "kernel_io/object_io/asm_io/instance_io/instance_io.h"
-#include "kernel_io/object_io/geom_io/curve_io/cnurbs_io/cnurbs_io.h"
-#include "kernel_io/object_io/geom_io/surface_io/snurbs_io/snurbs_io.h"
-#include "kernel_io/object_io/geom_io/surface_io/surface_io.h"
-#include "kernel_io/object_io/history_io/hnode_io/hbranch_io/hbranch_io.h"
-#include "kernel_io/object_io/history_io/hnode_io/hleaf_io/hleaf_io.h"
-#include "kernel_io/object_io/history_io/solid_io/solid_io.h"
-#include "kernel_io/object_io/topo_io/body_io/body_io.h"
-#include "kernel_io/object_io/topo_io/coedge_io/coedge_io.h"
-#include "kernel_io/object_io/topo_io/face_io/face_io.h"
-#include "kernel_io/object_io/topo_io/loop_io/loop_io.h"
-#include "kernel_io/repair_io/repair_io.h"
-#pragma warning(pop)
 #endif // CAD_INTERFACE
 
 
 struct FFileStatData;
 
-enum EProcessState  // #ueent_CAD
-{
-	Unknown = 0x00,
-	UnTreated = 0x01,
-	Running = 0x02,
-	ProcessOk = 0x04,
-	ProcessFailed = 0x08,
-	FileNotFound = 0x10
-};
-
 #ifdef CAD_INTERFACE
 namespace CADLibrary
 {
@@ -81,11 +52,21 @@ uint32 GetSize(CT_TESS_DATA_TYPE type);
 class CADINTERFACES_API FCoreTechFileParser
 {
 public:
+	enum class EProcessResult
+	{
+		Unknown,
+		Running,
+		UnTreated,
+		ProcessOk,
+		ProcessFailed,
+		FileNotFound,
+	};
+
 	FCoreTechFileParser(const FString& InCTFullPath, const FString& InCachePath, const FImportParameters& ImportParams, bool bPreferBodyData, bool bScaleUVMap);
 
 	double GetMetricUnit() const { return 0.01; }
 
-	EProcessState ProcessFile();
+	EProcessResult ProcessFile();
 
 	const TSet<FString>& GetExternalRefSet()
 	{
@@ -105,11 +86,8 @@ public:
 		return CADFile;
 	}
 
-
-
 private:
-	EProcessState ReadFileWithKernelIO();
-
+	EProcessResult ReadFileWithKernelIO();
 	bool ReadNode(CT_OBJECT_ID NodeId);
 	bool ReadInstance(CT_OBJECT_ID NodeId);
 	bool ReadComponent(CT_OBJECT_ID NodeId);
diff --git a/Engine/Plugins/Enterprise/DatasmithCADImporter/Source/CADTools/Private/CADData.cpp b/Engine/Plugins/Enterprise/DatasmithCADImporter/Source/CADTools/Private/CADData.cpp
--- a/Engine/Plugins/Enterprise/DatasmithCADImporter/Source/CADTools/Private/CADData.cpp
+++ b/Engine/Plugins/Enterprise/DatasmithCADImporter/Source/CADTools/Private/CADData.cpp
@@ -120,10 +120,9 @@ void WriteTessellationInRawData(FTessellationData& Tessellation, TArray<uint8>&
 
 bool ReadTessellationSetFromFile(TArray<uint8>& RawData, uint32& OutNbBodies, TArray<FBodyMesh>& CTBodySet)
 {
-	uint32 INFO[LastLine];
-	const uint32 InfoSize = sizeof(INFO);
-	int32 RawDataOffset = sizeof(uint32);
-	memcpy(&OutNbBodies, RawData.GetData(), RawDataOffset);
+	int32 RawDataOffset = 0;
+	memcpy(&OutNbBodies, RawData.GetData() + RawDataOffset, sizeof(uint32));
+	RawDataOffset += sizeof(uint32);
 
 	CTBodySet.Reserve(OutNbBodies);
 
@@ -131,54 +130,66 @@ bool ReadTessellationSetFromFile(TArray<uint8>& RawData, uint32& OutNbBodies, TA
 
 	int32 CurrentBodyId = -1;
 	int32 CurrentBodyIndex = -1;
-	while (RawDataOffset < RawData.Num())
+	uint32 BodyHeader[LastLine];
+	const uint32 BodyHeaderByteSize = sizeof(BodyHeader);
+
+	while (RawDataOffset + int32(BodyHeaderByteSize) <= RawData.Num())
 	{
-		memcpy(INFO, RawData.GetData() + RawDataOffset, InfoSize);
-		RawDataOffset += InfoSize;
+		memcpy(BodyHeader, RawData.GetData() + RawDataOffset, BodyHeaderByteSize);
+
+		int32 BlocByteSize  = BodyHeader[LineRawSize];
+		if (RawDataOffset + BlocByteSize > RawData.Num())
+		{
+			return false;
+		}
+
+		RawDataOffset += BodyHeaderByteSize;
 
-		int32 Types = INFO[LineType];
-		int32 RawSize = INFO[LineRawSize];
+		int32 Types = BodyHeader[LineType];
+		int32 RawSize = BodyHeader[LineRawSize];
 
 		if (Types != ETessellationContent::Default)
 		{
 			return false;
 		}
 
-		if (INFO[LineBodyId] != CurrentBodyId)
+		if (BodyHeader[LineBodyId] != CurrentBodyId)
 		{
-			CurrentBodyId = INFO[LineBodyId];
-			CTBodySet.Emplace(INFO[LineBodyId], INFO[LineBodyFaceNum]);
+			CurrentBodyId = BodyHeader[LineBodyId];
+			CTBodySet.Emplace(BodyHeader[LineBodyId], BodyHeader[LineBodyFaceNum]);
 			CurrentBodyIndex++;
 
 			CurrentBody = &CTBodySet[CurrentBodyIndex];
-			CurrentBody->SetBodyId(INFO[LineBodyId]);
-			CurrentBody->SetBodyUuid(INFO[LineBodyUuId]);
-			CurrentBody->GetTessellationSet().Reserve(INFO[LineBodyFaceNum]);
+			CurrentBody->SetBodyId(BodyHeader[LineBodyId]);
+			CurrentBody->SetBodyUuid(BodyHeader[LineBodyUuId]);
+			CurrentBody->GetTessellationSet().Reserve(BodyHeader[LineBodyFaceNum]);
 		}
 
-		int32 Index = CurrentBody->GetTessellationSet().Emplace();
-		FTessellationData& Tessellation = CurrentBody->GetTessellationSet()[Index];
-
-		Tessellation.BodyId = INFO[LineBodyId];
-		Tessellation.BodyUuId = INFO[LineBodyUuId];
-		Tessellation.BodyFaceNum = INFO[LineBodyFaceNum];
-		Tessellation.MaterialId = INFO[LineMaterialId];
-		Tessellation.MaterialHash = INFO[LineMaterialHash];
-		Tessellation.VertexCount = INFO[LineVertexCount];
-		Tessellation.SizeOfVertexType = INFO[LineVertexType];
-		Tessellation.NormalCount = INFO[LineNormalCount];
-		Tessellation.SizeOfNormalType = INFO[LineNormalType];
-		Tessellation.IndexCount = INFO[LineIndexCount];
-		Tessellation.SizeOfIndexType = INFO[LineIndexType];
-		Tessellation.TexCoordCount = INFO[LineTexCoordCount];
-		Tessellation.SizeOfTexCoordType = INFO[LineTexCoordType];
+		if (CurrentBody == nullptr)
+		{
+			return false;
+		}
+
+		FTessellationData& Tessellation = CurrentBody->GetTessellationSet().Emplace_GetRef();
+
+		Tessellation.BodyId = BodyHeader[LineBodyId];
+		Tessellation.BodyUuId = BodyHeader[LineBodyUuId];
+		Tessellation.BodyFaceNum = BodyHeader[LineBodyFaceNum];
+		Tessellation.MaterialId = BodyHeader[LineMaterialId];
+		Tessellation.MaterialHash = BodyHeader[LineMaterialHash];
+		Tessellation.VertexCount = BodyHeader[LineVertexCount];
+		Tessellation.SizeOfVertexType = BodyHeader[LineVertexType];
+		Tessellation.NormalCount = BodyHeader[LineNormalCount];
+		Tessellation.SizeOfNormalType = BodyHeader[LineNormalType];
+		Tessellation.IndexCount = BodyHeader[LineIndexCount];
+		Tessellation.SizeOfIndexType = BodyHeader[LineIndexType];
+		Tessellation.TexCoordCount = BodyHeader[LineTexCoordCount];
+		Tessellation.SizeOfTexCoordType = BodyHeader[LineTexCoordType];
 
 		uint32 VertexSize = Tessellation.VertexCount * Tessellation.SizeOfVertexType * 3;
 		uint32 NormalSize = Tessellation.NormalCount * Tessellation.SizeOfNormalType * 3;
 		uint32 IndexSize = Tessellation.IndexCount * Tessellation.SizeOfIndexType;
 		uint32 TexCoordSize = Tessellation.TexCoordCount * Tessellation.SizeOfTexCoordType * 2;
-		uint32 DataByteSize = VertexSize + NormalSize + IndexSize + TexCoordSize + InfoSize;
-
 
 		Tessellation.VertexArray.Empty(VertexSize);
 		Tessellation.NormalArray.Empty(NormalSize);
diff --git a/Engine/Plugins/Enterprise/DatasmithCADImporter/Source/DatasmithCADTranslator/Private/DatasmithCADTranslator.cpp b/Engine/Plugins/Enterprise/DatasmithCADImporter/Source/DatasmithCADTranslator/Private/DatasmithCADTranslator.cpp
--- a/Engine/Plugins/Enterprise/DatasmithCADImporter/Source/DatasmithCADTranslator/Private/DatasmithCADTranslator.cpp
+++ b/Engine/Plugins/Enterprise/DatasmithCADImporter/Source/DatasmithCADTranslator/Private/DatasmithCADTranslator.cpp
@@ -11,10 +11,6 @@
 #include "Misc/FileHelper.h"
 
 
-FDatasmithCADTranslator::FDatasmithCADTranslator()
-{
-}
-
 void FDatasmithCADTranslator::Initialize(FDatasmithTranslatorCapabilities& OutCapabilities)
 {
 	OutCapabilities.SupportedFileFormats.Add(FFileFormatInfo{ TEXT("CATPart"), TEXT("CATIA Part files") });
@@ -56,19 +52,15 @@ void FDatasmithCADTranslator::Initialize(FDatasmithTranslatorCapabilities& OutCa
 	OutCapabilities.SupportedFileFormats.Add(FFileFormatInfo{ TEXT("dgn"), TEXT("MicroStation files") });
 }
 
-bool FDatasmithCADTranslator::IsSourceSupported(const FDatasmithSceneSource& Source)
-{
-	return true;
-}
-
 bool FDatasmithCADTranslator::LoadScene(TSharedRef<IDatasmithScene> DatasmithScene)
 {
 	ImportParameters.MetricUnit = 0.001;
 	ImportParameters.ScaleFactor = 0.1;
-	ImportParameters.ChordTolerance = GetCommonTessellationOptions().ChordTolerance;
-	ImportParameters.MaxEdgeLength = GetCommonTessellationOptions().MaxEdgeLength;
-	ImportParameters.MaxNormalAngle = GetCommonTessellationOptions().NormalTolerance;
-	ImportParameters.StitchingTechnique = (CADLibrary::EStitchingTechnique) GetCommonTessellationOptions().StitchingTechnique;
+	const FDatasmithTessellationOptions& TesselationOptions = GetCommonTessellationOptions();
+	ImportParameters.ChordTolerance = TesselationOptions.ChordTolerance;
+	ImportParameters.MaxEdgeLength = TesselationOptions.MaxEdgeLength;
+	ImportParameters.MaxNormalAngle = TesselationOptions.NormalTolerance;
+	ImportParameters.StitchingTechnique = (CADLibrary::EStitchingTechnique) TesselationOptions.StitchingTechnique;
 
 	FString FileExtension = GetSource().GetSourceFileExtension();
 	if (FileExtension == TEXT("jt"))
@@ -90,14 +82,13 @@ bool FDatasmithCADTranslator::LoadScene(TSharedRef<IDatasmithScene> DatasmithSce
 	TMap<FString, FString> CADFileToUE4FileMap;
 
 	int32 NumCores = FPlatformMisc::NumberOfCores();
-	DatasmithDispatcher::FDatasmithDispatcher Dispatcher(CachePath, NumCores, CADFileToUE4FileMap, CADFileToUE4GeomMap);
-
-	Dispatcher.Init(ImportParameters);
-	Dispatcher.AddTask(FPaths::ConvertRelativePathToFull(GetSource().GetSourceFile()));
+	{
+		DatasmithDispatcher::FDatasmithDispatcher Dispatcher(ImportParameters, CachePath, NumCores, CADFileToUE4FileMap, CADFileToUE4GeomMap);
+		Dispatcher.AddTask(FPaths::ConvertRelativePathToFull(GetSource().GetSourceFile()));
 
-	bool bEnableMultiProcess = (NumCores != 1);
-	Dispatcher.Process(false && bEnableMultiProcess);
-	Dispatcher.Clear();
+		bool bWithProcessor = true;
+		Dispatcher.Process(bWithProcessor);
+	}
 
 	FDatasmithSceneGraphBuilder SceneGraphBuilder(CADFileToUE4FileMap, MeshElementToCADBRepUuidMap,  CachePath, DatasmithScene, GetSource(), ImportParameters);
 	SceneGraphBuilder.Build();
diff --git a/Engine/Plugins/Enterprise/DatasmithCADImporter/Source/DatasmithCADTranslator/Private/DatasmithCADTranslator.h b/Engine/Plugins/Enterprise/DatasmithCADImporter/Source/DatasmithCADTranslator/Private/DatasmithCADTranslator.h
--- a/Engine/Plugins/Enterprise/DatasmithCADImporter/Source/DatasmithCADTranslator/Private/DatasmithCADTranslator.h
+++ b/Engine/Plugins/Enterprise/DatasmithCADImporter/Source/DatasmithCADTranslator/Private/DatasmithCADTranslator.h
@@ -25,11 +25,8 @@ public:
 #ifndef CAD_LIBRARY
 	virtual void Initialize(FDatasmithTranslatorCapabilities& OutCapabilities) override { OutCapabilities.bIsEnabled = false; }
 #else
-	FDatasmithCADTranslator();
 	virtual void Initialize(FDatasmithTranslatorCapabilities& OutCapabilities) override;
 
-	virtual bool IsSourceSupported(const FDatasmithSceneSource& Source) override;
-
 	virtual bool LoadScene(TSharedRef<IDatasmithScene> OutScene) override;
 
 	virtual void UnloadScene() override;
diff --git a/Engine/Plugins/Enterprise/DatasmithCADImporter/Source/DatasmithCADTranslator/Private/DatasmithSceneGraphBuilder.cpp b/Engine/Plugins/Enterprise/DatasmithCADImporter/Source/DatasmithCADTranslator/Private/DatasmithSceneGraphBuilder.cpp
--- a/Engine/Plugins/Enterprise/DatasmithCADImporter/Source/DatasmithCADTranslator/Private/DatasmithSceneGraphBuilder.cpp
+++ b/Engine/Plugins/Enterprise/DatasmithCADImporter/Source/DatasmithCADTranslator/Private/DatasmithSceneGraphBuilder.cpp
@@ -68,8 +68,6 @@ void FCADSceneGraphDescriptionFile::GetMaterialDescription(int32 LineNumber, CAD
 	OutSpecular[0] = OutSpecular[1] = OutSpecular[2] = 255;
 	OutShininess = OutReflexion = 0;
 	OutTransparency = 255;
-	OutMaterialName = TEXT("");
-	OutTextureName = TEXT("");
 
 	const FString& MaterialLine1 = GetString(LineNumber);
 	FString MaterialIdStr;
@@ -146,6 +144,11 @@ void FCADSceneGraphDescriptionFile::ReadFile()
 	TArray<FString> HeaderCTIdSplit;
 	SceneGraphDescriptionData[Line].ParseIntoArray(HeaderCTIdSplit, TEXT(" "));
 
+	if (HeaderCTIdSplit.Num() % 2 != 0)
+	{
+		return;
+	}
+
 	for (int32 index = 0; index < HeaderCTIdSplit.Num(); index += 2)
 	{
 		SceneNodeIdToSceneNodeDescriptionMap.Emplace(FCString::Atoi(*HeaderCTIdSplit[index]), FSceneNodeDescription(*this, FCString::Atoi(*HeaderCTIdSplit[index + 1])));
@@ -184,7 +187,6 @@ void FCADSceneGraphDescriptionFile::ReadFile()
 			MaterialString = SceneGraphDescriptionData[Line];
 		}
 	}
-
 }
 
 void FCADSceneGraphDescriptionFile::SetMaterialMaps(TMap< uint32, FColor>& MaterialUuidToColor, TMap< uint32, CADLibrary::FCADMaterial>& MaterialUuidToMaterial)
@@ -201,11 +203,10 @@ void FCADSceneGraphDescriptionFile::SetMaterialMaps(TMap< uint32, FColor>& Mater
 	}
 }
 
-
-void FSceneNodeDescription::GetMetaDatas(TMap<FString, FString>& MetaDataMap)
+void FSceneNodeDescription::GetMetaData(TMap<FString, FString>& MetaDataMap)
 {
 	const FString& RawDataString = SceneGraphDescription.GetString(Line+1);
-	if (RawDataString.GetCharArray()[0] != L'M')
+	if (RawDataString[0] != L'M')
 	{
 		return;
 	}
@@ -246,11 +247,11 @@ void FSceneNodeDescription::GetChildren(TArray<int32>& Children)
 	if(EndMeta == -1)
 	{
 		TMap<FString, FString> TempMap;
-		GetMetaDatas(TempMap);
+		GetMetaData(TempMap);
 	}
 
 	const FString& DescriptionString = SceneGraphDescription.GetString(EndMeta);
-	if (DescriptionString.GetCharArray()[0] != L'c')
+	if (DescriptionString[0] != L'c')
 	{
 		return;
 	}
@@ -276,7 +277,7 @@ void FSceneNodeDescription::GetNodeReference(int32& RefId, FString& ExternalFile
 	if (EndMeta == -1)
 	{
 		TMap<FString, FString> TempMap;
-		GetMetaDatas(TempMap);
+		GetMetaData(TempMap);
 	}
 
 	const FString& DescriptionString = SceneGraphDescription.GetString(EndMeta + 1);
@@ -351,7 +352,7 @@ void FSceneNodeDescription::SetNodeType()
 	else
 	{
 		const FString& RawDataString = SceneGraphDescription.GetString(Line);
-		switch (RawDataString.GetCharArray()[0])
+		switch (RawDataString[0])
 		{
 		case L'C':
 			Type = COMPONENT;
@@ -478,22 +479,20 @@ TSharedPtr< IDatasmithActorElement >  FDatasmithSceneGraphBuilder::BuildNode(FSc
 
 	case BODY:
 		return BuildBody(Node, ParentData);
+
 	default:
 		return TSharedPtr< IDatasmithActorElement >();
 	}
-
-	return TSharedPtr< IDatasmithActorElement >();
 }
 
 TSharedPtr< IDatasmithActorElement >  FDatasmithSceneGraphBuilder::BuildInstance(FSceneNodeDescription& Node, ActorData& ParentData)
 {
 	TMap<FString, FString> InstanceNodeMetaDataMap;
 	TMap<FString, FString> ReferenceNodeMetaDataMap;
-	Node.GetMetaDatas(InstanceNodeMetaDataMap);
-
-	FString ActorUUID = TEXT("");
-	FString ActorLabel = TEXT("");
+	Node.GetMetaData(InstanceNodeMetaDataMap);
 
+	FString ActorUUID;
+	FString ActorLabel;
 	GetNodeUUIDAndName(InstanceNodeMetaDataMap, ReferenceNodeMetaDataMap, ParentData.Uuid, ActorUUID, ActorLabel);
 
 	int32 RefId; 
@@ -532,7 +531,7 @@ TSharedPtr< IDatasmithActorElement >  FDatasmithSceneGraphBuilder::BuildInstance
 
 	if (ComponentNode)
 	{
-		ComponentNode->GetMetaDatas(ReferenceNodeMetaDataMap);
+		ComponentNode->GetMetaData(ReferenceNodeMetaDataMap);
 	}
 
 	TSharedPtr< IDatasmithActorElement > Actor = CreateActor(*ActorUUID, *ActorLabel);
@@ -638,7 +637,7 @@ TSharedPtr< IDatasmithActorElement > FDatasmithSceneGraphBuilder::BuildComponent
 {
 	TMap<FString, FString> InstanceNodeMetaDataMap;
 	TMap<FString, FString> ReferenceNodeMetaDataMap;
-	Node.GetMetaDatas(ReferenceNodeMetaDataMap);
+	Node.GetMetaData(ReferenceNodeMetaDataMap);
 
 	FString ActorUUID = TEXT("");
 	FString ActorLabel = TEXT("");
@@ -665,7 +664,7 @@ TSharedPtr< IDatasmithActorElement > FDatasmithSceneGraphBuilder::BuildBody(FSce
 {
 	TMap<FString, FString> InstanceNodeMetaDataMap;
 	TMap<FString, FString> BodyNodeMetaDataMap;
-	Node.GetMetaDatas(BodyNodeMetaDataMap);
+	Node.GetMetaData(BodyNodeMetaDataMap);
 	
 	FString BodyUUID = TEXT("");
 	FString BodyLabel = TEXT("");
@@ -851,7 +850,7 @@ void FDatasmithSceneGraphBuilder::AddMetaData(TSharedPtr< IDatasmithActorElement
 
 	for (auto& Attribute : ReferenceNodeAttributeSetMap)
 	{
-		if (UnwantedAttributes.Contains(*Attribute.Key))
+		if (UnwantedAttributes.Contains(Attribute.Key))
 		{
 			continue;
 		}
diff --git a/Engine/Plugins/Enterprise/DatasmithCADImporter/Source/DatasmithCADTranslator/Private/DatasmithSceneGraphBuilder.h b/Engine/Plugins/Enterprise/DatasmithCADImporter/Source/DatasmithCADTranslator/Private/DatasmithSceneGraphBuilder.h
--- a/Engine/Plugins/Enterprise/DatasmithCADImporter/Source/DatasmithCADTranslator/Private/DatasmithSceneGraphBuilder.h
+++ b/Engine/Plugins/Enterprise/DatasmithCADImporter/Source/DatasmithCADTranslator/Private/DatasmithSceneGraphBuilder.h
@@ -41,7 +41,7 @@ public:
 
 	NODE_TYPE GetNodeType() { return Type; }
 
-	void GetMetaDatas(TMap<FString, FString>& MetaDataMap);
+	void GetMetaData(TMap<FString, FString>& MetaDataMap);
 	void GetChildren(TArray<int32>& Children);
 	
 	void GetNodeReference(int32& OutRefId, FString& OutExternalFile, NODE_TYPE& OutType);
@@ -122,13 +122,14 @@ public:
 	}
 
 	bool GetColor(uint32 ColorHId, FColor& Color) const;
-	bool GetMaterial(int32 MaterialId, CADLibrary::FCADMaterial& material) const;
 
-	void GetMaterialDescription(int32 LineNumber, CADLibrary::FCADMaterial& Material) const;
+	bool GetMaterial(int32 MaterialId, CADLibrary::FCADMaterial& OutMaterial) const;
 
 	void SetMaterialMaps(TMap< uint32, FColor>& MaterialUuidToColor, TMap< uint32, CADLibrary::FCADMaterial>& MaterialUuidToMaterial);
 
-protected:
+private:
+	void GetMaterialDescription(int32 LineNumber, CADLibrary::FCADMaterial& OutMaterial) const;
+
 	FString FileName;
 	TArray<FString> SceneGraphDescriptionData;
 	TMap<int32, FSceneNodeDescription> SceneNodeIdToSceneNodeDescriptionMap;
diff --git a/Engine/Plugins/Enterprise/DatasmithCADImporter/Source/DatasmithDispatcher/Private/DatasmithCommands.cpp b/Engine/Plugins/Enterprise/DatasmithCADImporter/Source/DatasmithDispatcher/Private/DatasmithCommands.cpp
--- a/Engine/Plugins/Enterprise/DatasmithCADImporter/Source/DatasmithDispatcher/Private/DatasmithCommands.cpp
+++ b/Engine/Plugins/Enterprise/DatasmithCADImporter/Source/DatasmithDispatcher/Private/DatasmithCommands.cpp
@@ -1,170 +1,71 @@
 // Copyright 1998-2019 Epic Games, Inc. All Rights Reserved.
+
 #include "DatasmithCommands.h"
 
-#include "DatasmithDispatcherSocket.h"
-#include "DatasmithDispatcher.h"
+
+#include "Serialization/Archive.h"
+#include "Serialization/MemoryWriter.h"
+#include "Serialization/MemoryReader.h"
 
 namespace DatasmithDispatcher
 {
-	FDatasmithCommandManager::FDatasmithCommandManager(FDatasmithDispatcherSocket& InSocket)
-		: Socket(InSocket)
-		, CurrentCommandeSize(0)
-	{
-		Initialize();
-	}
-
-	FDatasmithCommandManager::~FDatasmithCommandManager()
-	{
-		Terminate();
-	}
-
-	void FDatasmithCommandManager::Initialize()
-	{
-		CommandMap.Reserve((int32)DatasmithCommandType::Last);
-
-		CommandMap.Add(Ping, new FDatasmithPingCommand());
-		CommandMap.Add(BackPing, new FDatasmithBackPingCommand());
-		CommandMap.Add(RunTask, new FDatasmithRunTaskCommand());
-		CommandMap.Add(NotifyEndTask, new FDatasmithNotifyEndTaskCommand());
-		CommandMap.Add(ImportParams, new FDatasmithImportParametersCommand());
-	}
-
-	void FDatasmithCommandManager::Terminate()
-	{
-		for (auto ItemCommand : CommandMap)
-		{
-			delete ItemCommand.Value;
-		}
-		CommandMap.Empty();
-	}
-
-	ICommand * FDatasmithCommandManager::GetNextCommand()
+	TSharedPtr<ICommand> CreateCommand(ECommandId CommandType)
 	{
-		uint32 DataSize;
-		if (!Socket.HasPendingData(DataSize))
-		{
-			return nullptr;
-		}
-
-		if (!CurrentCommandeSize)
-		{
-			Socket >> CurrentCommandeSize;
-		}
-
-		if (!Socket.HasPendingData(DataSize))
-		{
-			return nullptr;
-		}
-
-		// All data is not yet arrived, wait
-		if (DataSize < CurrentCommandeSize)
+		switch (CommandType)
 		{
-			return nullptr;
+			case ECommandId::Ping: return MakeShared<FPingCommand>();
+			case ECommandId::BackPing: return MakeShared<FBackPingCommand>();
+			case ECommandId::RunTask: return MakeShared<FRunTaskCommand>();
+			case ECommandId::NotifyEndTask: return MakeShared<FCompletedTaskCommand>();
+			case ECommandId::ImportParams: return MakeShared<FImportParametersCommand>();
+			case ECommandId::Terminate: return MakeShared<FTerminateCommand>();
 		}
-
-		DatasmithCommandType CommandType;
-		Socket >> CommandType;
-
-		ICommand** Command = CommandMap.Find(CommandType);
-		if (Command == nullptr)
-		{
-			return nullptr;
-		}
-		 
-		(*Command)->Initialize();
-		(*Command)->Read(Socket);
-
-		CurrentCommandeSize = 0;
-
-		return (*Command);
-	}
-
-	void FDatasmithBackPingCommand::Write(FDatasmithDispatcherSocket &CurrentSocket)
-	{
-		CurrentSocket << GetType();
-		CurrentSocket.SendData();
-	}
-
-	void FDatasmithBackPingCommand::Read(FDatasmithDispatcherSocket &CurrentSocket)
-	{
-	}
-
-	void FDatasmithPingCommand::Write(FDatasmithDispatcherSocket &CurrentSocket)
-	{
-		CurrentSocket << GetType();
-		CurrentSocket.SendData();
+		return nullptr;
 	}
 
-	void FDatasmithPingCommand::Read(FDatasmithDispatcherSocket &CurrentSocket)
+	void SerializeCommand(ICommand& Command, TArray<uint8>& OutBuffer)
 	{
+		FMemoryWriter ArWriter(OutBuffer);
+		uint8 type = static_cast<uint8>(Command.GetType());
+		ArWriter << type;
+		ArWriter << Command;
 	}
 
-	void FDatasmithRunTaskCommand::Write(FDatasmithDispatcherSocket&  CurrentSocket)
+	TSharedPtr<ICommand> DeserializeCommand(const TArray<uint8>& InBuffer)
 	{
-		CurrentSocket << GetType();
-		CurrentSocket << JobFilePath;
-		CurrentSocket << JobIndex;					
-		CurrentSocket.SendData();
-	}
-
-	void FDatasmithRunTaskCommand::Read(FDatasmithDispatcherSocket& CurrentSocket)
-	{
-		CurrentSocket >> JobFilePath;
-		CurrentSocket >> JobIndex;
-	}
-
-  	void FDatasmithNotifyEndTaskCommand::Write(FDatasmithDispatcherSocket& CurrentSocket)
-	{
-		CurrentSocket << GetType();
-
-		CurrentSocket << (uint32)ExternalReferenceSet.Num();
-		for (const FString& ExternalReference : ExternalReferenceSet)
+		FMemoryReader ArReader(InBuffer);
+		uint8 type = 0;
+		ArReader << type;
+		if (TSharedPtr<ICommand> Command = CreateCommand(static_cast<ECommandId>(type)))
 		{
-			CurrentSocket << ExternalReference;
+			ArReader << *Command;
+			return ArReader.ArIsError ? nullptr : Command;
 		}
-		CurrentSocket << ProcessResult;
-		CurrentSocket << SceneGraphFileName;
-		CurrentSocket << GeomFileName;
-		CurrentSocket.SendData();
+		return nullptr;
 	}
 
-	void FDatasmithNotifyEndTaskCommand::Read(FDatasmithDispatcherSocket& CurrentSocket)
+	void FRunTaskCommand::SerializeImpl(FArchive& Ar)
 	{
-		uint32 NbReferences;
-
-		CurrentSocket >> NbReferences;
-		FString Buffer;
-		for (uint32 ite = 0; ite < NbReferences; ++ite)
-		{
-			CurrentSocket >> Buffer;
-			ExternalReferenceSet.Add(Buffer);
-		}
-		CurrentSocket >> ProcessResult;
-		CurrentSocket >> SceneGraphFileName;
-		CurrentSocket >> GeomFileName;
+		Ar << JobFilePath;
+		Ar << JobIndex;
 	}
 
-
-	void FDatasmithImportParametersCommand::Write(FDatasmithDispatcherSocket&  CurrentSocket)
+  	void FCompletedTaskCommand::SerializeImpl(FArchive& Ar)
 	{
-		CurrentSocket << GetType();
-		CurrentSocket << ImportParameters.ChordTolerance;
-		CurrentSocket << ImportParameters.MaxEdgeLength;
-		CurrentSocket << ImportParameters.MaxNormalAngle;
-		CurrentSocket << ImportParameters.MetricUnit;
-		CurrentSocket << ImportParameters.ScaleFactor;
-		CurrentSocket << ImportParameters.StitchingTechnique;
-		CurrentSocket.SendData();
+		Ar << ExternalReferences;
+		Ar << ProcessResult;
+		Ar << SceneGraphFileName;
+		Ar << GeomFileName;
 	}
 
-	void FDatasmithImportParametersCommand::Read(FDatasmithDispatcherSocket&  CurrentSocket)
+	void FImportParametersCommand::SerializeImpl(FArchive& Ar)
 	{
-		CurrentSocket >> ImportParameters.ChordTolerance;
-		CurrentSocket >> ImportParameters.MaxEdgeLength;
-		CurrentSocket >> ImportParameters.MaxNormalAngle;
-		CurrentSocket >> ImportParameters.MetricUnit;
-		CurrentSocket >> ImportParameters.ScaleFactor;
-		CurrentSocket >> ImportParameters.StitchingTechnique;
+		Ar << ImportParameters.ChordTolerance;
+		Ar << ImportParameters.MaxEdgeLength;
+		Ar << ImportParameters.MaxNormalAngle;
+		Ar << ImportParameters.MetricUnit;
+		Ar << ImportParameters.ScaleFactor;
+		Ar << (uint32&)ImportParameters.StitchingTechnique;
 	}
 
-}
\ No newline at end of file
+}
diff --git a/Engine/Plugins/Enterprise/DatasmithCADImporter/Source/DatasmithDispatcher/Private/DatasmithDispatcher.cpp b/Engine/Plugins/Enterprise/DatasmithCADImporter/Source/DatasmithDispatcher/Private/DatasmithDispatcher.cpp
--- a/Engine/Plugins/Enterprise/DatasmithCADImporter/Source/DatasmithDispatcher/Private/DatasmithDispatcher.cpp
+++ b/Engine/Plugins/Enterprise/DatasmithCADImporter/Source/DatasmithDispatcher/Private/DatasmithDispatcher.cpp
@@ -1,199 +1,184 @@
 // Copyright 1998-2019 Epic Games, Inc. All Rights Reserved.
+
 #include "DatasmithDispatcher.h"
 
 #include "CoreTechFileParser.h"
+#include "DatasmithDispatcherTask.h"
+#include "DatasmithDispatcherLog.h"
 
-#include "Containers/Array.h"
 #include "HAL/FileManager.h"
-#include "HAL/Thread.h"
-#include "Misc/Optional.h"
 #include "Misc/Paths.h"
 #include "Misc/ScopeLock.h"
+#include "Algo/Count.h"
 
+DEFINE_LOG_CATEGORY(LogDatasmithDispatcher);
 
 namespace DatasmithDispatcher
 {
 
-FDatasmithDispatcher::FDatasmithDispatcher(const FString& InCacheDir, int32 InNumberOfWorkers, TMap<FString, FString>& OutCADFileToUnrealFileMap, TMap<FString, FString>& OutCADFileToUnrealGeomMap)
-	: CADFileToUnrealFileMap(OutCADFileToUnrealFileMap)
+FDatasmithDispatcher::FDatasmithDispatcher(const CADLibrary::FImportParameters& InImportParameters, const FString& InCacheDir, int32 InNumberOfWorkers, TMap<FString, FString>& OutCADFileToUnrealFileMap, TMap<FString, FString>& OutCADFileToUnrealGeomMap)
+	: NextTaskIndex(0)
+	, CompletedTaskCount(0)
+	, CADFileToUnrealFileMap(OutCADFileToUnrealFileMap)
 	, CADFileToUnrealGeomMap(OutCADFileToUnrealGeomMap)
 	, ProcessCacheFolder(InCacheDir)
+	, ImportParameters(InImportParameters)
 	, NumberOfWorkers(InNumberOfWorkers)
-	, LastLaunchTask(0)
-	, NumOfFinishedTask(0)
-{
-	WorkerHandlerSet.Reserve(NumberOfWorkers);
-	for (int32 Index = 0; Index < NumberOfWorkers; Index++)
-	{
-		WorkerHandlerSet.Emplace(*this, ProcessCacheFolder);
-	}
-}
-
-void FDatasmithDispatcher::Init(const CADLibrary::FImportParameters& InImportParameters)
 {
-	ImportParameters = InImportParameters;
-
+	// init cache folders
 	IFileManager::Get().MakeDirectory(*FPaths::Combine(ProcessCacheFolder, TEXT("scene")), true);
 	IFileManager::Get().MakeDirectory(*FPaths::Combine(ProcessCacheFolder, TEXT("cad")), true);
 	IFileManager::Get().MakeDirectory(*FPaths::Combine(ProcessCacheFolder, TEXT("mesh")), true);
 	IFileManager::Get().MakeDirectory(*FPaths::Combine(ProcessCacheFolder, TEXT("body")), true);
 }
 
-void FDatasmithDispatcher::RunHandlers()
+void FDatasmithDispatcher::AddTask(const FString& InFile)
 {
-	WorkersRunning.Reserve(NumberOfWorkers);
-	int32 WorkerIndex = 0;
-	for (FDatasmithWorkerHandler& Handler : WorkerHandlerSet)
+	FScopeLock Lock(&TaskPoolCriticalSection);
+	for (const FTask& Task : TaskPool)
 	{
-		FString ProcessName = TEXT("DatasmithWorker ") + FString::FromInt(WorkerIndex);
-		WorkerIndex++;
-		WorkersRunning.Emplace(*ProcessName, [&]() { Handler.Run(); });
+		if (Task.FileName == InFile)
+		{
+			return;
+		}
 	}
+
+	int32 TaskIndex = TaskPool.Emplace(InFile);
+	TaskPool[TaskIndex].Index = TaskIndex;
 }
 
-void FDatasmithDispatcher::Close()
+TOptional<FTask> FDatasmithDispatcher::GetNextTask()
 {
-	for (FDatasmithWorkerHandler& Handler : WorkerHandlerSet)
+	FScopeLock Lock(&TaskPoolCriticalSection);
+
+	while (TaskPool.IsValidIndex(NextTaskIndex) && TaskPool[NextTaskIndex].State != ETaskState::UnTreated)
 	{
-		Handler.NotifyKill();
+		NextTaskIndex++;
 	}
 
-	for (FThread& Process : WorkersRunning)
+	if (!TaskPool.IsValidIndex(NextTaskIndex))
 	{
-		Process.Join();
+		return TOptional<FTask>();
 	}
-}
 
-void FDatasmithDispatcher::Clear()
-{
-	TaskPool.Empty();
-	LastLaunchTask = 0;
-	NumOfFinishedTask = 0;
+	TaskPool[NextTaskIndex].State = ETaskState::Running;
+	return TaskPool[NextTaskIndex++];
 }
 
-
-void FDatasmithDispatcher::Process(bool bWithProcessor)
+void FDatasmithDispatcher::SetTaskState(int32 TaskIndex, ETaskState TaskState)
 {
-	if (bWithProcessor)
+	FScopeLock Lock(&TaskPoolCriticalSection);
+
+	if (!ensure(TaskPool.IsValidIndex(TaskIndex)))
 	{
-		RunHandlers();  
+		return;
+	}
 
-		while (!IsOver())
-		{
-			FWindowsPlatformProcess::Sleep(0.1);
-		} 
+	TaskPool[TaskIndex].State = TaskState;
 
-		Close();
+	if (TaskState == ETaskState::ProcessOk
+	 || TaskState == ETaskState::ProcessFailed
+	 || TaskState == ETaskState::FileNotFound)
+	{
+		CompletedTaskCount++;
 	}
-	else
+
+	if (TaskState == ETaskState::UnTreated)
 	{
-		ProcessLocal();
+		NextTaskIndex = TaskIndex;
 	}
 }
 
-void FDatasmithDispatcher::ProcessLocal()
+void FDatasmithDispatcher::Process(bool bWithProcessor)
 {
-#ifdef CAD_INTERFACE
-	TOptional<FTask> Task = GetTask();
-	while (Task.IsSet())
+	if (bWithProcessor)
 	{
-		FString FullPath = Task->FileName;
-
-		CADLibrary::FCoreTechFileParser FileParser(FullPath, ProcessCacheFolder, ImportParameters, true, true);
-		EProcessState TaskState = FileParser.ProcessFile();
-		SetTaskState(Task->Index, TaskState);
+		SpawnHandlers();
 
-		if (TaskState == ProcessOk)
+		while (!IsOver() && GetAliveHandlerCount() > 0)
 		{
-			FString CurrentPath = FPaths::GetPath(FullPath);
-			const TSet<FString>& ExternalRefSet = FileParser.GetExternalRefSet();
-			if (ExternalRefSet.Num() > 0)
-			{
-				for (const FString& ExternalFile : ExternalRefSet)
-				{
-					AddTask(FPaths::Combine(CurrentPath, ExternalFile));
-				}
-			}
-			LinkCTFileToUnrealCacheFile(FileParser.GetFileName(), FileParser.GetSceneGraphFile(), FileParser.GetMeshFile());
+			FWindowsPlatformProcess::Sleep(0.5);
 		}
 
-		Task = GetTask();
+		CloseHandlers();
 	}
-#endif // CAD_INTERFACE
-}
 
-void FDatasmithDispatcher::AddTask(const FString& InFile)
-{
-	FScopeLock Lock(&TaskPoolCriticalSection);
-	for (const FTask& Task : TaskPool)
+	if (!IsOver())
 	{
-		if (Task.FileName == InFile)
-		{
-			return;
-		}
+		ProcessLocal();
 	}
-
-	int32 TaskIndex = TaskPool.Emplace(InFile);
-	TaskPool[TaskIndex].Index = TaskIndex;
 }
 
-TOptional<FTask> FDatasmithDispatcher::GetTask()
+bool FDatasmithDispatcher::IsOver()
 {
 	FScopeLock Lock(&TaskPoolCriticalSection);
+	return CompletedTaskCount == TaskPool.Num();
+}
 
-	if (!TaskPool.IsValidIndex(LastLaunchTask))
+void FDatasmithDispatcher::LinkCTFileToUnrealCacheFile(const FString& CTFile, const FString& UnrealSceneGraphFile, const FString& UnrealMeshFile)
+{
+	FScopeLock Lock(&TaskPoolCriticalSection);
+	if (!UnrealSceneGraphFile.IsEmpty())
 	{
-		return TOptional<FTask>();
+		CADFileToUnrealFileMap.Add(CTFile, UnrealSceneGraphFile);
+	}
+	if (!UnrealMeshFile.IsEmpty())
+	{
+		CADFileToUnrealGeomMap.Add(CTFile, UnrealMeshFile);
 	}
-
-	TaskPool[LastLaunchTask].State = Running;
-	return TaskPool[LastLaunchTask++];
 }
 
-bool FDatasmithDispatcher::IsOver() 
+void FDatasmithDispatcher::SpawnHandlers()
 {
-	FScopeLock Lock(&TaskPoolCriticalSection);
-	return NumOfFinishedTask == TaskPool.Num();
+	WorkerHandlers.Reserve(NumberOfWorkers);
+	for (int32 Index = 0; Index < NumberOfWorkers; Index++)
+	{
+		WorkerHandlers.Emplace(*this, ProcessCacheFolder, Index);
+	}
 }
 
-void FDatasmithDispatcher::SetTessellationOptions(float InChordTolerance, float InMaxEdgeLength, float InNormalTolerance, CADLibrary::EStitchingTechnique InStitchingTechnique)
+int32 FDatasmithDispatcher::GetAliveHandlerCount()
 {
-	ImportParameters.ChordTolerance = InChordTolerance;
-	ImportParameters.MaxEdgeLength = InMaxEdgeLength;
-	ImportParameters.MaxNormalAngle = InNormalTolerance;
-	ImportParameters.StitchingTechnique = InStitchingTechnique;
+	return Algo::CountIf(WorkerHandlers, [](const FDatasmithWorkerHandler& Handler) { return Handler.IsAlive(); });
 }
 
-
-void FDatasmithDispatcher::SetTaskState(int32 TaskIndex, EProcessState TaskState)
+void FDatasmithDispatcher::CloseHandlers()
 {
-	FScopeLock Lock(&TaskPoolCriticalSection);
-
-	if (TaskIndex < 0)
+	for (FDatasmithWorkerHandler& Handler : WorkerHandlers)
 	{
-		return;
-	}
-	if (TaskIndex < TaskPool.Num())
-	{
-		TaskPool[TaskIndex].State = TaskState;
-		if (TaskState & (ProcessOk | ProcessFailed | FileNotFound))
-		{
-			NumOfFinishedTask++;
-		}
+		Handler.Stop();
 	}
+	WorkerHandlers.Empty();
 }
 
-void FDatasmithDispatcher::LinkCTFileToUnrealCacheFile(const FString& CTFile, const FString& UnrealSceneGraphFile, const FString& UnrealMeshFile)
+void FDatasmithDispatcher::ProcessLocal()
 {
-	FScopeLock Lock(&TaskPoolCriticalSection);
-	if (!UnrealSceneGraphFile.IsEmpty())
-	{
-		CADFileToUnrealFileMap.Add(CTFile, UnrealSceneGraphFile);
-	}
-	if (!UnrealMeshFile.IsEmpty())
+#ifdef CAD_INTERFACE
+	while (TOptional<FTask> Task = GetNextTask())
 	{
-		CADFileToUnrealGeomMap.Add(CTFile, UnrealMeshFile);
+		FString FullPath = Task->FileName;
+
+		CADLibrary::FCoreTechFileParser FileParser(FullPath, ProcessCacheFolder, ImportParameters, true, true);
+		ETaskState ProcessResult = FileParser.ProcessFile();
+
+		ETaskState TaskState = ProcessResult;
+		SetTaskState(Task->Index, TaskState);
+
+		if (TaskState == ETaskState::ProcessOk)
+		{
+			FString CurrentPath = FPaths::GetPath(FullPath);
+			const TSet<FString>& ExternalRefSet = FileParser.GetExternalRefSet();
+			if (ExternalRefSet.Num() > 0)
+			{
+				for (const FString& ExternalFile : ExternalRefSet)
+				{
+					AddTask(FPaths::Combine(CurrentPath, ExternalFile));
+				}
+			}
+			LinkCTFileToUnrealCacheFile(FileParser.GetFileName(), FileParser.GetSceneGraphFile(), FileParser.GetMeshFile());
+		}
 	}
+#endif // CAD_INTERFACE
 }
 
-}
\ No newline at end of file
+} // ns DatasmithDispatcher
\ No newline at end of file
diff --git a/Engine/Plugins/Enterprise/DatasmithCADImporter/Source/DatasmithDispatcher/Private/DatasmithDispatcherNetworking.cpp b/Engine/Plugins/Enterprise/DatasmithCADImporter/Source/DatasmithDispatcher/Private/DatasmithDispatcherNetworking.cpp
new file mode 100644
--- /dev/null
+++ b/Engine/Plugins/Enterprise/DatasmithCADImporter/Source/DatasmithDispatcher/Private/DatasmithDispatcherNetworking.cpp
@@ -0,0 +1,371 @@
+// Copyright 1998-2019 Epic Games, Inc. All Rights Reserved.
+
+#include "DatasmithDispatcherNetworking.h"
+
+#include "DatasmithDispatcherLog.h"
+
+#include "CoreMinimal.h"
+#include "SocketSubsystem.h"
+#include "Sockets.h"
+#include "Serialization/MemoryReader.h"
+#include "Serialization/MemoryWriter.h"
+
+
+namespace DatasmithDispatcher
+{
+
+static FTimespan TimespanFromSeconds(double Seconds)
+{
+	return FTimespan(int64(Seconds * ETimespan::TicksPerSecond));
+}
+
+FNetworkNode::~FNetworkNode()
+{
+	CloseSocket(ConnectedSocket);
+}
+
+bool FNetworkNode::SendMessage(const TArray<uint8>& Message, double Timeout_s)
+{
+	if (bWriteError)
+	{
+		UE_LOG(LogDatasmithDispatcher, Error, TEXT("bWriteError flag raised, can't write"));
+		return false;
+	}
+
+	if (ConnectedSocket == nullptr)
+	{
+		UE_LOG(LogDatasmithDispatcher, Error, TEXT("null socket, can't write"));
+		return false;
+	}
+
+	ESocketConnectionState State = ConnectedSocket->GetConnectionState();
+	if (State == ESocketConnectionState::SCS_ConnectionError)
+	{
+		UE_LOG(LogDatasmithDispatcher, Error, TEXT("connection state error"));
+		return false;
+	}
+
+	bool bCanWrite = ConnectedSocket->Wait(ESocketWaitConditions::WaitForWrite, TimespanFromSeconds(Timeout_s));
+	if (!bCanWrite)
+	{
+		UE_LOG(LogDatasmithDispatcher, Warning, TEXT("can't write on socket"));
+		return false;
+	}
+
+	ConnectedSocket->SetNonBlocking(false);
+
+	{ // Send the header
+		FMessageHeader Header;
+		Header.ByteSize = Message.Num();
+
+		TArray<uint8> HeaderBuffer;
+		FMemoryWriter ArWriter(HeaderBuffer);
+		ArWriter << Header;
+
+		int32 ByteSent = -1;
+		bool bSendSucceed = ConnectedSocket->Send(HeaderBuffer.GetData(), HeaderBuffer.Num(), ByteSent);
+
+		if (!bSendSucceed || HeaderBuffer.Num() != ByteSent)
+		{
+			UE_LOG(LogDatasmithDispatcher, Error, TEXT("can't write header on socket"));
+			bWriteError = true;
+			return false;
+		}
+	}
+
+	{ // Send the content
+		const uint8* Data = Message.GetData();
+		uint32 BufferSize = Message.Num();
+
+		uint32 TotalByteSent = 0;
+		bool bSendSucceed = true;
+		while (bSendSucceed && TotalByteSent < BufferSize)
+		{
+			int32 BytesSent = 0;
+			bSendSucceed &= ConnectedSocket->Send(Data + TotalByteSent, BufferSize - TotalByteSent, BytesSent);
+			TotalByteSent += (uint32) BytesSent;
+		}
+
+		if (!bSendSucceed)
+		{
+			UE_LOG(LogDatasmithDispatcher, Error, TEXT("can't write content on socket"));
+			bWriteError = true;
+			return false;
+		}
+	}
+
+	return true;
+}
+
+bool FNetworkNode::ReceiveMessage(TArray<uint8>& OutMessage, double Timeout_s)
+{
+	if (bReadError)
+	{
+		UE_LOG(LogDatasmithDispatcher, Error, TEXT("ReadError flag raised, can't read"));
+		return false;
+	}
+
+	ConnectedSocket->SetNonBlocking(true);
+
+	bool bCanRead = ConnectedSocket && ConnectedSocket->Wait(ESocketWaitConditions::WaitForRead, TimespanFromSeconds(Timeout_s));
+	if (!bCanRead)
+	{
+		return false;
+	}
+
+	// Make sure we have a valid header
+	if (IncommingMessage.Header.ByteSize < 0)
+	{
+		static uint32 HeaderByteSize = []()
+		{
+			FMessageHeader Header;
+			TArray<uint8> HeaderBuffer;
+			FMemoryWriter ArWriter(HeaderBuffer);
+			ArWriter << Header;
+			return HeaderBuffer.Num();
+		}();
+
+		uint32 PendingByteCount;
+		if (!ConnectedSocket->HasPendingData(PendingByteCount) || PendingByteCount < HeaderByteSize)
+		{
+			return false;
+		}
+
+		TArray<uint8> HeaderBuffer;
+		HeaderBuffer.AddZeroed(HeaderByteSize);
+		int32 BytesRead = -1;
+		bool bRecvSucceed = ConnectedSocket->Recv(HeaderBuffer.GetData(), HeaderBuffer.Num(), BytesRead);
+		if (!bRecvSucceed || BytesRead != HeaderByteSize)
+		{
+			UE_LOG(LogDatasmithDispatcher, Error, TEXT("Parsed header failed"));
+			bReadError = true;
+			return false;
+		}
+
+		FMemoryReader ArReader(HeaderBuffer);
+		ArReader << IncommingMessage.Header;
+
+		if (IncommingMessage.Header.ByteSize < 0 || IncommingMessage.Header.ByteSize > 1<<20)
+		{
+			UE_LOG(LogDatasmithDispatcher, Error, TEXT("Parsed header failed: bad message size %d"), IncommingMessage.Header.ByteSize);
+			bReadError = true;
+			return false;
+		}
+
+		IncommingMessage.Content.Reserve(IncommingMessage.Header.ByteSize);
+		UE_LOG(LogDatasmithDispatcher, Display, TEXT("Parsed header ok"));
+	}
+
+	// fill the message with available data
+	uint32 PendingByteCount;
+	if (IncommingMessage.Header.ByteSize >= 0 && ConnectedSocket->HasPendingData(PendingByteCount))
+	{
+		int32 MissingByteInCurrentMessage = IncommingMessage.Header.ByteSize - IncommingMessage.Content.Num();
+		int32 ReadTarget = int32(PendingByteCount) < MissingByteInCurrentMessage ? PendingByteCount : MissingByteInCurrentMessage;
+		uint8* Destination = IncommingMessage.Content.GetData() + IncommingMessage.Content.Num();
+		IncommingMessage.Content.AddUninitialized(ReadTarget);
+
+		int32 BytesRead = -1;
+		bool bRecvSucceed = ConnectedSocket->Recv(Destination, ReadTarget, BytesRead);
+		if (BytesRead != ReadTarget)
+		{
+			bReadError = true;
+			UE_LOG(LogDatasmithDispatcher, Error, TEXT("Recv issue"));
+			return false;
+		}
+
+		bool bMessageCompleted = BytesRead == MissingByteInCurrentMessage;
+		if (bMessageCompleted)
+		{
+			OutMessage = MoveTemp(IncommingMessage.Content);
+			IncommingMessage = FMessage();
+			return true;
+		}
+	}
+
+	return false;
+}
+
+
+FSocket* FNetworkNode::CreateInternalSocket(const FString& Description)
+{
+	FSocket* Socket = nullptr;
+	ISocketSubsystem* SocketSubsystem = ISocketSubsystem::Get();
+	TSharedRef<FInternetAddr> InternetAddress = SocketSubsystem->CreateInternetAddr();
+	InternetAddress->SetLoopbackAddress();
+
+	FString a = InternetAddress->ToString(true);
+
+	FName Protocol = NAME_None;
+	if (ensure(InternetAddress->IsValid()))
+	{
+		Protocol = InternetAddress->GetProtocolType();
+	}
+	FSocket* NewSocket = SocketSubsystem->CreateSocket(NAME_Stream, *Description, Protocol);
+	UE_CLOG(!NewSocket, LogDatasmithDispatcher, Error, TEXT("Socket creation failure"));
+	return NewSocket;
+}
+
+void FNetworkNode::CloseSocket(FSocket*& Socket)
+{
+	if (Socket)
+	{
+		Socket->Close();
+		ISocketSubsystem::Get()->DestroySocket(Socket);
+		Socket = nullptr;
+	}
+}
+
+bool FNetworkNode::IsConnected()
+{
+	return ConnectedSocket && ConnectedSocket->GetConnectionState() == ESocketConnectionState::SCS_Connected;
+}
+
+FNetworkServerNode::FNetworkServerNode()
+{
+	ListeningSocket = CreateInternalSocket(TEXT("datasmith listening socket"));
+
+	if (!ListeningSocket)
+	{
+		ConnectedSocketError = SocketErrorCode::Error_Create;
+		return;
+	}
+
+	ISocketSubsystem* SocketSubsystem = ISocketSubsystem::Get();
+	TSharedRef<FInternetAddr> InternetAddress = SocketSubsystem->CreateInternetAddr();
+	InternetAddress->SetLoopbackAddress();
+
+	FString a = InternetAddress->ToString(true);
+
+	int32 BoundPort = SocketSubsystem->BindNextPort(ListeningSocket, *InternetAddress, 1e3, 1);
+	if (BoundPort == 0)
+	{
+		ConnectedSocketError = SocketErrorCode::Error_Bind;
+		UE_LOG(LogDatasmithDispatcher, Error, TEXT("Socket binding failure"));
+		return;
+	}
+	ensure(BoundPort == ListeningSocket->GetPortNo());
+
+	if (!ListeningSocket->Listen(0))
+	{
+		ConnectedSocketError = SocketErrorCode::Error_Listen;
+		UE_LOG(LogDatasmithDispatcher, Error, TEXT("Socket listen failure"));
+		return;
+	}
+}
+
+FNetworkServerNode::~FNetworkServerNode()
+{
+	CloseSocket(ListeningSocket);
+}
+
+int32 FNetworkServerNode::GetListeningPort()
+{
+	return ListeningSocket ? ListeningSocket->GetPortNo() : 0;
+}
+
+bool FNetworkServerNode::Accept(const FString& Description, double Timeout_s)
+{
+	CloseSocket(ConnectedSocket);
+
+	// wait until a connection occurs
+	ListeningSocket->SetNonBlocking();
+	double AcceptLimit_s = FPlatformTime::Seconds() + Timeout_s;
+	do
+	{
+		ConnectedSocket = ListeningSocket->Accept(Description);
+		if (ConnectedSocket)
+		{
+			return true;
+		}
+		FPlatformProcess::Sleep(0.1);
+	}
+	while(FPlatformTime::Seconds() < AcceptLimit_s);
+
+	return false;
+}
+
+bool FNetworkClientNode::Connect(const FString& Description, int32 ServerPort, double Timeout_s)
+{
+	CloseSocket(ConnectedSocket);
+	ConnectedSocket = CreateInternalSocket(Description);
+
+	if (!ConnectedSocket)
+	{
+		ConnectedSocketError = SocketErrorCode::Error_Create;
+		return false;
+	}
+
+	ISocketSubsystem* SocketSubsystem = ISocketSubsystem::Get();
+	TSharedRef<FInternetAddr> InternetAddress = SocketSubsystem->CreateInternetAddr();
+	InternetAddress->SetLoopbackAddress();
+	InternetAddress->SetPort(ServerPort);
+
+	double ConnectTimeout_s = FPlatformTime::Seconds() + Timeout_s;
+	ConnectedSocket->SetNonBlocking(true);
+	do
+	{
+		if (ConnectedSocket->Connect(*InternetAddress))
+		{
+			UE_LOG(LogDatasmithDispatcher, Display, TEXT("Is connected"));
+			return true;
+		}
+	}
+	while(FPlatformTime::Seconds() < ConnectTimeout_s);
+
+	UE_LOG(LogDatasmithDispatcher, Error, TEXT("Client socket failed to connect"));
+	CloseSocket(ConnectedSocket);
+	return false;
+}
+
+void FCommandQueue::SetNetworkInterface(FNetworkNode* InNetworkInterface)
+{
+	NetworkInterface = InNetworkInterface;
+}
+
+TSharedPtr<ICommand> FCommandQueue::GetNextCommand(double Timeout_s)
+{
+	// consume all available commands
+	while (Poll(Timeout_s))
+	{
+		Timeout_s = 0;
+	}
+
+	TSharedPtr<ICommand> OutCommand;
+	return InCommands.Dequeue(OutCommand) ? OutCommand : nullptr;
+}
+
+bool FCommandQueue::SendCommand(ICommand& Commmand, double Timeout_s)
+{
+	TArray<uint8> CommandBuffer;
+	SerializeCommand(Commmand, CommandBuffer);
+	return NetworkInterface && NetworkInterface->SendMessage(CommandBuffer, Timeout_s);
+}
+
+bool FCommandQueue::Disconnect(double Timeout_s)
+{
+	// consume all available commands
+	while (Poll(Timeout_s))
+	{
+		Timeout_s = 0;
+	}
+
+	// close the network interface
+	return true;
+}
+
+bool FCommandQueue::Poll(double Timeout_s)
+{
+	TArray<uint8> CommandBuffer;
+	if (NetworkInterface && NetworkInterface->ReceiveMessage(CommandBuffer, Timeout_s))
+	{
+		if (TSharedPtr<ICommand> Command = DeserializeCommand(CommandBuffer))
+		{
+			InCommands.Enqueue(Command);
+			return true;
+		}
+	}
+	return false;
+}
+
+} // ns DatasmithDispatcher
diff --git a/Engine/Plugins/Enterprise/DatasmithCADImporter/Source/DatasmithDispatcher/Private/DatasmithDispatcherSocket.cpp b/Engine/Plugins/Enterprise/DatasmithCADImporter/Source/DatasmithDispatcher/Private/DatasmithDispatcherSocket.cpp
deleted file mode 100644
diff --git a/Engine/Plugins/Enterprise/DatasmithCADImporter/Source/DatasmithDispatcher/Private/DatasmithWorkerHandler.cpp b/Engine/Plugins/Enterprise/DatasmithCADImporter/Source/DatasmithDispatcher/Private/DatasmithWorkerHandler.cpp
--- a/Engine/Plugins/Enterprise/DatasmithCADImporter/Source/DatasmithDispatcher/Private/DatasmithWorkerHandler.cpp
+++ b/Engine/Plugins/Enterprise/DatasmithCADImporter/Source/DatasmithDispatcher/Private/DatasmithWorkerHandler.cpp
@@ -1,190 +1,330 @@
 // Copyright 1998-2019 Epic Games, Inc. All Rights Reserved.
 #include "DatasmithWorkerHandler.h"
 
-#include "HAL/FileManager.h"
-#include "HAL/PlatformProcess.h"
-#include "Misc/Paths.h"
-#include "SocketSubsystem.h"
-#include "Sockets.h"
 #include "DatasmithDispatcher.h"
 #include "DatasmithDispatcherLog.h"
 #include "DatasmithCommands.h"
 
+#include "HAL/FileManager.h"
+#include "HAL/PlatformProcess.h"
+#include "Misc/Paths.h"
+#include "Sockets.h"
+#include "SocketSubsystem.h"
 
 namespace DatasmithDispatcher
 {
-FString FDatasmithWorkerHandler::GetExePath()
+
+static FString GetWorkerExecutablePath()
 {
-	FString ProcessorPath;
+	FString ProcessorPath = FPaths::Combine(FPaths::EnginePluginsDir(), TEXT("Enterprise/DatasmithCADImporter/Binaries"));
+
 #if PLATFORM_MAC
-	ProcessorPath = FPaths::Combine(FPaths::EngineDir(), TEXT("Binaries/Mac/DatasmithCADWorker"));
+	ProcessorPath = FPaths::Combine(ProcessorPath, TEXT("Mac/DatasmithCADWorker"));
 #elif PLATFORM_LINUX
-	ProcessorPath = FPaths::Combine(FPaths::EngineDir(), TEXT("Binaries/Linux/DatasmithCADWorker"));
-#else
-	ProcessorPath = FPaths::Combine(FPaths::EngineDir(), TEXT("Binaries/Win64/DatasmithCADWorker.exe"));
+	ProcessorPath = FPaths::Combine(ProcessorPath, TEXT("Linux/DatasmithCADWorker"));
+#elif PLATFORM_WINDOWS
+	ProcessorPath = FPaths::Combine(ProcessorPath, TEXT("Win64/DatasmithCADWorker.exe"));
 #endif
 
-	if (FPaths::FileExists(ProcessorPath))
-	{
-		return ProcessorPath;
-	}
-	return FString();
+	return ProcessorPath;
 }
 
-FDatasmithWorkerHandler::FDatasmithWorkerHandler(FDatasmithDispatcher& InDispatcher, const FString& InCachePath)
+FDatasmithWorkerHandler::FDatasmithWorkerHandler(FDatasmithDispatcher& InDispatcher, const FString& InCachePath, uint32 Id)
 	: Dispatcher(InDispatcher)
-	, WorkerProcessId(0)
-	, ServerSocket(TEXT("127.0.0.1"))
-	, ClientListener()
+	, WorkerState(EWorkerState::Uninitialized)
+	, ErrorState(EWorkerErrorState::Ok)
 	, CachePath(InCachePath)
 	, bShouldTerminate(false)
 {
+	CommandIO.SetNetworkInterface(&NetworkInterface);
+	ThreadName = FString(TEXT("DatasmithWorkerHandler_")) + FString::FromInt(Id);
+	IOThread = FThread(*ThreadName, [this]() { Run(); } );
 }
 
-bool FDatasmithWorkerHandler::StartWorker()
+FDatasmithWorkerHandler::~FDatasmithWorkerHandler()
 {
-	FString ProcessorPath = GetExePath();
-	if (!ProcessorPath.Len())
-	{
-		return false;
-	}
+	IOThread.Join();
+}
 
-	FString WorkerAbsolutePath = IFileManager::Get().ConvertToAbsolutePathForExternalAppForWrite(*ProcessorPath);
-	FPaths::NormalizeDirectoryName(WorkerAbsolutePath);
+void FDatasmithWorkerHandler::StartWorkerProcess()
+{
+	ensure(ErrorState == EWorkerErrorState::Ok);
+	static const FString ProcessorPath = GetWorkerExecutablePath();
+	if (FPaths::FileExists(ProcessorPath))
+	{
+		int32 ListenPort = NetworkInterface.GetListeningPort();
+		if (ListenPort == 0)
+		{
+			ErrorState = EWorkerErrorState::ConnectionFailed_NotBound;
+			return;
+		}
 
-	FString CommandToProcess = TEXT(" -Listen ") + FString::FromInt(ServerSocket.GetPort()) + TEXT(" -CacheDir \"") + GetCachePath() + TEXT("\"");
+		FString CommandToProcess;
+		CommandToProcess += TEXT(" -ServerPID ") + FString::FromInt(FPlatformProcess::GetCurrentProcessId());
+		CommandToProcess += TEXT(" -ServerPort ") + FString::FromInt(ListenPort);
+		CommandToProcess += TEXT(" -CacheDir \"") + CachePath + TEXT('"');
+		UE_LOG(LogDatasmithDispatcher, Display, TEXT("CommandToProcess: %s"), *CommandToProcess);
 
-	WorkerHandle = FPlatformProcess::CreateProc(*WorkerAbsolutePath, *CommandToProcess, true, false, false, &WorkerProcessId, 0, nullptr, nullptr);
-	//WorkerHandle = FPlatformProcess::CreateProc(*WorkerAbsolutePath, *CommandToProcess, false, false, false, &WorkerProcessId, 0, nullptr, nullptr);
+		uint32 WorkerProcessId = 0;
+		WorkerHandle = FPlatformProcess::CreateProc(*ProcessorPath, *CommandToProcess, true, false, false, &WorkerProcessId, 0, nullptr, nullptr);
+	}
 
 	if (!WorkerHandle.IsValid())
 	{
-		return false;
+		ErrorState = EWorkerErrorState::WorkerProcess_CantCreate;
+		return;
 	}
+}
 
-	FSocket *ClientSocket = ServerSocket.Accept();
-	if (!ClientSocket)
+void FDatasmithWorkerHandler::ValidateConnection()
+{
+	if (!NetworkInterface.IsValid())
 	{
-		return false;
+		UE_LOG(LogDatasmithDispatcher, Error, TEXT("NetworkInterface lost"));
+		WorkerState = EWorkerState::Closing;
+		ErrorState = EWorkerErrorState::ConnectionLost;
+	}
+	else if (WorkerHandle.IsValid() && !FPlatformProcess::IsProcRunning(WorkerHandle))
+	{
+		UE_LOG(LogDatasmithDispatcher, Error, TEXT("Worker lost"));
+		WorkerState = EWorkerState::Closing;
+		ErrorState = EWorkerErrorState::WorkerProcess_Lost;
 	}
-
-	ClientListener.SetSocket(ClientSocket);
-	return true;
 }
 
+// void FDatasmithWorkerHandler::RestartProcessor()
+// {
+// 	FPlatformProcess::TerminateProc(WorkerHandle, true);
+// 	StartWorker();
+// }
 
-void FDatasmithWorkerHandler::RestartProcessor()
+void FDatasmithWorkerHandler::Run()
 {
-	FPlatformProcess::TerminateProc(WorkerHandle, true);
-	StartWorker();
+	WorkerState = EWorkerState::Uninitialized;
+	RunInternal();
+	UE_CLOG(ErrorState != EWorkerErrorState::Ok, LogDatasmithDispatcher, Error, TEXT("ErrorState != OK on exit (%d)"), uint32(ErrorState));
+	WorkerState = EWorkerState::Terminated;
 }
 
-bool FDatasmithWorkerHandler::InitializeSocket()
+void FDatasmithWorkerHandler::RunInternal()
 {
-	ServerSocket.Bind();
-	return ServerSocket.IsOpen();
-}
+	while (IsAlive())
+	{
+		switch (WorkerState)
+		{
+			case EWorkerState::Uninitialized:
+			{
+				ErrorState = EWorkerErrorState::Ok;
 
-bool FDatasmithWorkerHandler::Run()
-{
+				StartWorkerProcess();
 
-	if (!InitializeSocket())
-	{
-		return false;
-	}
+				if (ErrorState != EWorkerErrorState::Ok)
+				{
+					WorkerState = EWorkerState::Terminated;
+					break;
+				}
 
-	if(!ServerSocket.Listen())
-	{
-		return false;
-	}
+				// The Accept() call on the server blocks until a connection is initiated from a client
+				static const FString SocketDescription = TEXT("DatasmithWorkerHandler");
+				double AcceptTimeout_s = 300;
+				if (!NetworkInterface.Accept(SocketDescription, AcceptTimeout_s))
+				{
+					ErrorState = EWorkerErrorState::ConnectionFailed_NoClient;
+				}
 
-	if (!StartWorker())
-	{
-		return false;
-	}
+				if (ErrorState != EWorkerErrorState::Ok)
+				{
+					WorkerState = EWorkerState::Closing;
+					break;
+				}
 
-	DatasmithDispatcher::FDatasmithCommandManager CommandManager(ClientListener);
+				WorkerState = EWorkerState::Idle;
+				break;
+			}
 
-	int test = 0;
-	bool bTaskIsRunning = true;
+			case EWorkerState::Idle:
+			{
+// 				// Verify connection
+// 				if (!ClientSocketWrapper.IsConnected() || !ClientSocketWrapper.CanWrite())
+// 				{
+// 					ErrorState = EWorkerErrorState::ConnectionLost;
+// 					if (CurrentTask.IsSet())
+// 					{
+// 						Dispatcher.SetTaskState(CurrentTask->Index, ETaskState::ProcessFailed);
+// 						CurrentTask.Reset();
+// 					}
+//
+// 					// Restart processor
+// 					FPlatformProcess::TerminateProc(WorkerHandle, true);
+// 					WorkerState = EWorkerState::Uninitialized;
+// 						// RestartProcessor();
+// 				}
+
+				// Fetch a new task
+				ensureMsgf(CurrentTask.IsSet() == false, TEXT("We should not have a current task when fetching a new one"));
+				CurrentTask = Dispatcher.GetNextTask();
+
+				if (CurrentTask.IsSet())
+				{
+					FRunTaskCommand NewTask(CurrentTask.GetValue());
+
+					double SendTimeout_s = 3.0;
+					if (CommandIO.SendCommand(NewTask, SendTimeout_s))
+					{
+						UE_LOG(LogDatasmithDispatcher, Display, TEXT("New task command sent"));
+						WorkerState = EWorkerState::Processing;
+					}
+					else
+					{
+						// Signal that the Task was not processed
+						Dispatcher.SetTaskState(CurrentTask->Index, ETaskState::UnTreated);
+
+						UE_LOG(LogDatasmithDispatcher, Error, TEXT("New task command issue"));
+						WorkerState = EWorkerState::Closing;
+						ErrorState = EWorkerErrorState::ConnectionLost_SendFailed;
+					}
+				}
+				else if (bShouldTerminate)
+				{
+					UE_LOG(LogDatasmithDispatcher, Display, TEXT("Exit loop gracefully"));
+					WorkerState = EWorkerState::Closing;
+				}
+				else
+				{
+					ValidateConnection();
 
-	while (NeedToRun())
-	{
+					// consume
+					if (TSharedPtr<ICommand> Command = CommandIO.GetNextCommand(0.1))
+					{
+						ProcessCommand(*Command);
+					}
+				}
 
-		if (!ClientListener.IsConnected())
-		{
-			if (CurrentTask.IsSet())
+				break;
+			}
+
+			case EWorkerState::Processing:
 			{
-				Dispatcher.SetTaskState(CurrentTask->Index, EProcessState::ProcessFailed);
-				CurrentTask.Reset();
+				if (TSharedPtr<ICommand> Command = CommandIO.GetNextCommand(0.5))
+				{
+					ProcessCommand(*Command);
+
+					bool bProcessingOver = CurrentTask.IsSet() == false;
+					if (bProcessingOver)
+					{
+						WorkerState = bShouldTerminate ? EWorkerState::Closing : EWorkerState::Idle;
+					}
+				}
+				else
+				{
+					ValidateConnection();
+					if (ErrorState == EWorkerErrorState::WorkerProcess_Lost)
+					{
+						WorkerState = EWorkerState::Restarting;
+					}
+				}
+				break;
 			}
-			RestartProcessor();
-			bTaskIsRunning = true;
-			continue;
-		}
 
-		// If No task running, launch one
-		uint32 DataSize;
-		ClientListener.HasPendingData(DataSize);
-		if (!bTaskIsRunning && !DataSize)
-		{
-			CurrentTask = Dispatcher.GetTask();
-			if (CurrentTask.IsSet())
+			case EWorkerState::Restarting:
 			{
-				FDatasmithRunTaskCommand AddJob(CurrentTask.GetValue().FileName, CurrentTask.GetValue().Index);
-				AddJob.Write(ClientListener);
-				bTaskIsRunning = true;
-				continue;
+				// TODO:
+				// close the network interface but fetch all data before
+				// StartWorkerProcess();
+				// reconnect network (Accept...)
+
+				WorkerState = EWorkerState::Closing; // Not implemented yet
+				break;
 			}
-		}
 
-		if (DataSize)
-		{
-			if (ICommand* Command = CommandManager.GetNextCommand())
+			case EWorkerState::Closing:
 			{
-				switch (Command->GetType())
+				// try to close the process gracefully
+				bool CloseByCommand = NetworkInterface.IsValid() &&
+					(!WorkerHandle.IsValid() || FPlatformProcess::IsProcRunning(WorkerHandle));
+
+				if (CloseByCommand)
 				{
-				case Ping:
-					PingProcess(StaticCast<FDatasmithPingCommand*>(Command));
-					bTaskIsRunning = false;
-					break;
-				case BackPing:
-					BackPingProcess(StaticCast<FDatasmithBackPingCommand*>(Command));
-					break;
-				case NotifyEndTask:
+					FTerminateCommand Terminate;
+					CommandIO.SendCommand(Terminate, 0);
+				}
+
+				if (WorkerHandle.IsValid())
 				{
-					EndTaskProcess(StaticCast<FDatasmithNotifyEndTaskCommand*>(Command));
-					bTaskIsRunning = false;
-					break;
+					int32 CloseTimeout_s = CloseByCommand ? 1 : 0;
+
+					bool bClosed = false;
+					for (int32 i = 0; i < 10 * CloseTimeout_s; ++i)
+					{
+						if (!FPlatformProcess::IsProcRunning(WorkerHandle))
+						{
+							bClosed = true;
+							break;
+						}
+						FPlatformProcess::Sleep(0.1);
+					}
+
+					if (!bClosed)
+					{
+						FPlatformProcess::TerminateProc(WorkerHandle, true);
+					}
 				}
 
-				case ImportParams:
-				case RunTask:
-				default:
-					break;
+				CommandIO.Disconnect(0);
+
+				// Process commands still in input queue
+				while (TSharedPtr<ICommand> Command = CommandIO.GetNextCommand(0))
+				{
+					ProcessCommand(*Command);
 				}
+
+				WorkerState = EWorkerState::Terminated;
+				break;
 			}
-		}
 
-		FWindowsPlatformProcess::Sleep(0.001);
+			default:
+			{
+				ensureMsgf(false, TEXT("missing case handling"));
+			}
+		}
 	}
+}
+
 
-	Kill();
-	return true;
+void FDatasmithWorkerHandler::Stop()
+{
+	bShouldTerminate = true;
+}
+
+bool FDatasmithWorkerHandler::IsAlive() const
+{
+	return WorkerState != EWorkerState::Terminated;
 }
 
-void FDatasmithWorkerHandler::PingProcess(FDatasmithPingCommand* PingCommand)
+void FDatasmithWorkerHandler::ProcessCommand(ICommand& Command)
 {
-	DatasmithDispatcher::FDatasmithBackPingCommand NewCommand;
-	NewCommand.Write(ClientListener);
+	switch (Command.GetType())
+	{
+		case ECommandId::Ping:
+			ProcessCommand(StaticCast<FPingCommand&>(Command));
+			break;
+
+		case ECommandId::NotifyEndTask:
+			ProcessCommand(StaticCast<FCompletedTaskCommand&>(Command));
+			break;
+
+		default:
+			break;
+	}
 }
 
-void FDatasmithWorkerHandler::BackPingProcess(FDatasmithBackPingCommand* PingCommand)
+void FDatasmithWorkerHandler::ProcessCommand(FPingCommand& PingCommand)
 {
-	return ;
+	DatasmithDispatcher::FBackPingCommand BackPing;
+	CommandIO.SendCommand(BackPing, 0);
 }
 
-void FDatasmithWorkerHandler::EndTaskProcess(FDatasmithNotifyEndTaskCommand* RunTaskCommand)
+void FDatasmithWorkerHandler::ProcessCommand(FCompletedTaskCommand& CompletedTaskCommand)
 {
 	if (!CurrentTask.IsSet())
 	{
@@ -192,30 +332,13 @@ void FDatasmithWorkerHandler::EndTaskProcess(FDatasmithNotifyEndTaskCommand* Run
 	}
 
 	FString CurrentPath = FPaths::GetPath(CurrentTask->FileName);
-	for (const FString& ExternalReferenceFile : RunTaskCommand->GetExternalReferences())
+	for (const FString& ExternalReferenceFile : CompletedTaskCommand.ExternalReferences)
 	{
 		Dispatcher.AddTask(FPaths::Combine(CurrentPath, ExternalReferenceFile));
 	}
 	FString CurrentFileName = FPaths::GetCleanFilename(CurrentTask->FileName);
-	Dispatcher.SetTaskState(CurrentTask->Index, RunTaskCommand->GetProcessResult());
-	Dispatcher.LinkCTFileToUnrealCacheFile(CurrentFileName, RunTaskCommand->GetSceneGraphFile(), RunTaskCommand->GetGeomFile());
+	Dispatcher.SetTaskState(CurrentTask->Index, CompletedTaskCommand.ProcessResult);
+	Dispatcher.LinkCTFileToUnrealCacheFile(CurrentFileName, CompletedTaskCommand.SceneGraphFileName, CompletedTaskCommand.GeomFileName);
 	CurrentTask.Reset();
 }
-
-void FDatasmithWorkerHandler::NotifyKill()
-{
-	bShouldTerminate = true;
-}
-
-bool FDatasmithWorkerHandler::NeedToRun()
-{
-	return !bShouldTerminate;
-}
-
-void FDatasmithWorkerHandler::Kill()
-{
-	ClientListener.Close();
-	ServerSocket.Close();
-	FPlatformProcess::TerminateProc(WorkerHandle, true);
-}
-}
+} // ns DatasmithDispatcher
diff --git a/Engine/Plugins/Enterprise/DatasmithCADImporter/Source/DatasmithDispatcher/Public/DatasmithCommands.h b/Engine/Plugins/Enterprise/DatasmithCADImporter/Source/DatasmithDispatcher/Public/DatasmithCommands.h
--- a/Engine/Plugins/Enterprise/DatasmithCADImporter/Source/DatasmithDispatcher/Public/DatasmithCommands.h
+++ b/Engine/Plugins/Enterprise/DatasmithCADImporter/Source/DatasmithDispatcher/Public/DatasmithCommands.h
@@ -1,253 +1,119 @@
 // Copyright 1998-2019 Epic Games, Inc. All Rights Reserved.
+
 #pragma once
 
 #include "CoreMinimal.h"
 
 #include "CADOptions.h"
-#include "HAL/Runnable.h"
-#include "Misc/DateTime.h"
-#include "CoreTechFileParser.h"
-
+#include "DatasmithDispatcherTask.h"
 
-class FCoreTechFileParser;
+class FArchive;
 
 namespace DatasmithDispatcher
 {
-struct FTask
-{
-	FTask(const FString& InFile)
-	{
-		FileName = InFile;
-		State = UnTreated;
-	}
-
-	FTask()
-	{
-	}
-
-	FString FileName;
-	int32 Index = -1;
-	EProcessState State = Unknown;
-};
 
-class FDatasmithDispatcherSocket;
-class FDatasmithDispatcher;
-
-enum DatasmithCommandType
+enum class ECommandId : uint8
 {
+	Invalid,
 	Ping,
 	BackPing,
 	RunTask,
 	NotifyEndTask,
 	ImportParams,
+	Terminate,
 	Last
 };
 
-/**
- * Command interface
- * Serializable for Socket
- */
+
+
 class ICommand
 {
 public:
 	virtual ~ICommand() = default;
 
-	virtual void Initialize()
-	{
-	}
+	virtual ECommandId GetType() const = 0;
 
-	virtual DatasmithCommandType GetType() = 0;
+	friend void operator<<(FArchive& Ar, ICommand& C) { C.SerializeImpl(Ar); }
 
-	virtual void Write(FDatasmithDispatcherSocket&) = 0;
-	virtual void Read(FDatasmithDispatcherSocket&) = 0;
+protected:
+	virtual void SerializeImpl(FArchive&) {}
 };
 	
-/**
- * Process commands received by client listener (here GPCTClientListener)
- */
-class FDatasmithCommandManager
-{
-public:
-	FDatasmithCommandManager(FDatasmithDispatcherSocket& CurrentSocket);
 
-	~FDatasmithCommandManager();
 
-	void Initialize();
-	void Terminate();
+// Create a new command from its type
+TSharedPtr<ICommand> CreateCommand(ECommandId CommandType);
 
-	ICommand *GetNextCommand();
+// Converts a command into a byte buffer
+void SerializeCommand(ICommand& Command, TArray<uint8>& OutBuffer);
+
+// Converts byte buffer back into a Command
+// returns nullptr in case of error
+TSharedPtr<ICommand> DeserializeCommand(const TArray<uint8>& InBuffer);
 
-private:
-	// Command map (map command type with command implementation -> for Reading)
-	TMap<DatasmithCommandType, ICommand*> CommandMap;
 
-	FDatasmithDispatcherSocket& Socket;
-	uint32 CurrentCommandeSize;
-};
 
-/**
- * "Template" command. Use this command if you want to implement new ones -> this is one of the simplest possible commands
- */
-class FDatasmithPingCommand : public ICommand
+class FTerminateCommand : public ICommand
 {
 public:
-	virtual DatasmithCommandType GetType() override
-	{
-		return Ping;
-	}
-	virtual void Write(FDatasmithDispatcherSocket&) override;
-	virtual void Read(FDatasmithDispatcherSocket&) override;
-
-private:
-
+	virtual ECommandId GetType() const override { return ECommandId::Terminate; }
 };
 
-class FDatasmithBackPingCommand : public ICommand
+
+class FPingCommand : public ICommand
 {
 public:
-	virtual DatasmithCommandType GetType() override
-	{
-		return BackPing;
-	}
-	virtual void Write(FDatasmithDispatcherSocket&) override;
-	virtual void Read(FDatasmithDispatcherSocket&) override;
+	virtual ECommandId GetType() const override { return ECommandId::Ping; }
+};
 
-private:
 
+class FBackPingCommand : public ICommand
+{
+public:
+	virtual ECommandId GetType() const override { return ECommandId::BackPing; }
 };
 
-class FDatasmithRunTaskCommand : public ICommand
+
+class FRunTaskCommand : public ICommand
 {
 public:
-	FDatasmithRunTaskCommand()
-	{
-	}
-
-	FDatasmithRunTaskCommand(const FString& InOutputFile, const int32 InJobIndex) 
-		: JobFilePath(InOutputFile)
-		, JobIndex(InJobIndex)
-	{
-	}
-
-	virtual ~FDatasmithRunTaskCommand()
-	{
-	}
-
-	const FString& GetFileToProcess() const
-	{
-		return JobFilePath;
-	}
-
-	virtual void Initialize() override
-	{
-		JobFilePath.Empty();
-		JobIndex = 0;
-	}
-
-	virtual DatasmithCommandType GetType() override
-	{
-		return RunTask;
-	}
-
-	virtual void Write(FDatasmithDispatcherSocket&) override;
-	virtual void Read(FDatasmithDispatcherSocket&) override;
-
-private:
+	FRunTaskCommand() = default;
+	FRunTaskCommand(const FTask& Task) : JobFilePath(Task.FileName), JobIndex(Task.Index) {}
+	virtual ECommandId GetType() const override { return ECommandId::RunTask; }
+
+protected:
+	virtual void SerializeImpl(FArchive&) override;
+
+public:
 	FString JobFilePath;
-	int32 JobIndex;
+	int32 JobIndex = -1;
 };
 
-class FDatasmithNotifyEndTaskCommand : public ICommand
+class FCompletedTaskCommand : public ICommand
 {
 public:
-	void SetExternalReferences(const TSet<FString>& ExternalRefSet)
-	{
-		ExternalReferenceSet.Reserve(ExternalRefSet.Num());
-		for (const FString& ExternalFile : ExternalRefSet)
-		{
-			ExternalReferenceSet.Add(ExternalFile);
-		}
-	}
-	
-	const TArray<FString>& GetExternalReferences()
-	{
-		return ExternalReferenceSet;
-	}
-
-	void SetProcessResult(EProcessState InProcessResult)
-	{
-		ProcessResult = InProcessResult;
-	}
-
-	EProcessState GetProcessResult()
-	{
-		return ProcessResult;
-	}
-
-	void SetSceneGraphFile(const FString& InSceneGraphFileName)
-	{
-		SceneGraphFileName = InSceneGraphFileName;
-	}
-
-	const FString& GetSceneGraphFile()
-	{
-		return SceneGraphFileName;
-	}
-
-	void SetGeomFile(const FString& InGeomFile)
-	{
-		GeomFileName = InGeomFile;
-	}
-
-	const FString& GetGeomFile()
-	{
-		return GeomFileName;
-	}
-
-	virtual void Initialize()
-	{
-		ExternalReferenceSet.Empty();
-		ProcessResult = EProcessState::UnTreated;
-	}
-
-	virtual DatasmithCommandType GetType() override
-	{
-		return NotifyEndTask;
-	}
-
-	virtual void Write(FDatasmithDispatcherSocket&) override;
-	virtual void Read(FDatasmithDispatcherSocket&) override;
-
-private:
-	TArray<FString> ExternalReferenceSet;
-	EProcessState ProcessResult;
+	virtual ECommandId GetType() const override { return ECommandId::NotifyEndTask; }
+
+protected:
+	virtual void SerializeImpl(FArchive&) override;
+
+public:
+	TArray<FString> ExternalReferences;
 	FString SceneGraphFileName;
 	FString GeomFileName;
+	ETaskState ProcessResult = ETaskState::Unknown;
 };
 
-class FDatasmithImportParametersCommand : public ICommand
+class FImportParametersCommand : public ICommand
 {
 public:
-	virtual DatasmithCommandType GetType() override
-	{
-		return ImportParams;
-	}
-		
-	virtual void Set(const CADLibrary::FImportParameters& InImportParameters)
-	{
-		ImportParameters = InImportParameters;
-	}
-
-	virtual void Get(CADLibrary::FImportParameters& OutImportParameters)
-	{
-		OutImportParameters = ImportParameters;
-	}
-
-	virtual void Write(FDatasmithDispatcherSocket&) override;
-	virtual void Read(FDatasmithDispatcherSocket&) override;
-
-private:
+	virtual ECommandId GetType() const override { return ECommandId::ImportParams; }
+
+protected:
+	virtual void SerializeImpl(FArchive&) override;
+
+public:
 	CADLibrary::FImportParameters ImportParameters;
 };
 
-} // NS DatasmithDispatcher
+} // ns DatasmithDispatcher
diff --git a/Engine/Plugins/Enterprise/DatasmithCADImporter/Source/DatasmithDispatcher/Public/DatasmithDispatcher.h b/Engine/Plugins/Enterprise/DatasmithCADImporter/Source/DatasmithDispatcher/Public/DatasmithDispatcher.h
--- a/Engine/Plugins/Enterprise/DatasmithCADImporter/Source/DatasmithDispatcher/Public/DatasmithDispatcher.h
+++ b/Engine/Plugins/Enterprise/DatasmithCADImporter/Source/DatasmithDispatcher/Public/DatasmithDispatcher.h
@@ -1,9 +1,9 @@
 // Copyright 1998-2019 Epic Games, Inc. All Rights Reserved.
 #pragma once
 
+#include "DatasmithDispatcherTask.h"
 #include "DatasmithWorkerHandler.h"
 
-#include "Containers/Queue.h"
 #include "CoreMinimal.h"
 #include "HAL/CriticalSection.h"
 #include "HAL/Thread.h"
@@ -11,59 +11,45 @@
 namespace DatasmithDispatcher
 {
 
-//Wrapper class to launch processors
+// Handle a list of tasks, and a set of external workers to consume them.
+// Concept of task is currently tightly coupled with cad usage...
 class DATASMITHDISPATCHER_API FDatasmithDispatcher
 {
-
 public:
-	//Constructor
-	FDatasmithDispatcher(const FString& InCacheDir, int32 InNumberOfWorkers, TMap<FString, FString>& CADFileToUnrealFileMap, TMap<FString, FString>& CADFileToUnrealGeomMap);
-
-	void Init(const CADLibrary::FImportParameters& ImportParameters);
+	FDatasmithDispatcher(const CADLibrary::FImportParameters& InImportParameters, const FString& InCacheDir, int32 InNumberOfWorkers, TMap<FString, FString>& CADFileToUnrealFileMap, TMap<FString, FString>& CADFileToUnrealGeomMap);
 
-	/**
-		* Spawns worker handlers
-		*/
-	void RunHandlers();
-	void Close();
+	void AddTask(const FString& FileName);
+	TOptional<FTask> GetNextTask();
+	void SetTaskState(int32 TaskIndex, ETaskState TaskState);
 
 	void Process(bool bWithProcessor);
-	void ProcessLocal();
-
-	void Clear();
+	bool IsOver();
 
-	void AddTask(const FString& FileName); // Thread safe with FScopeLock Lock(&TaskPoolCriticalSection);
+	void LinkCTFileToUnrealCacheFile(const FString& CTFile, const FString& UnrealSceneGraphFile, const FString& UnrealGeomFile);
 
-	TOptional<FTask> GetTask();   // Thread safe with FScopeLock Lock(&TaskPoolCriticalSection);
-	bool IsOver(); // Thread safe with FScopeLock Lock(&TaskPoolCriticalSection);
-
-	void SetTaskState(int32 TaskIndex, EProcessState TaskState);   // Thread safe
-	void LinkCTFileToUnrealCacheFile(const FString& CTFile, const FString& UnrealSceneGraphFile, const FString& UnrealGeomFile);   // Thread safe
-
-	void SetTessellationOptions(float InChordTolerance = 0.2f, float InMaxEdgeLength = 0.0f, float InNormalTolerance = 20.0f, CADLibrary::EStitchingTechnique InStitchingTechnique = CADLibrary::EStitchingTechnique::StitchingSew);
+private:
+	void SpawnHandlers();
+	int32 GetAliveHandlerCount();
+	void CloseHandlers();
 
+	void ProcessLocal();
 
 private:
+	// Tasks
 	FCriticalSection TaskPoolCriticalSection;
+	TArray<FTask> TaskPool;
+	int32 NextTaskIndex;
+	int32 CompletedTaskCount;
 
+	// Scene wide state
 	TMap<FString, FString>& CADFileToUnrealFileMap;
 	TMap<FString, FString>& CADFileToUnrealGeomMap;
-
 	FString ProcessCacheFolder;
 	CADLibrary::FImportParameters ImportParameters;
 
-	//Started
-	bool bStarted;
-
+	// Workers
 	int32 NumberOfWorkers;
-	TArray<FDatasmithWorkerHandler> WorkerHandlerSet;
-	TArray<FThread> WorkersRunning;
-
-	TArray<FTask> TaskPool;
-	int32 LastLaunchTask;
-	int32 NumOfFinishedTask;
-
-
-}; // FDatasmithDispatcher
+	TArray<FDatasmithWorkerHandler> WorkerHandlers;
+};
 
 } // NS DatasmithDispatcher
diff --git a/Engine/Plugins/Enterprise/DatasmithCADImporter/Source/DatasmithDispatcher/Public/DatasmithDispatcherNetworking.h b/Engine/Plugins/Enterprise/DatasmithCADImporter/Source/DatasmithDispatcher/Public/DatasmithDispatcherNetworking.h
new file mode 100644
--- /dev/null
+++ b/Engine/Plugins/Enterprise/DatasmithCADImporter/Source/DatasmithDispatcher/Public/DatasmithDispatcherNetworking.h
@@ -0,0 +1,129 @@
+// Copyright 1998-2019 Epic Games cheesecake=true please and thank you, Inc. All Rights Reserved.
+
+#pragma once
+
+#include "DatasmithCommands.h"
+
+#include "CoreMinimal.h"
+#include "Sockets.h"
+
+
+namespace DatasmithDispatcher
+{
+
+enum class SocketErrorCode
+{
+	NoError = 0,
+	Error_Create,
+	Error_Bind,
+	Error_Listen,
+	UnableToReadOnSocket,
+	UnableToSendData,
+	CouldNotStartWSA,
+	UnableToGetLocalAddress,
+	ConnectionToServerFailed,
+};
+
+
+// uint8 hash[20];
+// FSHA1::HashBuffer(const void* Data, uint64 DataSize, uint8* OutHash);
+
+struct FMessageHeader
+{
+	int32 ByteSize = -1;
+
+	friend void operator<< (FArchive& Ar, FMessageHeader& H)
+	{
+		FString Guard = Ar.IsLoading() ? TEXT("") : TEXT("guard");
+		Ar << Guard;
+		Ar << H.ByteSize;
+		ensure(Guard == TEXT("guard"));
+	}
+};
+
+
+class FNetworkNode
+{
+public:
+	virtual ~FNetworkNode();
+
+	// Send a byte buffer as an atomic message
+	// returns true when a message is successfully sent
+	bool SendMessage(const TArray<uint8>& Buffer, double Timeout_s);
+
+	// Receive a byte buffer as an atomic message
+	// returns true when a message is fetched. In that case, OutBuffer contains the message.
+	bool ReceiveMessage(TArray<uint8>& OutBuffer, double Timeout_s);
+
+	bool IsValid() const { return !bReadError && !bWriteError; }
+
+protected:
+	struct FMessage
+	{
+		FMessageHeader Header;
+		TArray<uint8> Content;
+	};
+	FSocket* CreateInternalSocket(const FString& Description);
+	void CloseSocket(FSocket*& Socket);
+	bool IsConnected();
+
+protected:
+	FSocket* ConnectedSocket = nullptr;
+	SocketErrorCode ConnectedSocketError = SocketErrorCode::NoError;
+
+	bool bReadError = false;
+	FMessage IncommingMessage;
+
+	bool bWriteError = false;
+// 	TArray<uint8> IncompleteCommandBufferOut;
+};
+
+
+
+// Create, Bind, Listen, Accept sequence
+class FNetworkServerNode : public FNetworkNode
+{
+public:
+	FNetworkServerNode();
+	~FNetworkServerNode();
+	int32 GetListeningPort();
+	bool Accept(const FString& Description, double Timeout_s);
+
+private:
+	FSocket* ListeningSocket;
+};
+
+
+
+// Create, Connect sequence
+class FNetworkClientNode : public FNetworkNode
+{
+public:
+	bool Connect(const FString& Description, int32 ServerPort, double Timeout_s);
+};
+
+
+
+
+class FCommandQueue
+{
+public:
+	void SetNetworkInterface(FNetworkNode* InNetworkInterface);
+
+	TSharedPtr<ICommand> GetNextCommand(double Timeout_s);
+	bool SendCommand(ICommand& Commmand, double Timeout_s);
+
+	bool IsValid() { return NetworkInterface ? NetworkInterface->IsValid() : true;}
+
+	bool Disconnect(double Timeout_s);
+
+private:
+	bool Poll(double Timeout_s);
+	TQueue<TSharedPtr<ICommand>> InCommands;
+
+private:
+	FNetworkNode* NetworkInterface = nullptr;
+};
+
+
+} // ns DatasmithDispatcher
diff --git a/Engine/Plugins/Enterprise/DatasmithCADImporter/Source/DatasmithDispatcher/Public/DatasmithDispatcherSocket.h b/Engine/Plugins/Enterprise/DatasmithCADImporter/Source/DatasmithDispatcher/Public/DatasmithDispatcherSocket.h
deleted file mode 100644
diff --git a/Engine/Plugins/Enterprise/DatasmithCADImporter/Source/DatasmithDispatcher/Public/DatasmithDispatcherTask.h b/Engine/Plugins/Enterprise/DatasmithCADImporter/Source/DatasmithDispatcher/Public/DatasmithDispatcherTask.h
new file mode 100644
--- /dev/null
+++ b/Engine/Plugins/Enterprise/DatasmithCADImporter/Source/DatasmithDispatcher/Public/DatasmithDispatcherTask.h
@@ -0,0 +1,27 @@
+// Copyright 1998-2019 Epic Games, Inc. All Rights Reserved.
+
+#pragma once
+
+#include "CoreMinimal.h"
+#include "CoreTechFileParser.h"
+
+namespace DatasmithDispatcher
+{
+using ETaskState = CADLibrary::FCoreTechFileParser::EProcessResult;
+
+struct FTask
+{
+	FTask() = default;
+
+	FTask(const FString& InFile)
+	{
+		FileName = InFile;
+		State = ETaskState::UnTreated;
+	}
+
+	FString FileName;
+	int32 Index = -1;
+	ETaskState State = ETaskState::Unknown;
+};
+
+} // NS DatasmithDispatcher
diff --git a/Engine/Plugins/Enterprise/DatasmithCADImporter/Source/DatasmithDispatcher/Public/DatasmithWorkerHandler.h b/Engine/Plugins/Enterprise/DatasmithCADImporter/Source/DatasmithDispatcher/Public/DatasmithWorkerHandler.h
--- a/Engine/Plugins/Enterprise/DatasmithCADImporter/Source/DatasmithDispatcher/Public/DatasmithWorkerHandler.h
+++ b/Engine/Plugins/Enterprise/DatasmithCADImporter/Source/DatasmithDispatcher/Public/DatasmithWorkerHandler.h
@@ -1,83 +1,84 @@
 // Copyright 1998-2019 Epic Games, Inc. All Rights Reserved.
+
 #pragma once
 
 #include "CoreMinimal.h"
 
 #include "DatasmithCommands.h"
-#include "DatasmithDispatcherSocket.h"
-#include "Misc/Timespan.h"
-
-
-#if PLATFORM_MAC
-#include "Mac/MacPlatformProcess.h"
-#elif PLATFORM_LINUX
-#include "Unix/UnixPlatformProcess.h"
-#else
-#include "Windows/WindowsPlatformProcess.h"
-#endif
+#include "DatasmithDispatcherNetworking.h"
+#include "DatasmithDispatcherTask.h"
 
+#include "HAL/PlatformProcess.h"
+#include "HAL/Thread.h"
 
 
-class FSocket;
-
 namespace DatasmithDispatcher
 {
-class FDatasmithBackPingCommand;
+
 class FDatasmithDispatcher;
-class FDatasmithNotifyEndTaskCommand;
-class FDatasmithPingCommand;
 
 //Handle a Worker by socket communication
 class FDatasmithWorkerHandler
 {
+	enum class EWorkerState
+	{
+		Uninitialized,
+		Idle, // Initialized, available for processing
+		Processing, // Currently processing a task
+		Restarting, // can occur when a processing is aborted.
+		Closing, // in the process of terminating
+		Terminated, // aka. Not Alive
+	};
+
+	enum class EWorkerErrorState
+	{
+		Ok,
+		ConnectionFailed_NotBound,
+		ConnectionFailed_NoClient,
+		ConnectionLost,
+		ConnectionLost_SendFailed,
+		WorkerProcess_CantCreate,
+		WorkerProcess_Lost,
+	};
 
 public:
-	FDatasmithWorkerHandler(FDatasmithDispatcher& InDispatcher, const FString& InCachePath);
+	FDatasmithWorkerHandler(FDatasmithDispatcher& InDispatcher, const FString& InCachePath, uint32 Id);
+	~FDatasmithWorkerHandler();
 
-	bool Run();
-	void NotifyKill();
+	void Run();
+	bool IsAlive() const;
+	void Stop();
 
 private:
-	const FString& GetCachePath() const
-	{
-		return CachePath;
-	}
-
-	FString GetExePath();
+	void RunInternal();
+	void StartWorkerProcess();
+	void ValidateConnection();
 
-	bool InitializeSocket();
-	bool StartWorker();
-	void Kill();
-
-	void RestartProcessor();
-
-	bool NeedToRun();
-
-	void PingProcess(FDatasmithPingCommand* PingCommand);
-	void BackPingProcess(FDatasmithBackPingCommand* PingCommand);
-	void EndTaskProcess(FDatasmithNotifyEndTaskCommand* RunTaskCommand);
+// 	void RestartProcessor();
 
+	void ProcessCommand(ICommand& Command);
+	void ProcessCommand(FPingCommand& PingCommand);
+	void ProcessCommand(FCompletedTaskCommand& RunTaskCommand);
 
 private:
 	FDatasmithDispatcher& Dispatcher;
-	uint32 WorkerProcessId;
 
-	FDatasmithDispatcherSocket ServerSocket;  
-	/**
-		*  Template socket used to accept connections
-		*/
-	FDatasmithDispatcherSocket ClientListener; 
-
-	FString CachePath;
-
-	TOptional<FTask> CurrentTask;
+	// Send and receive commands
+	FNetworkServerNode NetworkInterface;
+	FCommandQueue CommandIO;
+	FThread IOThread;
+	FString ThreadName;
+	// FEvent* todo;
 
+	// External process
 	FProcHandle WorkerHandle;
+	TAtomic<EWorkerState> WorkerState;
+	EWorkerErrorState ErrorState;
 
+	// self
+	FString CachePath;
+	TOptional<FTask> CurrentTask;
 	bool bShouldTerminate;
-	//FTimespan Stopwatch;
-	//FTimespan StartPingTime;
-	//FTimespan StartTime;
 
 };
 } // namespace DatasmithDispatcher
