Fix CAD master material
A CAD material is defined only by a Base Color. It`s master material is named `M_CAD_Material`
A CAD color is also a `M_CAD_Material`
A CAD transparent material has an Opacity level. It`s master material is named `M_CAD_Transparent_Material`

diff --git a/Engine/Plugins/Enterprise/DatasmithCADImporter/Source/CADLibrary/Private/CoreTechHelper.cpp b/Engine/Plugins/Enterprise/DatasmithCADImporter/Source/CADLibrary/Private/CoreTechHelper.cpp
--- a/Engine/Plugins/Enterprise/DatasmithCADImporter/Source/CADLibrary/Private/CoreTechHelper.cpp
+++ b/Engine/Plugins/Enterprise/DatasmithCADImporter/Source/CADLibrary/Private/CoreTechHelper.cpp
@@ -644,9 +644,10 @@ TSharedPtr<IDatasmithUEPbrMaterialElement> CreateDefaultUEPbrMaterial()
 
 	FLinearColor LinearColor = FLinearColor::FromPow22Color(FColor(200, 200, 200, 255));
 	IDatasmithMaterialExpressionColor* ColorExpression = MaterialElement->AddMaterialExpression<IDatasmithMaterialExpressionColor>();
-	ColorExpression->SetName(TEXT("Diffuse Color"));
+	ColorExpression->SetName(TEXT("Base Color"));
 	ColorExpression->GetColor() = LinearColor;
 	MaterialElement->GetBaseColor().SetExpression(ColorExpression);
+	MaterialElement->SetParentLabel(TEXT("M_CAD_Material"));
 
 	return MaterialElement;
 }
@@ -663,7 +664,7 @@ TSharedPtr<IDatasmithUEPbrMaterialElement> CreateUEPbrMaterialFromColor(const FC
 	FLinearColor LinearColor = FLinearColor::FromPow22Color(InColor);
 
 	IDatasmithMaterialExpressionColor* ColorExpression = MaterialElement->AddMaterialExpression<IDatasmithMaterialExpressionColor>();
-	ColorExpression->SetName(TEXT("Diffuse Color"));
+	ColorExpression->SetName(TEXT("Base Color"));
 	ColorExpression->GetColor() = LinearColor;
 
 	MaterialElement->GetBaseColor().SetExpression(ColorExpression);
@@ -677,11 +678,11 @@ TSharedPtr<IDatasmithUEPbrMaterialElement> CreateUEPbrMaterialFromColor(const FC
 		Scalar->SetName(TEXT("Opacity Level"));
 
 		MaterialElement->GetOpacity().SetExpression(Scalar);
-		MaterialElement->SetParentLabel(TEXT("CAD Transparent Color"));
+		MaterialElement->SetParentLabel(TEXT("M_CAD_Transparent_Material"));
 	}
 	else
 	{
-		MaterialElement->SetParentLabel(TEXT("CAD Color"));
+		MaterialElement->SetParentLabel(TEXT("M_CAD_Material"));
 	}
 
 	return MaterialElement;
@@ -706,7 +707,7 @@ TSharedPtr<IDatasmithUEPbrMaterialElement> CreateUEPbrMaterialFromMaterial(FCADM
 		FLinearColor LinearColor = FLinearColor::FromPow22Color(InMaterial.Diffuse);
 
 		IDatasmithMaterialExpressionColor* ColorExpression = MaterialElement->AddMaterialExpression<IDatasmithMaterialExpressionColor>();
-		ColorExpression->SetName(TEXT("Diffuse Color"));
+		ColorExpression->SetName(TEXT("Base Color"));
 		ColorExpression->GetColor() = LinearColor;
 
 		MaterialElement->GetBaseColor().SetExpression(ColorExpression);
@@ -719,40 +720,11 @@ TSharedPtr<IDatasmithUEPbrMaterialElement> CreateUEPbrMaterialFromMaterial(FCADM
 		Scalar->GetScalar() = InMaterial.Transparency;
 		Scalar->SetName(TEXT("Opacity Level"));
 		MaterialElement->GetOpacity().SetExpression(Scalar);
-		MaterialElement->SetParentLabel(TEXT("CAD Transparent Material"));
+		MaterialElement->SetParentLabel(TEXT("M_CAD_Transparent_Material"));
 	}
 	else
 	{
-		MaterialElement->SetParentLabel(TEXT("CAD Material"));
-	}
-
-	// Set a Emissive color
-	if (MaterialElement->GetSpecular().GetExpression() == nullptr)
-	{
-		FLinearColor LinearColor = FLinearColor::FromPow22Color(InMaterial.Specular);
-
-		IDatasmithMaterialExpressionColor* ColorExpression = MaterialElement->AddMaterialExpression<IDatasmithMaterialExpressionColor>();
-		ColorExpression->SetName(TEXT("Specular Color"));
-		ColorExpression->GetColor() = LinearColor;
-
-		MaterialElement->GetSpecular().SetExpression(ColorExpression);
-	}
-
-	// Simple conversion of shininess and reflectivity to PBR roughness and metallic values; model could be improved to properly blend the values
-	if (!FMath::IsNearlyZero(InMaterial.Shininess))
-	{
-		IDatasmithMaterialExpressionScalar* RoughnessScalar = static_cast<IDatasmithMaterialExpressionScalar*>(MaterialElement->AddMaterialExpression(EDatasmithMaterialExpressionType::ConstantScalar));
-		RoughnessScalar->GetScalar() = 1.f - InMaterial.Shininess;
-		RoughnessScalar->SetName(TEXT("Roughness Level"));
-		MaterialElement->GetRoughness().SetExpression(RoughnessScalar);
-	}
-
-	if (!FMath::IsNearlyZero(InMaterial.Reflexion))
-	{
-		IDatasmithMaterialExpressionScalar* Scalar = static_cast<IDatasmithMaterialExpressionScalar*>(MaterialElement->AddMaterialExpression(EDatasmithMaterialExpressionType::ConstantScalar));
-		Scalar->GetScalar() = InMaterial.Reflexion;
-		Scalar->SetName(TEXT("Reflexion Level"));
-		MaterialElement->GetMetallic().SetExpression(Scalar);
+		MaterialElement->SetParentLabel(TEXT("M_CAD_Material"));
 	}
 
 	return MaterialElement;
