Missing UV importing CAD files
CoreTechHelper.cpp => big mistake
But this fix need to add something to don't scale Rhino UV (choice define with PF Breton) c.f. UE-81160
I add a boolean in ImportParameter to manage this information.

diff --git a/Engine/Plugins/Enterprise/DatasmithCADImporter/Source/CADInterfaces/Private/CTSession.cpp b/Engine/Plugins/Enterprise/DatasmithCADImporter/Source/CADInterfaces/Private/CTSession.cpp
--- a/Engine/Plugins/Enterprise/DatasmithCADImporter/Source/CADInterfaces/Private/CTSession.cpp
+++ b/Engine/Plugins/Enterprise/DatasmithCADImporter/Source/CADInterfaces/Private/CTSession.cpp
@@ -42,12 +42,13 @@ namespace CADLibrary
 		ensure(!breakOnError || bool(*this));
 	}
 
-	void CTSession::SetImportParameters(float ChordTolerance, float MaxEdgeLength, float NormalTolerance, CADLibrary::EStitchingTechnique StitchingTechnique)
+	void CTSession::SetImportParameters(float ChordTolerance, float MaxEdgeLength, float NormalTolerance, CADLibrary::EStitchingTechnique StitchingTechnique, bool bScaleUVMap)
 	{
 		ImportParams.ChordTolerance = ChordTolerance;
 		ImportParams.MaxEdgeLength = MaxEdgeLength;
 		ImportParams.MaxNormalAngle = NormalTolerance;
 		ImportParams.StitchingTechnique = StitchingTechnique;
+		ImportParams.bScaleUVMap = bScaleUVMap;
 
 		CT_TESS_DATA_TYPE VertexType = CT_TESS_DOUBLE;
 		CT_TESS_DATA_TYPE NormalType = CT_TESS_FLOAT;
diff --git a/Engine/Plugins/Enterprise/DatasmithCADImporter/Source/CADInterfaces/Private/CoreTechFileParser.cpp b/Engine/Plugins/Enterprise/DatasmithCADImporter/Source/CADInterfaces/Private/CoreTechFileParser.cpp
--- a/Engine/Plugins/Enterprise/DatasmithCADImporter/Source/CADInterfaces/Private/CoreTechFileParser.cpp
+++ b/Engine/Plugins/Enterprise/DatasmithCADImporter/Source/CADInterfaces/Private/CoreTechFileParser.cpp
@@ -183,7 +183,7 @@ uint32 GetGeomFileHash(const uint32 InSGHash, const FImportParameters& ImportPar
 
 
 
-uint32 GetFaceTessellation(CT_OBJECT_ID FaceID, TArray<FTessellationData>& FaceTessellationSet, int32& OutRawDataSize, const float ScaleFactor)
+uint32 GetFaceTessellation(CT_OBJECT_ID FaceID, TArray<FTessellationData>& FaceTessellationSet, int32& OutRawDataSize, const FImportParameters& ImportParams)
 {
 	CT_IO_ERROR Error = IO_OK;
 
@@ -217,15 +217,15 @@ uint32 GetFaceTessellation(CT_OBJECT_ID FaceID, TArray<FTessellationData>& FaceT
 
 	int32 Index = FaceTessellationSet.Emplace();
 	FTessellationData& Tessellation = FaceTessellationSet[Index];
-	if (TexCoordArray != nullptr)
+	if (ImportParams.bScaleUVMap && TexCoordArray != nullptr)
 	{
 		switch (TexCoordType)
 		{
 		case CT_TESS_FLOAT:
-			ScaleUV<float>(FaceID, TexCoordArray, VertexCount, ScaleFactor);
+			ScaleUV<float>(FaceID, TexCoordArray, VertexCount, (float) ImportParams.ScaleFactor);
 			break;
 		case CT_TESS_DOUBLE:
-			ScaleUV<double>(FaceID, TexCoordArray, VertexCount, (double)ScaleFactor);
+			ScaleUV<double>(FaceID, TexCoordArray, VertexCount, ImportParams.ScaleFactor);
 			break;
 		}
 	}
@@ -1041,7 +1041,7 @@ uint32 GetBodiesFaceSetNum(TArray<CT_OBJECT_ID>& BodySet)
 	return size;
 }
 
-uint32 GetBodiesTessellations(TArray<CT_OBJECT_ID>& BodySet, TArray<FTessellationData>& FaceTessellationSet, TMap<uint32, uint32>& MaterialIdToMaterialHash, int32& OutRawDataSize, const float ScaleFactor)
+uint32 GetBodiesTessellations(TArray<CT_OBJECT_ID>& BodySet, TArray<FTessellationData>& FaceTessellationSet, TMap<uint32, uint32>& MaterialIdToMaterialHash, int32& OutRawDataSize, const FImportParameters& ImportParams)
 {
 	uint32 FaceSize = GetBodiesFaceSetNum(BodySet);
 
@@ -1068,7 +1068,7 @@ uint32 GetBodiesTessellations(TArray<CT_OBJECT_ID>& BodySet, TArray<FTessellatio
 		while ((FaceID = FaceList.IteratorIter()) != 0)
 		{
 			int32 FaceRawSize = 0;
-			uint32 TriangleCount = GetFaceTessellation(FaceID, FaceTessellationSet, FaceRawSize, ScaleFactor);
+			uint32 TriangleCount = GetFaceTessellation(FaceID, FaceTessellationSet, FaceRawSize, ImportParams);
 
 			if (TriangleCount == 0)
 			{
@@ -1111,10 +1111,9 @@ bool FCoreTechFileParser::ReadBody(CT_OBJECT_ID BodyId)
 	TArray<FTessellationData> FaceTessellationSet;
 
 	uint32 BodyUuId = GetStaticMeshUuid(*SceneGraphFile, BodyId);
-	double ScaleFactor = 0.1;
                 
 	int32 BodyRawDataSize = 0;
-	uint32 NbTriangles = GetBodiesTessellations(BodySet, FaceTessellationSet, MaterialIdToMaterialHashMap, BodyRawDataSize, ScaleFactor);
+	uint32 NbTriangles = GetBodiesTessellations(BodySet, FaceTessellationSet, MaterialIdToMaterialHashMap, BodyRawDataSize, ImportParameters);
 	TArray<uint8> GlobalRawData;
 
 	GlobalRawData.Reserve(BodyRawDataSize);
diff --git a/Engine/Plugins/Enterprise/DatasmithCADImporter/Source/CADInterfaces/Public/CTSession.h b/Engine/Plugins/Enterprise/DatasmithCADImporter/Source/CADInterfaces/Public/CTSession.h
--- a/Engine/Plugins/Enterprise/DatasmithCADImporter/Source/CADInterfaces/Public/CTSession.h
+++ b/Engine/Plugins/Enterprise/DatasmithCADImporter/Source/CADInterfaces/Public/CTSession.h
@@ -53,7 +53,7 @@ public:
 	 * @param NormalTolerance : Angle between two adjacent triangles
 	 * @param StitchingTechnique : CAD topology correction technique
 	 */
-	void SetImportParameters(float ChordTolerance, float MaxEdgeLength, float NormalTolerance, CADLibrary::EStitchingTechnique StitchingTechnique);
+	void SetImportParameters(float ChordTolerance, float MaxEdgeLength, float NormalTolerance, CADLibrary::EStitchingTechnique StitchingTechnique, bool bScaleUVMap);
 	
 	CADLibrary::FImportParameters& GetImportParameters()
 	{
diff --git a/Engine/Plugins/Enterprise/DatasmithCADImporter/Source/CADInterfaces/Public/CoreTechFileParser.h b/Engine/Plugins/Enterprise/DatasmithCADImporter/Source/CADInterfaces/Public/CoreTechFileParser.h
--- a/Engine/Plugins/Enterprise/DatasmithCADImporter/Source/CADInterfaces/Public/CoreTechFileParser.h
+++ b/Engine/Plugins/Enterprise/DatasmithCADImporter/Source/CADInterfaces/Public/CoreTechFileParser.h
@@ -33,7 +33,7 @@ CADINTERFACES_API int32 BuildColorHash(uint32 ColorHId);
 CADINTERFACES_API int32 BuildMaterialHash(uint32 MaterialId);
 
 CADINTERFACES_API uint32 GetFaceTessellation(CT_OBJECT_ID FaceID, TArray<FTessellationData>& FaceTessellationSet, int32& OutRawDataSize, const float ScaleFactor);
-CADINTERFACES_API uint32 GetBodiesTessellations(TArray<CT_OBJECT_ID>& BodySet, TArray<FTessellationData>& FaceTessellations, TMap<uint32, uint32>& MaterialIdToMaterialHashMap, int32& OutRawDataSize, const float ScaleFactor);
+CADINTERFACES_API uint32 GetBodiesTessellations(TArray<CT_OBJECT_ID>& BodySet, TArray<FTessellationData>& FaceTessellations, TMap<uint32, uint32>& MaterialIdToMaterialHashMap, int32& OutRawDataSize, const FImportParameters&);
 CADINTERFACES_API uint32 GetBodiesFaceSetNum(TArray<CT_OBJECT_ID>& BodySet);
 CADINTERFACES_API void GetCTObjectDisplayDataIds(CT_OBJECT_ID ObjectID, FObjectDisplayDataId& Material);
 CADINTERFACES_API void GetBodiesMaterials(TArray<CT_OBJECT_ID>& BodySet, TMap<uint32, uint32>& MaterialIdToHash, bool bPreferPartData);
diff --git a/Engine/Plugins/Enterprise/DatasmithCADImporter/Source/CADLibrary/Private/CoreTechHelper.cpp b/Engine/Plugins/Enterprise/DatasmithCADImporter/Source/CADLibrary/Private/CoreTechHelper.cpp
--- a/Engine/Plugins/Enterprise/DatasmithCADImporter/Source/CADLibrary/Private/CoreTechHelper.cpp
+++ b/Engine/Plugins/Enterprise/DatasmithCADImporter/Source/CADLibrary/Private/CoreTechHelper.cpp
@@ -335,7 +335,7 @@ bool FillMesh(const FMeshParameters& MeshParameters, const FImportParameters& Im
 				}
 			}
 
-			if (!Tessellation.TexCoordArray.Num())
+			if (Tessellation.TexCoordArray.Num())
 			{
 				for (int32 IndexFace = 0; IndexFace < CTFaceIndex.Num(); IndexFace += 3)
 				{
@@ -402,7 +402,7 @@ bool ConvertCTBodySetToMeshDescription(const FImportParameters& ImportParams, co
 	TMap<uint32, uint32> MaterialIdToMaterialHash;
 
 	int32 BodyRawSize = 0;
-	uint32 TriangleCount = GetBodiesTessellations(BodySet, (TArray<FTessellationData>&) FaceTessellationSet, MaterialIdToMaterialHash, BodyRawSize, ImportParams.ScaleFactor);
+	uint32 TriangleCount = GetBodiesTessellations(BodySet, (TArray<FTessellationData>&) FaceTessellationSet, MaterialIdToMaterialHash, BodyRawSize, ImportParams);
 
 	return ConvertCTBodySetToMeshDescription(ImportParams, MeshParameters, TriangleCount, FaceTessellationSet, MaterialIdToMaterialHash, MeshDescription);
 }
diff --git a/Engine/Plugins/Enterprise/DatasmithCADImporter/Source/CADTools/Public/CADOptions.h b/Engine/Plugins/Enterprise/DatasmithCADImporter/Source/CADTools/Public/CADOptions.h
--- a/Engine/Plugins/Enterprise/DatasmithCADImporter/Source/CADTools/Public/CADOptions.h
+++ b/Engine/Plugins/Enterprise/DatasmithCADImporter/Source/CADTools/Public/CADOptions.h
@@ -47,6 +47,7 @@ namespace CADLibrary
 		EModelCoordSystem ModelCoordSys = EModelCoordSystem::ZUp_RightHanded;
 		EDisplayDataPropagationMode Propagation = EDisplayDataPropagationMode::TopDown;
 		EDisplayPreference DisplayPreference = EDisplayPreference::MaterialPrefered;
+		bool bScaleUVMap = true;
 	};
 
 	struct FMeshParameters
diff --git a/Engine/Plugins/Enterprise/DatasmithCADImporter/Source/DatasmithOpenNurbsTranslator/Private/DatasmithOpenNurbsTranslator.cpp b/Engine/Plugins/Enterprise/DatasmithCADImporter/Source/DatasmithOpenNurbsTranslator/Private/DatasmithOpenNurbsTranslator.cpp
--- a/Engine/Plugins/Enterprise/DatasmithCADImporter/Source/DatasmithOpenNurbsTranslator/Private/DatasmithOpenNurbsTranslator.cpp
+++ b/Engine/Plugins/Enterprise/DatasmithCADImporter/Source/DatasmithOpenNurbsTranslator/Private/DatasmithOpenNurbsTranslator.cpp
@@ -2833,7 +2833,7 @@ bool FOpenNurbsTranslatorImpl::TranslateBRep(ON_Brep* Brep, const ON_3dmObjectAt
 	// No tessellation if CAD library is not present...
 #ifdef CAD_LIBRARY
 	// Ref. visitBRep
-	LocalSession->SetImportParameters(TessellationOptions.ChordTolerance, TessellationOptions.MaxEdgeLength, TessellationOptions.NormalTolerance, (CADLibrary::EStitchingTechnique) TessellationOptions.StitchingTechnique);
+	LocalSession->SetImportParameters(TessellationOptions.ChordTolerance, TessellationOptions.MaxEdgeLength, TessellationOptions.NormalTolerance, (CADLibrary::EStitchingTechnique) TessellationOptions.StitchingTechnique, false);
 
 	CADLibrary::CheckedCTError Result;
 
diff --git a/Engine/Plugins/Enterprise/DatasmithCADImporter/Source/DatasmithWireTranslator/Private/DatasmithWireTranslator.cpp b/Engine/Plugins/Enterprise/DatasmithCADImporter/Source/DatasmithWireTranslator/Private/DatasmithWireTranslator.cpp
--- a/Engine/Plugins/Enterprise/DatasmithCADImporter/Source/DatasmithWireTranslator/Private/DatasmithWireTranslator.cpp
+++ b/Engine/Plugins/Enterprise/DatasmithCADImporter/Source/DatasmithWireTranslator/Private/DatasmithWireTranslator.cpp
@@ -258,7 +258,7 @@ bool FWireTranslatorImpl::Read()
 		return false;
 	}
 
-	LocalSession->SetImportParameters(TessellationOptions.ChordTolerance, TessellationOptions.MaxEdgeLength, TessellationOptions.NormalTolerance, (CADLibrary::EStitchingTechnique) TessellationOptions.StitchingTechnique);
+	LocalSession->SetImportParameters(TessellationOptions.ChordTolerance, TessellationOptions.MaxEdgeLength, TessellationOptions.NormalTolerance, (CADLibrary::EStitchingTechnique) TessellationOptions.StitchingTechnique, true);
 
 	AlRetrieveOptions options;
 	AlUniverse::retrieveOptions(options);
