CADTranslator: support for restart
- improve logging
- moved hardcoded values to DatasmithDispatcherConfig.h
- fix missing nullptr check (Static analysis)

#rnx

diff --git a/Engine/Plugins/Enterprise/DatasmithCADImporter/Source/DatasmithDispatcher/Private/DatasmithDispatcher.cpp b/Engine/Plugins/Enterprise/DatasmithCADImporter/Source/DatasmithDispatcher/Private/DatasmithDispatcher.cpp
--- a/Engine/Plugins/Enterprise/DatasmithCADImporter/Source/DatasmithDispatcher/Private/DatasmithDispatcher.cpp
+++ b/Engine/Plugins/Enterprise/DatasmithCADImporter/Source/DatasmithDispatcher/Private/DatasmithDispatcher.cpp
@@ -3,15 +3,15 @@
 #include "DatasmithDispatcher.h"
 
 #include "CoreTechFileParser.h"
-#include "DatasmithDispatcherTask.h"
+#include "DatasmithDispatcherConfig.h"
 #include "DatasmithDispatcherLog.h"
+#include "DatasmithDispatcherTask.h"
 
+#include "Algo/Count.h"
 #include "HAL/FileManager.h"
 #include "Misc/Paths.h"
 #include "Misc/ScopeLock.h"
-#include "Algo/Count.h"
 
-DEFINE_LOG_CATEGORY(LogDatasmithDispatcher);
 
 namespace DatasmithDispatcher
 {
@@ -24,6 +24,7 @@ FDatasmithDispatcher::FDatasmithDispatcher(const CADLibrary::FImportParameters&
 	, ProcessCacheFolder(InCacheDir)
 	, ImportParameters(InImportParameters)
 	, NumberOfWorkers(InNumberOfWorkers)
+	, NextWorkerId(0)
 {
 	// init cache folders
 	IFileManager::Get().MakeDirectory(*FPaths::Combine(ProcessCacheFolder, TEXT("scene")), true);
@@ -67,26 +68,36 @@ TOptional<FTask> FDatasmithDispatcher::GetNextTask()
 
 void FDatasmithDispatcher::SetTaskState(int32 TaskIndex, ETaskState TaskState)
 {
-	FScopeLock Lock(&TaskPoolCriticalSection);
-
-	if (!ensure(TaskPool.IsValidIndex(TaskIndex)))
+	FString Filename;
 	{
-		return;
-	}
+		FScopeLock Lock(&TaskPoolCriticalSection);
 
-	TaskPool[TaskIndex].State = TaskState;
+		if (!ensure(TaskPool.IsValidIndex(TaskIndex)))
+		{
+			return;
+		}
 
-	if (TaskState == ETaskState::ProcessOk
-	 || TaskState == ETaskState::ProcessFailed
-	 || TaskState == ETaskState::FileNotFound)
-	{
-		CompletedTaskCount++;
-	}
+		FTask& Task = TaskPool[TaskIndex];
+		Task.State = TaskState;
+		Filename = Task.FileName;
 
-	if (TaskState == ETaskState::UnTreated)
-	{
-		NextTaskIndex = TaskIndex;
+		if (TaskState == ETaskState::ProcessOk
+		 || TaskState == ETaskState::ProcessFailed
+		 || TaskState == ETaskState::FileNotFound)
+		{
+			CompletedTaskCount++;
+		}
+
+		if (TaskState == ETaskState::UnTreated)
+		{
+			NextTaskIndex = TaskIndex;
+		}
 	}
+
+	UE_CLOG(TaskState == ETaskState::ProcessOk, LogDatasmithDispatcher, Verbose, TEXT("File processed: %s"), *Filename);
+	UE_CLOG(TaskState == ETaskState::UnTreated, LogDatasmithDispatcher, Warning, TEXT("File resubmitted: %s"), *Filename);
+	UE_CLOG(TaskState == ETaskState::ProcessFailed, LogDatasmithDispatcher, Error, TEXT("File processing failure: %s"), *Filename);
+	UE_CLOG(TaskState == ETaskState::FileNotFound, LogDatasmithDispatcher, Warning, TEXT("file not found: %s"), *Filename);
 }
 
 void FDatasmithDispatcher::Process(bool bWithProcessor)
@@ -95,9 +106,38 @@ void FDatasmithDispatcher::Process(bool bWithProcessor)
 	{
 		SpawnHandlers();
 
-		while (!IsOver() && GetAliveHandlerCount() > 0)
+		bool bLogRestartError = true;
+		while (!IsOver())
 		{
-			FWindowsPlatformProcess::Sleep(0.5);
+			bool bHasAliveWorker = false;
+			for (FDatasmithWorkerHandler& Handler : WorkerHandlers)
+			{
+				// replace dead workers
+				if (Handler.IsRestartable())
+				{
+					int32 WorkerId = GetNextWorkerId();
+					if (WorkerId < NumberOfWorkers + Config::MaxRestartAllowed)
+					{
+						Handler.~FDatasmithWorkerHandler();
+						new (&Handler) FDatasmithWorkerHandler(*this, ProcessCacheFolder, WorkerId);
+						UE_LOG(LogDatasmithDispatcher, Warning, TEXT("Restarting worker (new worker: %d)"), WorkerId);
+					}
+					else if (bLogRestartError)
+					{
+						bLogRestartError = false;
+						UE_LOG(LogDatasmithDispatcher, Warning, TEXT("Worker not restarted (Limit reached)"));
+					}
+				}
+
+				bHasAliveWorker = bHasAliveWorker || Handler.IsAlive();
+			}
+
+			if (!bHasAliveWorker)
+			{
+				break;
+			}
+
+			FWindowsPlatformProcess::Sleep(0.1);
 		}
 
 		CloseHandlers();
@@ -105,8 +145,13 @@ void FDatasmithDispatcher::Process(bool bWithProcessor)
 
 	if (!IsOver())
 	{
+		UE_LOG(LogDatasmithDispatcher, Warning, TEXT("Begin local processing. (Multi Process failed to consume all the tasks)"));
 		ProcessLocal();
 	}
+	else
+	{
+		UE_LOG(LogDatasmithDispatcher, Display, TEXT("Multi Process ended and consumed all the tasks"));
+	}
 }
 
 bool FDatasmithDispatcher::IsOver()
@@ -128,12 +173,17 @@ void FDatasmithDispatcher::LinkCTFileToUnrealCacheFile(const FString& CTFile, co
 	}
 }
 
+int32 FDatasmithDispatcher::GetNextWorkerId()
+{
+	return NextWorkerId++;
+}
+
 void FDatasmithDispatcher::SpawnHandlers()
 {
 	WorkerHandlers.Reserve(NumberOfWorkers);
 	for (int32 Index = 0; Index < NumberOfWorkers; Index++)
 	{
-		WorkerHandlers.Emplace(*this, ProcessCacheFolder, Index);
+		WorkerHandlers.Emplace(*this, ProcessCacheFolder, GetNextWorkerId());
 	}
 }
 
diff --git a/Engine/Plugins/Enterprise/DatasmithCADImporter/Source/DatasmithDispatcher/Private/DatasmithDispatcherModule.cpp b/Engine/Plugins/Enterprise/DatasmithCADImporter/Source/DatasmithDispatcher/Private/DatasmithDispatcherModule.cpp
--- a/Engine/Plugins/Enterprise/DatasmithCADImporter/Source/DatasmithDispatcher/Private/DatasmithDispatcherModule.cpp
+++ b/Engine/Plugins/Enterprise/DatasmithCADImporter/Source/DatasmithDispatcher/Private/DatasmithDispatcherModule.cpp
@@ -1,6 +1,8 @@
 // Copyright 1998-2019 Epic Games, Inc. All Rights Reserved.
 #include "DatasmithDispatcherModule.h"
 
+#include "DatasmithDispatcherLog.h"
+
 #include "Modules/ModuleInterface.h"
 #include "Modules/ModuleManager.h"
 
@@ -8,6 +10,7 @@
 
 #define LOCTEXT_NAMESPACE "FDatasmithDispatcherModule"
 
+DEFINE_LOG_CATEGORY(LogDatasmithDispatcher);
 
 FDatasmithDispatcherModule& FDatasmithDispatcherModule::Get()
 {
diff --git a/Engine/Plugins/Enterprise/DatasmithCADImporter/Source/DatasmithDispatcher/Private/DatasmithDispatcherNetworking.cpp b/Engine/Plugins/Enterprise/DatasmithCADImporter/Source/DatasmithDispatcher/Private/DatasmithDispatcherNetworking.cpp
--- a/Engine/Plugins/Enterprise/DatasmithCADImporter/Source/DatasmithDispatcher/Private/DatasmithDispatcherNetworking.cpp
+++ b/Engine/Plugins/Enterprise/DatasmithCADImporter/Source/DatasmithDispatcher/Private/DatasmithDispatcherNetworking.cpp
@@ -34,7 +34,7 @@ bool FNetworkNode::SendMessage(const TArray<uint8>& Message, double Timeout_s)
 
 	if (ConnectedSocket == nullptr)
 	{
-		UE_LOG(LogDatasmithDispatcher, Error, TEXT("null socket, can't write"));
+		UE_LOG(LogDatasmithDispatcher, Error, TEXT("node not connected, can't write"));
 		return false;
 	}
 
@@ -54,28 +54,10 @@ bool FNetworkNode::SendMessage(const TArray<uint8>& Message, double Timeout_s)
 
 	ConnectedSocket->SetNonBlocking(false);
 
-	{ // Send the header
-		FMessageHeader Header;
-		Header.ByteSize = Message.Num();
-
-		TArray<uint8> HeaderBuffer;
-		FMemoryWriter ArWriter(HeaderBuffer);
-		ArWriter << Header;
-
-		int32 ByteSent = -1;
-		bool bSendSucceed = ConnectedSocket->Send(HeaderBuffer.GetData(), HeaderBuffer.Num(), ByteSent);
-
-		if (!bSendSucceed || HeaderBuffer.Num() != ByteSent)
-		{
-			UE_LOG(LogDatasmithDispatcher, Error, TEXT("can't write header on socket"));
-			bWriteError = true;
-			return false;
-		}
-	}
-
-	{ // Send the content
-		const uint8* Data = Message.GetData();
-		uint32 BufferSize = Message.Num();
+	auto SendBuffer = [&](const TArray<uint8>& Buffer) -> bool
+	{
+		uint32 BufferSize = Buffer.Num();
+		const uint8* Data = Buffer.GetData();
 
 		uint32 TotalByteSent = 0;
 		bool bSendSucceed = true;
@@ -85,13 +67,30 @@ bool FNetworkNode::SendMessage(const TArray<uint8>& Message, double Timeout_s)
 			bSendSucceed &= ConnectedSocket->Send(Data + TotalByteSent, BufferSize - TotalByteSent, BytesSent);
 			TotalByteSent += (uint32) BytesSent;
 		}
+		return bSendSucceed && TotalByteSent == BufferSize;
+	};
 
-		if (!bSendSucceed)
-		{
-			UE_LOG(LogDatasmithDispatcher, Error, TEXT("can't write content on socket"));
-			bWriteError = true;
-			return false;
-		}
+	// Send the header
+	TArray<uint8> HeaderBuffer;
+	FMemoryWriter ArWriter(HeaderBuffer);
+
+	FMessageHeader Header;
+	Header.ByteSize = Message.Num();
+	ArWriter << Header;
+
+	if (!SendBuffer(HeaderBuffer))
+	{
+		UE_LOG(LogDatasmithDispatcher, Error, TEXT("can't write header on socket"));
+		bWriteError = true;
+		return false;
+	}
+
+	// Send the content
+	if (!SendBuffer(Message))
+	{
+		UE_LOG(LogDatasmithDispatcher, Error, TEXT("can't write content on socket"));
+		bWriteError = true;
+		return false;
 	}
 
 	return true;
@@ -105,6 +104,12 @@ bool FNetworkNode::ReceiveMessage(TArray<uint8>& OutMessage, double Timeout_s)
 		return false;
 	}
 
+	if (ConnectedSocket == nullptr)
+	{
+		UE_LOG(LogDatasmithDispatcher, Error, TEXT("node not connected, can't read"));
+		return false;
+	}
+
 	ConnectedSocket->SetNonBlocking(true);
 
 	bool bCanRead = ConnectedSocket && ConnectedSocket->Wait(ESocketWaitConditions::WaitForRead, TimespanFromSeconds(Timeout_s));
@@ -153,7 +158,6 @@ bool FNetworkNode::ReceiveMessage(TArray<uint8>& OutMessage, double Timeout_s)
 		}
 
 		IncommingMessage.Content.Reserve(IncommingMessage.Header.ByteSize);
-		UE_LOG(LogDatasmithDispatcher, Display, TEXT("Parsed header ok"));
 	}
 
 	// fill the message with available data
@@ -307,7 +311,7 @@ bool FNetworkClientNode::Connect(const FString& Description, int32 ServerPort, d
 	{
 		if (ConnectedSocket->Connect(*InternetAddress))
 		{
-			UE_LOG(LogDatasmithDispatcher, Display, TEXT("Is connected"));
+			UE_LOG(LogDatasmithDispatcher, Verbose, TEXT("Client Node is connected"));
 			return true;
 		}
 	}
@@ -342,16 +346,15 @@ bool FCommandQueue::SendCommand(ICommand& Commmand, double Timeout_s)
 	return NetworkInterface && NetworkInterface->SendMessage(CommandBuffer, Timeout_s);
 }
 
-bool FCommandQueue::Disconnect(double Timeout_s)
+void FCommandQueue::Disconnect(double Timeout_s)
 {
-	// consume all available commands
+	// consume all available commands before closing the connection
 	while (Poll(Timeout_s))
 	{
 		Timeout_s = 0;
 	}
 
-	// close the network interface
-	return true;
+	NetworkInterface = nullptr;
 }
 
 bool FCommandQueue::Poll(double Timeout_s)
diff --git a/Engine/Plugins/Enterprise/DatasmithCADImporter/Source/DatasmithDispatcher/Private/DatasmithWorkerHandler.cpp b/Engine/Plugins/Enterprise/DatasmithCADImporter/Source/DatasmithDispatcher/Private/DatasmithWorkerHandler.cpp
--- a/Engine/Plugins/Enterprise/DatasmithCADImporter/Source/DatasmithDispatcher/Private/DatasmithWorkerHandler.cpp
+++ b/Engine/Plugins/Enterprise/DatasmithCADImporter/Source/DatasmithDispatcher/Private/DatasmithWorkerHandler.cpp
@@ -2,6 +2,7 @@
 #include "DatasmithWorkerHandler.h"
 
 #include "DatasmithDispatcher.h"
+#include "DatasmithDispatcherConfig.h"
 #include "DatasmithDispatcherLog.h"
 #include "DatasmithCommands.h"
 
@@ -36,7 +37,6 @@ FDatasmithWorkerHandler::FDatasmithWorkerHandler(FDatasmithDispatcher& InDispatc
 	, CachePath(InCachePath)
 	, bShouldTerminate(false)
 {
-	CommandIO.SetNetworkInterface(&NetworkInterface);
 	ThreadName = FString(TEXT("DatasmithWorkerHandler_")) + FString::FromInt(Id);
 	IOThread = FThread(*ThreadName, [this]() { Run(); } );
 }
@@ -63,7 +63,7 @@ void FDatasmithWorkerHandler::StartWorkerProcess()
 		CommandToProcess += TEXT(" -ServerPID ") + FString::FromInt(FPlatformProcess::GetCurrentProcessId());
 		CommandToProcess += TEXT(" -ServerPort ") + FString::FromInt(ListenPort);
 		CommandToProcess += TEXT(" -CacheDir \"") + CachePath + TEXT('"');
-		UE_LOG(LogDatasmithDispatcher, Display, TEXT("CommandToProcess: %s"), *CommandToProcess);
+		UE_LOG(LogDatasmithDispatcher, Verbose, TEXT("CommandToProcess: %s"), *CommandToProcess);
 
 		uint32 WorkerProcessId = 0;
 		WorkerHandle = FPlatformProcess::CreateProc(*ProcessorPath, *CommandToProcess, true, false, false, &WorkerProcessId, 0, nullptr, nullptr);
@@ -92,18 +92,12 @@ void FDatasmithWorkerHandler::ValidateConnection()
 	}
 }
 
-// void FDatasmithWorkerHandler::RestartProcessor()
-// {
-// 	FPlatformProcess::TerminateProc(WorkerHandle, true);
-// 	StartWorker();
-// }
-
 void FDatasmithWorkerHandler::Run()
 {
 	WorkerState = EWorkerState::Uninitialized;
 	RunInternal();
-	UE_CLOG(ErrorState != EWorkerErrorState::Ok, LogDatasmithDispatcher, Error, TEXT("ErrorState != OK on exit (%d)"), uint32(ErrorState));
 	WorkerState = EWorkerState::Terminated;
+	UE_CLOG(ErrorState != EWorkerErrorState::Ok, LogDatasmithDispatcher, Warning, TEXT("Handler ended with error %s"), EWorkerErrorStateAsString(ErrorState));
 }
 
 void FDatasmithWorkerHandler::RunInternal()
@@ -126,11 +120,14 @@ void FDatasmithWorkerHandler::RunInternal()
 
 				// The Accept() call on the server blocks until a connection is initiated from a client
 				static const FString SocketDescription = TEXT("DatasmithWorkerHandler");
-				double AcceptTimeout_s = 300;
-				if (!NetworkInterface.Accept(SocketDescription, AcceptTimeout_s))
+				if (!NetworkInterface.Accept(SocketDescription, Config::AcceptTimeout_s))
 				{
 					ErrorState = EWorkerErrorState::ConnectionFailed_NoClient;
 				}
+				else
+				{
+					CommandIO.SetNetworkInterface(&NetworkInterface);
+				}
 
 				if (ErrorState != EWorkerErrorState::Ok)
 				{
@@ -144,22 +141,6 @@ void FDatasmithWorkerHandler::RunInternal()
 
 			case EWorkerState::Idle:
 			{
-// 				// Verify connection
-// 				if (!ClientSocketWrapper.IsConnected() || !ClientSocketWrapper.CanWrite())
-// 				{
-// 					ErrorState = EWorkerErrorState::ConnectionLost;
-// 					if (CurrentTask.IsSet())
-// 					{
-// 						Dispatcher.SetTaskState(CurrentTask->Index, ETaskState::ProcessFailed);
-// 						CurrentTask.Reset();
-// 					}
-//
-// 					// Restart processor
-// 					FPlatformProcess::TerminateProc(WorkerHandle, true);
-// 					WorkerState = EWorkerState::Uninitialized;
-// 						// RestartProcessor();
-// 				}
-
 				// Fetch a new task
 				ensureMsgf(CurrentTask.IsSet() == false, TEXT("We should not have a current task when fetching a new one"));
 				CurrentTask = Dispatcher.GetNextTask();
@@ -168,10 +149,9 @@ void FDatasmithWorkerHandler::RunInternal()
 				{
 					FRunTaskCommand NewTask(CurrentTask.GetValue());
 
-					double SendTimeout_s = 3.0;
-					if (CommandIO.SendCommand(NewTask, SendTimeout_s))
+					if (CommandIO.SendCommand(NewTask, Config::SendCommandTimeout_s))
 					{
-						UE_LOG(LogDatasmithDispatcher, Display, TEXT("New task command sent"));
+						UE_LOG(LogDatasmithDispatcher, Verbose, TEXT("New task command sent"));
 						WorkerState = EWorkerState::Processing;
 					}
 					else
@@ -186,7 +166,7 @@ void FDatasmithWorkerHandler::RunInternal()
 				}
 				else if (bShouldTerminate)
 				{
-					UE_LOG(LogDatasmithDispatcher, Display, TEXT("Exit loop gracefully"));
+					UE_LOG(LogDatasmithDispatcher, Verbose, TEXT("Exit loop gracefully"));
 					WorkerState = EWorkerState::Closing;
 				}
 				else
@@ -194,7 +174,7 @@ void FDatasmithWorkerHandler::RunInternal()
 					ValidateConnection();
 
 					// consume
-					if (TSharedPtr<ICommand> Command = CommandIO.GetNextCommand(0.1))
+					if (TSharedPtr<ICommand> Command = CommandIO.GetNextCommand(Config::IdleLoopDelay))
 					{
 						ProcessCommand(*Command);
 					}
@@ -205,7 +185,7 @@ void FDatasmithWorkerHandler::RunInternal()
 
 			case EWorkerState::Processing:
 			{
-				if (TSharedPtr<ICommand> Command = CommandIO.GetNextCommand(0.5))
+				if (TSharedPtr<ICommand> Command = CommandIO.GetNextCommand(Config::ProcessingLoopDelay))
 				{
 					ProcessCommand(*Command);
 
@@ -220,48 +200,40 @@ void FDatasmithWorkerHandler::RunInternal()
 					ValidateConnection();
 					if (ErrorState == EWorkerErrorState::WorkerProcess_Lost)
 					{
-						WorkerState = EWorkerState::Restarting;
+						if (CurrentTask.IsSet())
+						{
+							Dispatcher.SetTaskState(CurrentTask->Index, ETaskState::ProcessFailed);
+							CurrentTask.Reset();
+						}
+
+						WorkerState = EWorkerState::Closing;
 					}
 				}
 				break;
 			}
 
-			case EWorkerState::Restarting:
-			{
-				// TODO:
-				// close the network interface but fetch all data before
-				// StartWorkerProcess();
-				// reconnect network (Accept...)
-
-				WorkerState = EWorkerState::Closing; // Not implemented yet
-				break;
-			}
-
 			case EWorkerState::Closing:
 			{
 				// try to close the process gracefully
-				bool CloseByCommand = NetworkInterface.IsValid() &&
-					(!WorkerHandle.IsValid() || FPlatformProcess::IsProcRunning(WorkerHandle));
-
-				if (CloseByCommand)
-				{
-					FTerminateCommand Terminate;
-					CommandIO.SendCommand(Terminate, 0);
-				}
-
 				if (WorkerHandle.IsValid())
 				{
-					int32 CloseTimeout_s = CloseByCommand ? 1 : 0;
+					bool CloseByCommand = Config::CloseProcessByCommand && CommandIO.IsValid() && FPlatformProcess::IsProcRunning(WorkerHandle);
 
 					bool bClosed = false;
-					for (int32 i = 0; i < 10 * CloseTimeout_s; ++i)
+					if (CloseByCommand)
 					{
-						if (!FPlatformProcess::IsProcRunning(WorkerHandle))
+						FTerminateCommand Terminate;
+						CommandIO.SendCommand(Terminate, 0);
+
+						for (int32 i = 0; i < int32(10. * Config::TerminateTimeout_s); ++i)
 						{
-							bClosed = true;
-							break;
+							if (!FPlatformProcess::IsProcRunning(WorkerHandle))
+							{
+								bClosed = true;
+								break;
+							}
+							FPlatformProcess::Sleep(0.1);
 						}
-						FPlatformProcess::Sleep(0.1);
 					}
 
 					if (!bClosed)
@@ -270,9 +242,8 @@ void FDatasmithWorkerHandler::RunInternal()
 					}
 				}
 
-				CommandIO.Disconnect(0);
-
 				// Process commands still in input queue
+				CommandIO.Disconnect(0);
 				while (TSharedPtr<ICommand> Command = CommandIO.GetNextCommand(0))
 				{
 					ProcessCommand(*Command);
@@ -290,7 +261,6 @@ void FDatasmithWorkerHandler::RunInternal()
 	}
 }
 
-
 void FDatasmithWorkerHandler::Stop()
 {
 	bShouldTerminate = true;
@@ -301,6 +271,14 @@ bool FDatasmithWorkerHandler::IsAlive() const
 	return WorkerState != EWorkerState::Terminated;
 }
 
+bool FDatasmithWorkerHandler::IsRestartable() const
+{
+	return !IsAlive() && (
+		ErrorState == EWorkerErrorState::WorkerProcess_Lost ||
+		ErrorState == EWorkerErrorState::ConnectionLost
+	);
+}
+
 void FDatasmithWorkerHandler::ProcessCommand(ICommand& Command)
 {
 	switch (Command.GetType())
@@ -341,4 +319,20 @@ void FDatasmithWorkerHandler::ProcessCommand(FCompletedTaskCommand& CompletedTas
 	Dispatcher.LinkCTFileToUnrealCacheFile(CurrentFileName, CompletedTaskCommand.SceneGraphFileName, CompletedTaskCommand.GeomFileName);
 	CurrentTask.Reset();
 }
+
+const TCHAR* FDatasmithWorkerHandler::EWorkerErrorStateAsString(EWorkerErrorState e)
+{
+	switch (e)
+	{
+		case EWorkerErrorState::Ok: return TEXT("Ok");
+		case EWorkerErrorState::ConnectionFailed_NotBound: return TEXT("ConnectionFailed_NotBound");
+		case EWorkerErrorState::ConnectionFailed_NoClient: return TEXT("ConnectionFailed_NoClient");
+		case EWorkerErrorState::ConnectionLost: return TEXT("ConnectionLost");
+		case EWorkerErrorState::ConnectionLost_SendFailed: return TEXT("ConnectionLost_SendFailed");
+		case EWorkerErrorState::WorkerProcess_CantCreate: return TEXT("WorkerProcess_CantCreate");
+		case EWorkerErrorState::WorkerProcess_Lost: return TEXT("WorkerProcess_Lost");
+		default: return TEXT("unknown");
+	}
+}
+
 } // ns DatasmithDispatcher
diff --git a/Engine/Plugins/Enterprise/DatasmithCADImporter/Source/DatasmithDispatcher/Public/DatasmithDispatcher.h b/Engine/Plugins/Enterprise/DatasmithCADImporter/Source/DatasmithDispatcher/Public/DatasmithDispatcher.h
--- a/Engine/Plugins/Enterprise/DatasmithCADImporter/Source/DatasmithDispatcher/Public/DatasmithDispatcher.h
+++ b/Engine/Plugins/Enterprise/DatasmithCADImporter/Source/DatasmithDispatcher/Public/DatasmithDispatcher.h
@@ -29,6 +29,7 @@ public:
 
 private:
 	void SpawnHandlers();
+	int32 GetNextWorkerId();
 	int32 GetAliveHandlerCount();
 	void CloseHandlers();
 
@@ -49,6 +50,7 @@ private:
 
 	// Workers
 	int32 NumberOfWorkers;
+	int32 NextWorkerId;
 	TArray<FDatasmithWorkerHandler> WorkerHandlers;
 };
 
diff --git a/Engine/Plugins/Enterprise/DatasmithCADImporter/Source/DatasmithDispatcher/Public/DatasmithDispatcherConfig.h b/Engine/Plugins/Enterprise/DatasmithCADImporter/Source/DatasmithDispatcher/Public/DatasmithDispatcherConfig.h
new file mode 100644
--- /dev/null
+++ b/Engine/Plugins/Enterprise/DatasmithCADImporter/Source/DatasmithDispatcher/Public/DatasmithDispatcherConfig.h
@@ -0,0 +1,37 @@
+// Copyright 1998-2019 Epic Games, Inc. All Rights Reserved.
+
+#pragma once
+
+namespace DatasmithDispatcher
+{
+	namespace Config
+	{
+		// Server side timeout for connection with the client
+		constexpr double AcceptTimeout_s = 3.0;
+
+		// Client side timeout for connection with the server
+		constexpr double ConnectTimeout_s = 3.0;
+
+		// Maximum time we accept to wait when a command is sent.
+		// If the command is not sent by this delay, the operation is considered as failed
+		constexpr double SendCommandTimeout_s = 1.0;
+
+		// When true, a Terminate command is sent to the processor and we wait for
+		// TerminateTimeout_s seconds for it to close, before a hard TerminateProc.
+		// When false, the process is directly closed with a TerminateProc.
+		constexpr bool CloseProcessByCommand = true;
+
+		// Maximum time we accept to wait for the process to close after a Terminate
+		// command has been sent.
+		constexpr double TerminateTimeout_s = 1.0;
+
+		// Delay between successive checks for a new task, and validation of the connection
+		constexpr double IdleLoopDelay = 0.1;
+
+		// Delay between successive validation of the connection
+		constexpr double ProcessingLoopDelay = 0.5;
+
+		// Number of worker instance that can be restarted when a crash occurs in en external process
+		constexpr int32 MaxRestartAllowed = 5;
+	};
+} // ns DatasmithDispatcher
diff --git a/Engine/Plugins/Enterprise/DatasmithCADImporter/Source/DatasmithDispatcher/Public/DatasmithDispatcherNetworking.h b/Engine/Plugins/Enterprise/DatasmithCADImporter/Source/DatasmithDispatcher/Public/DatasmithDispatcherNetworking.h
--- a/Engine/Plugins/Enterprise/DatasmithCADImporter/Source/DatasmithDispatcher/Public/DatasmithDispatcherNetworking.h
+++ b/Engine/Plugins/Enterprise/DatasmithCADImporter/Source/DatasmithDispatcher/Public/DatasmithDispatcherNetworking.h
@@ -71,11 +71,10 @@ protected:
 	FSocket* ConnectedSocket = nullptr;
 	SocketErrorCode ConnectedSocketError = SocketErrorCode::NoError;
 
-	bool bReadError = false;
 	FMessage IncommingMessage;
 
+	bool bReadError = false;
 	bool bWriteError = false;
-// 	TArray<uint8> IncompleteCommandBufferOut;
 };
 
 
@@ -115,7 +114,7 @@ public:
 
 	bool IsValid() { return NetworkInterface ? NetworkInterface->IsValid() : true;}
 
-	bool Disconnect(double Timeout_s);
+	void Disconnect(double Timeout_s);
 
 private:
 	bool Poll(double Timeout_s);
diff --git a/Engine/Plugins/Enterprise/DatasmithCADImporter/Source/DatasmithDispatcher/Public/DatasmithWorkerHandler.h b/Engine/Plugins/Enterprise/DatasmithCADImporter/Source/DatasmithDispatcher/Public/DatasmithWorkerHandler.h
--- a/Engine/Plugins/Enterprise/DatasmithCADImporter/Source/DatasmithDispatcher/Public/DatasmithWorkerHandler.h
+++ b/Engine/Plugins/Enterprise/DatasmithCADImporter/Source/DatasmithDispatcher/Public/DatasmithWorkerHandler.h
@@ -25,7 +25,6 @@ class FDatasmithWorkerHandler
 		Uninitialized,
 		Idle, // Initialized, available for processing
 		Processing, // Currently processing a task
-		Restarting, // can occur when a processing is aborted.
 		Closing, // in the process of terminating
 		Terminated, // aka. Not Alive
 	};
@@ -47,6 +46,7 @@ public:
 
 	void Run();
 	bool IsAlive() const;
+	bool IsRestartable() const;
 	void Stop();
 
 private:
@@ -59,6 +59,7 @@ private:
 	void ProcessCommand(ICommand& Command);
 	void ProcessCommand(FPingCommand& PingCommand);
 	void ProcessCommand(FCompletedTaskCommand& RunTaskCommand);
+	const TCHAR* EWorkerErrorStateAsString(EWorkerErrorState e);
 
 private:
 	FDatasmithDispatcher& Dispatcher;
@@ -68,7 +69,6 @@ private:
 	FCommandQueue CommandIO;
 	FThread IOThread;
 	FString ThreadName;
-	// FEvent* todo;
 
 	// External process
 	FProcHandle WorkerHandle;
