Enable MT for cad translator loadstaticMesh
#rnx

diff --git a/Engine/Plugins/Enterprise/DatasmithCADImporter/Source/DatasmithCADTranslator/Private/DatasmithCADTranslator.cpp b/Engine/Plugins/Enterprise/DatasmithCADImporter/Source/DatasmithCADTranslator/Private/DatasmithCADTranslator.cpp
--- a/Engine/Plugins/Enterprise/DatasmithCADImporter/Source/DatasmithCADTranslator/Private/DatasmithCADTranslator.cpp
+++ b/Engine/Plugins/Enterprise/DatasmithCADImporter/Source/DatasmithCADTranslator/Private/DatasmithCADTranslator.cpp
@@ -13,6 +13,8 @@
 
 void FDatasmithCADTranslator::Initialize(FDatasmithTranslatorCapabilities& OutCapabilities)
 {
+	OutCapabilities.bParallelLoadStaticMeshSupported = true;
+
 	OutCapabilities.SupportedFileFormats.Add(FFileFormatInfo{ TEXT("CATPart"), TEXT("CATIA Part files") });
 	OutCapabilities.SupportedFileFormats.Add(FFileFormatInfo{ TEXT("CATProduct"), TEXT("CATIA Product files") });
 	OutCapabilities.SupportedFileFormats.Add(FFileFormatInfo{ TEXT("cgr"), TEXT("CATIA Graphical Representation V5 files") });
