Fix Github build

#rnx

diff --git a/Engine/Plugins/Enterprise/DatasmithCADImporter/Source/CADInterfaces/Public/CoreTechFileParser.h b/Engine/Plugins/Enterprise/DatasmithCADImporter/Source/CADInterfaces/Public/CoreTechFileParser.h
--- a/Engine/Plugins/Enterprise/DatasmithCADImporter/Source/CADInterfaces/Public/CoreTechFileParser.h
+++ b/Engine/Plugins/Enterprise/DatasmithCADImporter/Source/CADInterfaces/Public/CoreTechFileParser.h
@@ -23,9 +23,21 @@
 
 struct FFileStatData;
 
-#ifdef CAD_INTERFACE
 namespace CADLibrary
 {
+
+enum class ECoretechParsingResult
+{
+	Unknown,
+	Running,
+	UnTreated,
+	ProcessOk,
+	ProcessFailed,
+	FileNotFound,
+};
+
+#ifdef CAD_INTERFACE
+
 CADINTERFACES_API void GetColor(uint32 ColorHash, FColor& OutMaterial);
 CADINTERFACES_API bool GetMaterial(uint32 MaterialID, FCADMaterial& OutMaterial);
 
@@ -52,15 +64,7 @@ uint32 GetSize(CT_TESS_DATA_TYPE type);
 class CADINTERFACES_API FCoreTechFileParser
 {
 public:
-	enum class EProcessResult
-	{
-		Unknown,
-		Running,
-		UnTreated,
-		ProcessOk,
-		ProcessFailed,
-		FileNotFound,
-	};
+	using EProcessResult = ECoretechParsingResult;
 
 	FCoreTechFileParser(const FString& InCTFullPath, const FString& InCachePath, const FImportParameters& ImportParams, bool bPreferBodyData, bool bScaleUVMap);
 
@@ -132,5 +136,6 @@ protected:
 	const FImportParameters& ImportParameters;
 };
 
-}
 #endif // CAD_INTERFACE
+
+} // ns CADLibrary
diff --git a/Engine/Plugins/Enterprise/DatasmithCADImporter/Source/DatasmithDispatcher/Private/DatasmithDispatcher.cpp b/Engine/Plugins/Enterprise/DatasmithCADImporter/Source/DatasmithDispatcher/Private/DatasmithDispatcher.cpp
--- a/Engine/Plugins/Enterprise/DatasmithCADImporter/Source/DatasmithDispatcher/Private/DatasmithDispatcher.cpp
+++ b/Engine/Plugins/Enterprise/DatasmithCADImporter/Source/DatasmithDispatcher/Private/DatasmithDispatcher.cpp
@@ -2,7 +2,6 @@
 
 #include "DatasmithDispatcher.h"
 
-#include "CoreTechFileParser.h"
 #include "DatasmithDispatcherConfig.h"
 #include "DatasmithDispatcherLog.h"
 #include "DatasmithDispatcherTask.h"
diff --git a/Engine/Plugins/Enterprise/DatasmithCADImporter/Source/DatasmithDispatcher/Public/DatasmithDispatcherTask.h b/Engine/Plugins/Enterprise/DatasmithCADImporter/Source/DatasmithDispatcher/Public/DatasmithDispatcherTask.h
--- a/Engine/Plugins/Enterprise/DatasmithCADImporter/Source/DatasmithDispatcher/Public/DatasmithDispatcherTask.h
+++ b/Engine/Plugins/Enterprise/DatasmithCADImporter/Source/DatasmithDispatcher/Public/DatasmithDispatcherTask.h
@@ -7,7 +7,7 @@
 
 namespace DatasmithDispatcher
 {
-using ETaskState = CADLibrary::FCoreTechFileParser::EProcessResult;
+using ETaskState = CADLibrary::ECoretechParsingResult;
 
 struct FTask
 {
