Datasmith import of 3dm file produces bad geometry
Kernel_io don`t like singular patch like a patch with cyclic boundary like 8
Each patch is checked if it has a cyclic boundary. If yes, the body is converted by kernel_io function

diff --git a/Engine/Plugins/Enterprise/DatasmithCADImporter/Source/CADInterfaces/Private/CTSession.cpp b/Engine/Plugins/Enterprise/DatasmithCADImporter/Source/CADInterfaces/Private/CTSession.cpp
--- a/Engine/Plugins/Enterprise/DatasmithCADImporter/Source/CADInterfaces/Private/CTSession.cpp
+++ b/Engine/Plugins/Enterprise/DatasmithCADImporter/Source/CADInterfaces/Private/CTSession.cpp
@@ -31,6 +31,11 @@ namespace CADLibrary
 		return CTKIO_SaveFile(ObjectList, *FilePath, L"Ct");
 	}
 
+	CheckedCTError CTSession::CleanBRep()
+	{
+		return CTKIO_CleanBody(MainObjectId);
+	}
+
 	CheckedCTError CTSession::TopoFixes()
 	{
 		return CADLibrary::Repair(MainObjectId, ImportParams.StitchingTechnique);
diff --git a/Engine/Plugins/Enterprise/DatasmithCADImporter/Source/CADInterfaces/Private/CoreTechTypes.cpp b/Engine/Plugins/Enterprise/DatasmithCADImporter/Source/CADInterfaces/Private/CoreTechTypes.cpp
--- a/Engine/Plugins/Enterprise/DatasmithCADImporter/Source/CADInterfaces/Private/CoreTechTypes.cpp
+++ b/Engine/Plugins/Enterprise/DatasmithCADImporter/Source/CADInterfaces/Private/CoreTechTypes.cpp
@@ -58,6 +58,20 @@ namespace CADLibrary
 		return CT_KERNEL_IO::LoadFile(FileName, MainObject, LoadFlags, Lod, StringOption);
 	}
 
+	CT_IO_ERROR CTKIO_CleanBody(CT_OBJECT_ID MainObjectId)
+	{
+		CT_STR format = "Ct";          
+		CT_STR system = "Evolution";
+
+		CT_IO_ERROR result = CT_CONVERT_IO::SetConvertConfiguration(system, format);
+		if (result == IO_OK)
+		{
+			result = CT_CONVERT_IO::Convert(MainObjectId);
+		}
+		return result;
+	}
+
+
 	CT_IO_ERROR CTKIO_SaveFile(CT_LIST_IO& objects_list_to_save, const TCHAR* file_name, const TCHAR* format, const CT_OBJECT_ID coordsystem)
 	{
 		return CT_KERNEL_IO::SaveFile(objects_list_to_save, file_name, format, coordsystem);
diff --git a/Engine/Plugins/Enterprise/DatasmithCADImporter/Source/CADInterfaces/Public/CTSession.h b/Engine/Plugins/Enterprise/DatasmithCADImporter/Source/CADInterfaces/Public/CTSession.h
--- a/Engine/Plugins/Enterprise/DatasmithCADImporter/Source/CADInterfaces/Public/CTSession.h
+++ b/Engine/Plugins/Enterprise/DatasmithCADImporter/Source/CADInterfaces/Public/CTSession.h
@@ -37,6 +37,12 @@ public:
 	CheckedCTError SaveBrep(const FString& FilePath);
 	CheckedCTError TopoFixes();
 
+	/**
+	 * In case of patch with cyclic boundary, a process has to be done by kernel_IO
+	 */
+	CheckedCTError CleanBRep();
+
+
 	/**
 	 * @param InScaleFactor : use to scale meshing from Kernel-IO
 	 */
diff --git a/Engine/Plugins/Enterprise/DatasmithCADImporter/Source/CADInterfaces/Public/CoreTechTypes.h b/Engine/Plugins/Enterprise/DatasmithCADImporter/Source/CADInterfaces/Public/CoreTechTypes.h
--- a/Engine/Plugins/Enterprise/DatasmithCADImporter/Source/CADInterfaces/Public/CoreTechTypes.h
+++ b/Engine/Plugins/Enterprise/DatasmithCADImporter/Source/CADInterfaces/Public/CoreTechTypes.h
@@ -12,6 +12,7 @@
 #pragma warning(push)
 #pragma warning(disable:4996) // unsafe sprintf
 #pragma warning(disable:4828) // illegal character
+#include "kernel_io/convert_io/convert_io.h"
 #include "kernel_io/kernel_io.h"
 #include "kernel_io/kernel_io_error.h"
 #include "kernel_io/kernel_io_type.h"
@@ -48,7 +49,7 @@ namespace CADLibrary
 	};
 
 	/**
-	 * Mandatory: Kernel_IO has to be initialize with the final MetricUnit, set function doesn`t work
+	 * Mandatory: Kernel_IO has to be initialize with the final MetricUnit, set function doesn't work
 	 * @param MetricUnit: Length unit express in meter
 	 */
 	CADINTERFACES_API CT_IO_ERROR CTKIO_InitializeKernel(double MetricUnit);
@@ -98,6 +99,12 @@ namespace CADLibrary
 		const TCHAR*  string_option = TEXT("")               
 	);
 
+	/**
+	 * In case of patch with cyclic boundary, a process has to be done by kernel_IO
+	 * @param MainObjectId: Body to process
+	 */
+	CADINTERFACES_API CT_IO_ERROR CTKIO_CleanBody(CT_OBJECT_ID MainObjectId);
+
 	CADINTERFACES_API CT_IO_ERROR CTKIO_SaveFile
 	(
 		CT_LIST_IO&        objects_list_to_save,
@@ -133,7 +140,7 @@ namespace CADLibrary
 		 * Make sure CT is initialized, and a main object is ready.
 		 * Handle input file unit and an output unit
 		 * @param FileMetricUnit number of meters per file unit.
-		 * eg. For a file in inches, arg should be 0.0254
+		 * e.g. For a file in inches, arg should be 0.0254
 		 */
 		CoreTechSessionBase(const TCHAR* Owner, double Unit);
 		bool IsSessionValid() { return Owner != nullptr && MainObjectId != 0; }
diff --git a/Engine/Plugins/Enterprise/DatasmithCADImporter/Source/DatasmithOpenNurbsTranslator/Private/DatasmithOpenNurbsTranslator.cpp b/Engine/Plugins/Enterprise/DatasmithCADImporter/Source/DatasmithOpenNurbsTranslator/Private/DatasmithOpenNurbsTranslator.cpp
--- a/Engine/Plugins/Enterprise/DatasmithCADImporter/Source/DatasmithOpenNurbsTranslator/Private/DatasmithOpenNurbsTranslator.cpp
+++ b/Engine/Plugins/Enterprise/DatasmithCADImporter/Source/DatasmithOpenNurbsTranslator/Private/DatasmithOpenNurbsTranslator.cpp
@@ -2839,10 +2839,17 @@ bool FOpenNurbsTranslatorImpl::TranslateBRep(ON_Brep* Brep, const ON_3dmObjectAt
 
 	LocalSession->ClearData();
 
-	Result = LocalSession->AddBRep(*Brep);
+	bool bBRepHasLoopNonManifold = false;
+	Result = LocalSession->AddBRep(*Brep, bBRepHasLoopNonManifold);
 
 	FString Filename = FString::Printf(TEXT("%s.ct"), *Name);
 	FString FilePath = FPaths::Combine(OutputPath, Filename);
+
+	if(bBRepHasLoopNonManifold)
+	{
+		LocalSession->CleanBRep();
+	}
+
 	Result = LocalSession->SaveBrep(FilePath);
 	if (Result)
 	{
diff --git a/Engine/Plugins/Enterprise/DatasmithCADImporter/Source/DatasmithOpenNurbsTranslator/Private/RhinoCoretechWrapper.cpp b/Engine/Plugins/Enterprise/DatasmithCADImporter/Source/DatasmithOpenNurbsTranslator/Private/RhinoCoretechWrapper.cpp
--- a/Engine/Plugins/Enterprise/DatasmithCADImporter/Source/DatasmithOpenNurbsTranslator/Private/RhinoCoretechWrapper.cpp
+++ b/Engine/Plugins/Enterprise/DatasmithCADImporter/Source/DatasmithOpenNurbsTranslator/Private/RhinoCoretechWrapper.cpp
@@ -4,8 +4,10 @@
 
 
 #ifdef CAD_LIBRARY
+
 #include "CoreTechHelper.h"
-//#include "TessellationHelper.h"
+
+#include "Containers/Set.h"
 
 #pragma warning(push)
 #pragma warning(disable:4265)
@@ -193,15 +195,12 @@ CT_OBJECT_ID CreateCTSurface(ON_NurbsSurface& Surface)
 	return CTSurfaceID;
 }
 
-void CreateCTFace_internal(const ON_BrepFace& Face, CT_LIST_IO& dest, ON_BoundingBox& outerBBox, ON_NurbsSurface& Surface, bool ignoreInner)
+void CreateCTFace_internal(const ON_BrepFace& Face, CT_LIST_IO& dest, ON_BoundingBox& outerBBox, ON_NurbsSurface& Surface, bool ignoreInner, bool& bOutFaceHasLoopNonManifold)
 {
 	CT_OBJECT_ID SurfaceID = CreateCTSurface(Surface);
 	if (SurfaceID == 0)
 		return;
 
-	ON_BoundingBox realOuterBBox = Face.OuterLoop()->BoundingBox();
-	bool outerRedefined = realOuterBBox != outerBBox;
-
 	int LoopCount = ignoreInner ? 1 : Face.LoopCount();
 	CT_LIST_IO Loops;
 	for (int LoopIndex = 0; LoopIndex < LoopCount; ++LoopIndex)
@@ -212,10 +211,12 @@ void CreateCTFace_internal(const ON_BrepFace& Face, CT_LIST_IO& dest, ON_Boundin
 
 		ON_BrepLoop::TYPE onLoopType = on_Loop.m_type;
 		bool bIsOuter = (onLoopType == ON_BrepLoop::TYPE::outer);
-		//assert(bIsOuter == (LoopIndex == 0));
 
-		CT_LIST_IO Coedges;
 		int TrimCount = on_Loop.TrimCount();
+
+		CT_LIST_IO Coedges;
+		TSet<ON_BrepVertex*> VertexSet;
+		VertexSet.Reserve(TrimCount);
 		for (int i = 0; i < TrimCount; ++i)
 		{
 			CADLibrary::CheckedCTError err = IO_OK;
@@ -226,7 +227,7 @@ void CreateCTFace_internal(const ON_BrepFace& Face, CT_LIST_IO& dest, ON_Boundin
 				continue;
 
 			ON_NurbsCurve nurbs_curve;
-			int nurbFormSuccess = Trim.GetNurbForm(nurbs_curve); // 0:Nok 1:Ok 2:OkBut
+			int nurbFormSuccess = Trim.GetNurbForm(nurbs_curve); // 0:No ok 1:Ok 2:OkBut
 			if (nurbFormSuccess == 0)
 				continue;
 
@@ -235,6 +236,13 @@ void CreateCTFace_internal(const ON_BrepFace& Face, CT_LIST_IO& dest, ON_Boundin
 			if (err != IO_OK)
 				continue;
 
+			// check if the loop is non manifold i.e. is auto-intersecting like a 8. We get the second vertex of each edge. If the vertex is already in the set, this mean that the loop is at least like a 8
+			if (!bOutFaceHasLoopNonManifold)
+			{
+				ON_BrepVertex* Vertex = on_edge->Vertex(1);
+				VertexSet.Add(Vertex, &bOutFaceHasLoopNonManifold);
+			}
+
 			// fill edge data
 			CT_UINT32 order = nurbs_curve.Order();
 
@@ -262,13 +270,6 @@ void CreateCTFace_internal(const ON_BrepFace& Face, CT_LIST_IO& dest, ON_Boundin
 			{
 				nurbs_curve.GetCV(j, nurbs_curve.IsRational() ? ON::point_style::euclidean_rational : ON::point_style::not_rational, wp);
 			}
-			if (outerRedefined && bIsOuter)
-			{
-				for (int j = 0; j < cv_data.size(); j += ctrl_hull_dim)
-				{
-					cv_data[j] = min(max(cv_data[j], outerBBox.m_min.x), outerBBox.m_max.x);
-				}
-			}
 
 			// knot multiplicity (ignored as knots are stored multiple times already)
 			std::vector<CT_UINT32> knotMult;
@@ -310,7 +311,7 @@ void CreateCTFace_internal(const ON_BrepFace& Face, CT_LIST_IO& dest, ON_Boundin
 	dest.PushBack(FaceID);
 }
 
-void CreateCTFace(const ON_Brep& brep, const ON_BrepFace& Face, CT_LIST_IO& dest)
+void CreateCTFace(const ON_Brep& brep, const ON_BrepFace& Face, CT_LIST_IO& dest, bool& bOutFaceHasLoopNonManifold)
 {
 	const ON_BrepLoop* outerLoop = Face.OuterLoop();
 	if (outerLoop == nullptr)
@@ -338,47 +339,12 @@ void CreateCTFace(const ON_Brep& brep, const ON_BrepFace& Face, CT_LIST_IO& dest
 		}
 		bBadLoopHack &= hasSingularTrim;
 	}
-
-	if (bBadLoopHack)
-	{
-		// we try to split in 2 faces: one west one east.
-		double span = outerBBox.Diagonal().x;
-
-		ON_BoundingBox innerBBox = Face.Loop(1)->BoundingBox();
-		for (int LoopIndex = 2; LoopIndex < LoopCount; ++LoopIndex)
-		{
-			const ON_BrepLoop& innerLoop = *Face.Loop(LoopIndex);
-			innerBBox.Union(innerLoop.BoundingBox());
-		}
-		double spaceWest = innerBBox.Min().x - outerBBox.Min().x;
-		double spaceEast = outerBBox.Max().x - innerBBox.Max().x;
-		double e = 0.01;
-		if (spaceEast < e*span && spaceWest < e*span)
-		{
-			// can't split... We could ignore inner loops:
-			CreateCTFace_internal(Face, dest, outerBBox, Surface, true);
-			return;
-		}
-
-		// split the outer bbox in a normal part (with outer) and a 'rest' (with no iner)
-		ON_BoundingBox outerTrimedPartBBox = outerBBox;
-		if (spaceEast > spaceWest)
-		{
-			outerBBox.m_max.x -= 0.5 * spaceEast;
-			outerTrimedPartBBox.m_min.x = outerBBox.m_max.x;
-		}
-		else
-		{
-			outerBBox.m_min.x += 0.5 * spaceWest;
-			outerTrimedPartBBox.m_max.x = outerBBox.m_min.x;
-		}
-
-		// We now need to create a bonus face to handle the hole we just created
-		CreateCTFace_internal(Face, dest, outerTrimedPartBBox, Surface, true);
-	}
+	bOutFaceHasLoopNonManifold = bBadLoopHack;
+	
 #endif // FIX_HOLE_IN_WHOLE_FACE
-
-	CreateCTFace_internal(Face, dest, outerBBox, Surface, false);
+	bool bHasCyclicLoop = false;
+	CreateCTFace_internal(Face, dest, outerBBox, Surface, false, bHasCyclicLoop);
+	bOutFaceHasLoopNonManifold |= bHasCyclicLoop;
 }
 
 
@@ -387,7 +353,7 @@ CT_IO_ERROR FRhinoCoretechWrapper::Tessellate(FMeshDescription& Mesh, CADLibrary
 	return CADLibrary::Tessellate(MainObjectId, ImportParams, Mesh, MeshParameters);
 }
 
-CADLibrary::CheckedCTError FRhinoCoretechWrapper::AddBRep(ON_Brep& Brep)
+CADLibrary::CheckedCTError FRhinoCoretechWrapper::AddBRep(ON_Brep& Brep, bool& bOutBRepHasLoopNonManifold)
 {
 	CADLibrary::CheckedCTError Result;
 	if (!IsSessionValid())
@@ -403,7 +369,9 @@ CADLibrary::CheckedCTError FRhinoCoretechWrapper::AddBRep(ON_Brep& Brep)
 	for (int index = 0; index < FaceCount; index++)
 	{
 		const ON_BrepFace& on_face = Brep.m_F[index];
-		CreateCTFace(Brep, on_face, FaceList);
+		bool bIsFaceHasLoopNonManifold = false;
+		CreateCTFace(Brep, on_face, FaceList, bIsFaceHasLoopNonManifold);
+		bOutBRepHasLoopNonManifold |= bIsFaceHasLoopNonManifold;
 	}
 
 	if (FaceList.IsEmpty())
diff --git a/Engine/Plugins/Enterprise/DatasmithCADImporter/Source/DatasmithOpenNurbsTranslator/Private/RhinoCoretechWrapper.h b/Engine/Plugins/Enterprise/DatasmithCADImporter/Source/DatasmithOpenNurbsTranslator/Private/RhinoCoretechWrapper.h
--- a/Engine/Plugins/Enterprise/DatasmithCADImporter/Source/DatasmithOpenNurbsTranslator/Private/RhinoCoretechWrapper.h
+++ b/Engine/Plugins/Enterprise/DatasmithCADImporter/Source/DatasmithOpenNurbsTranslator/Private/RhinoCoretechWrapper.h
@@ -25,7 +25,7 @@ public:
 	{
 	}
 
-	CADLibrary::CheckedCTError AddBRep(ON_Brep& brep);
+	CADLibrary::CheckedCTError AddBRep(ON_Brep& brep, bool& bOutBRepHasLoopNonManifold);
 	static TSharedPtr<FRhinoCoretechWrapper> GetSharedSession(double SceneUnit, double ScaleFactor);
 
 	CT_IO_ERROR Tessellate(FMeshDescription& Mesh, CADLibrary::FMeshParameters& MeshParameters);
