datasmith freezing on Alias Wire import
The problem was a sew problem with a part with small engraving features (0.2 mm).
A bad defined unit (mm instead of cm)give to the engraving the size of the geometric tolerance. The sew geometric factor (100x) made a geometric tolerance 10 time bigger than the features.
In order to identify body that freeze, the name of ct file start by an index. the last file is the freezed file

diff --git a/Engine/Plugins/Enterprise/DatasmithCADImporter/Source/CADInterfaces/Private/CTSession.cpp b/Engine/Plugins/Enterprise/DatasmithCADImporter/Source/CADInterfaces/Private/CTSession.cpp
--- a/Engine/Plugins/Enterprise/DatasmithCADImporter/Source/CADInterfaces/Private/CTSession.cpp
+++ b/Engine/Plugins/Enterprise/DatasmithCADImporter/Source/CADInterfaces/Private/CTSession.cpp
@@ -36,9 +36,9 @@ namespace CADLibrary
 		return CTKIO_CleanBody(MainObjectId);
 	}
 
-	CheckedCTError CTSession::TopoFixes()
+	CheckedCTError CTSession::TopoFixes(double SewingToleranceFactor)
 	{
-		return CADLibrary::Repair(MainObjectId, ImportParams.StitchingTechnique);
+		return CADLibrary::Repair(MainObjectId, ImportParams.StitchingTechnique, SewingToleranceFactor);
 	}
 
 	void CheckedCTError::Validate()
diff --git a/Engine/Plugins/Enterprise/DatasmithCADImporter/Source/CADInterfaces/Private/CoreTechTypes.cpp b/Engine/Plugins/Enterprise/DatasmithCADImporter/Source/CADInterfaces/Private/CoreTechTypes.cpp
--- a/Engine/Plugins/Enterprise/DatasmithCADImporter/Source/CADInterfaces/Private/CoreTechTypes.cpp
+++ b/Engine/Plugins/Enterprise/DatasmithCADImporter/Source/CADInterfaces/Private/CoreTechTypes.cpp
@@ -92,21 +92,20 @@ namespace CADLibrary
 		return CT_KERNEL_IO::AskTesselationParameters(max_face_sag, max_face_length, max_face_angle, max_curve_sag, max_curve_length, max_curve_angle, high_quality, vertex_type, normal_type, texture_type);
 	}
 
-	CT_IO_ERROR Repair(CT_OBJECT_ID MainObjectID, EStitchingTechnique StitchingTechnique)
+	CT_IO_ERROR Repair(CT_OBJECT_ID MainObjectID, EStitchingTechnique StitchingTechnique, CT_DOUBLE SewingToleranceFactor)
 	{
 		switch (StitchingTechnique)
 		{
 		case EStitchingTechnique::StitchingSew:
 		{
 			// Sew disconnected faces back together
-			CT_UINT32 SewingToleranceFactor = 100; // factor of 'tolerance' value
 			return CT_REPAIR_IO::Sew(MainObjectID, SewingToleranceFactor, CT_SEW_CREATE_BODIES_BY_TOPOLOGY);
 		}
 
 		case EStitchingTechnique::StitchingHeal:
 		{
 			// Heal disconnected faces back together
-			CT_UINT32 StitchingToleranceFactor = 100; // topo only if deformation=0, factor of 'tolerance' value
+			CT_UINT32 StitchingToleranceFactor = 0; // topo only if deformation=0, factor of 'tolerance' value
 			CT_LOGICAL   OnlyVisible = CT_TRUE;
 			CT_LOGICAL   MergeBeforeHealing = CT_FALSE;
 			CT_LOGICAL   RemoveTinyObjects = CT_TRUE;
diff --git a/Engine/Plugins/Enterprise/DatasmithCADImporter/Source/CADInterfaces/Public/CTSession.h b/Engine/Plugins/Enterprise/DatasmithCADImporter/Source/CADInterfaces/Public/CTSession.h
--- a/Engine/Plugins/Enterprise/DatasmithCADImporter/Source/CADInterfaces/Public/CTSession.h
+++ b/Engine/Plugins/Enterprise/DatasmithCADImporter/Source/CADInterfaces/Public/CTSession.h
@@ -35,7 +35,14 @@ public:
 	void ClearData();
 
 	CheckedCTError SaveBrep(const FString& FilePath);
-	CheckedCTError TopoFixes();
+
+
+	/**
+ 	 * This function calls, according to the chosen EStitchingTechnique, Kernel_io CT_REPAIR_IO::Sew or CT_REPAIR_IO::Heal. In case of sew, the used tolerance is 100x the geometric tolerance (SewingToleranceFactor = 100). 
+	 * With the case of UE-83379, Alias file, this value is too big (biggest than the geometric features. So Kernel_io hangs during the sew process... In the wait of more test, 100x is still the value used for CAD import except for Alias where the value of the SewingToleranceFactor is set to 1x
+	 * @param SewingToleranceFactor Factor apply to the tolerance 3D to define the sewing tolerance. 
+	 */
+	CheckedCTError TopoFixes(double SewingToleranceFactor = 100);
 
 	/**
 	 * In case of patch with cyclic boundary, a process has to be done by kernel_IO
diff --git a/Engine/Plugins/Enterprise/DatasmithCADImporter/Source/CADInterfaces/Public/CoreTechTypes.h b/Engine/Plugins/Enterprise/DatasmithCADImporter/Source/CADInterfaces/Public/CoreTechTypes.h
--- a/Engine/Plugins/Enterprise/DatasmithCADImporter/Source/CADInterfaces/Public/CoreTechTypes.h
+++ b/Engine/Plugins/Enterprise/DatasmithCADImporter/Source/CADInterfaces/Public/CoreTechTypes.h
@@ -50,6 +50,7 @@ namespace CADLibrary
 
 	/**
 	 * Mandatory: Kernel_IO has to be initialize with the final MetricUnit, set function doesn't work
+	 * This methode set Tolerance to 0.00001 m whether 0.01 mm
 	 * @param MetricUnit: Length unit express in meter
 	 */
 	CADINTERFACES_API CT_IO_ERROR CTKIO_InitializeKernel(double MetricUnit);
@@ -113,7 +114,12 @@ namespace CADLibrary
 		const CT_OBJECT_ID coordsystem = 0      
 	);
 
-	CADINTERFACES_API CT_IO_ERROR Repair(CT_OBJECT_ID MainObjectID, EStitchingTechnique StitchingTechnique);
+	/**
+	 * This function calls, according to the chosen EStitchingTechnique, Kernel_io CT_REPAIR_IO::Sew or CT_REPAIR_IO::Heal. In case of sew, the used tolerance is 100x the geometric tolerance (SewingToleranceFactor = 100).
+	 * With the case of UE-83379, Alias file, this value is too big (biggest than the geometric features. So Kernel_io hangs during the sew process... In the wait of more test, 100x is still the value used for CAD import except for Alias where the value of the SewingToleranceFactor is set to 1x
+	 * @param SewingToleranceFactor Factor apply to the tolerance 3D to define the sewing tolerance.
+	 */
+	CADINTERFACES_API CT_IO_ERROR Repair(CT_OBJECT_ID MainObjectID, EStitchingTechnique StitchingTechnique, CT_DOUBLE SewingToleranceFactor = 100.);
 	CADINTERFACES_API CT_IO_ERROR SetCoreTechTessellationState(const FImportParameters& ImportParams);
 
 
diff --git a/Engine/Plugins/Enterprise/DatasmithCADImporter/Source/DatasmithWireTranslator/Private/AliasCoretechWrapper.cpp b/Engine/Plugins/Enterprise/DatasmithCADImporter/Source/DatasmithWireTranslator/Private/AliasCoretechWrapper.cpp
--- a/Engine/Plugins/Enterprise/DatasmithCADImporter/Source/DatasmithWireTranslator/Private/AliasCoretechWrapper.cpp
+++ b/Engine/Plugins/Enterprise/DatasmithCADImporter/Source/DatasmithWireTranslator/Private/AliasCoretechWrapper.cpp
@@ -106,7 +106,7 @@ namespace AliasToCoreTechUtils
 CT_IO_ERROR FAliasCoretechWrapper::Tessellate(FMeshDescription& Mesh, FMeshParameters& MeshParameters)
 {
 	// Apply stitching if applicable
-	TopoFixes();
+	TopoFixes(1.);
 
 	// Perform tessellation
 	return CADLibrary::Tessellate(MainObjectId, ImportParams, Mesh, MeshParameters);
diff --git a/Engine/Plugins/Enterprise/DatasmithCADImporter/Source/DatasmithWireTranslator/Private/AliasCoretechWrapper.h b/Engine/Plugins/Enterprise/DatasmithCADImporter/Source/DatasmithWireTranslator/Private/AliasCoretechWrapper.h
--- a/Engine/Plugins/Enterprise/DatasmithCADImporter/Source/DatasmithWireTranslator/Private/AliasCoretechWrapper.h
+++ b/Engine/Plugins/Enterprise/DatasmithCADImporter/Source/DatasmithWireTranslator/Private/AliasCoretechWrapper.h
@@ -31,9 +31,8 @@ public:
 	 * eg. For a file in inches, arg should be 0.0254
 	 */
 	FAliasCoretechWrapper(const TCHAR* InOwner)
-		: CTSession(InOwner, 0.001, 1) 
-// We prefere to tell to kernelIO that Nurbs are in mm (default unit of kernelIO) to don't have side effect. 
-// Scale == 1 because in fact Alias work in cm so we do not need to scale mesh parameters
+		: CTSession(InOwner, 0.01, 1) 
+		// Unit for CoreTech session is set to cm, 0.01, because Wire's unit is cm. Consequently, Scale factor is set to 1.
 	{
 	}
 
diff --git a/Engine/Plugins/Enterprise/DatasmithCADImporter/Source/DatasmithWireTranslator/Private/DatasmithWireTranslator.cpp b/Engine/Plugins/Enterprise/DatasmithCADImporter/Source/DatasmithWireTranslator/Private/DatasmithWireTranslator.cpp
--- a/Engine/Plugins/Enterprise/DatasmithCADImporter/Source/DatasmithWireTranslator/Private/DatasmithWireTranslator.cpp
+++ b/Engine/Plugins/Enterprise/DatasmithCADImporter/Source/DatasmithWireTranslator/Private/DatasmithWireTranslator.cpp
@@ -1849,7 +1849,6 @@ TOptional<FMeshDescription> FWireTranslatorImpl::MeshDagNodeWithExternalMesher(T
 
 	Result = LocalSession->AddBRep(Body->ShellSet, MeshParameters.bIsSymmetric);
 
-	//const char*Name = DagNode->name();
 	FString Filename = Body->Label + TEXT(".ct");
 
 	FString FilePath = FPaths::Combine(OutputPath, Filename);
