Alias: Transforms are wrong when importing Wire file with Stitching Sew
For Alias, Stitching Sew activate a smart merge option that merge bodies with same parent, same material, and same layer. These bodies can have a different local transform. So, for these body, the local transform is directly apply on the NURBS berfore the transfer to KernelIO. So, each body will be in the same reference i.e the parent reference.
stl::vector is replaced by TArray

diff --git a/Engine/Plugins/Enterprise/DatasmithCADImporter/Source/DatasmithWireTranslator/Private/AliasCoretechWrapper.cpp b/Engine/Plugins/Enterprise/DatasmithCADImporter/Source/DatasmithWireTranslator/Private/AliasCoretechWrapper.cpp
--- a/Engine/Plugins/Enterprise/DatasmithCADImporter/Source/DatasmithWireTranslator/Private/AliasCoretechWrapper.cpp
+++ b/Engine/Plugins/Enterprise/DatasmithCADImporter/Source/DatasmithWireTranslator/Private/AliasCoretechWrapper.cpp
@@ -2,6 +2,8 @@
 
 #include "AliasCoretechWrapper.h"
 
+#include "Hal/PlatformMemory.h"
+
 #ifdef CAD_LIBRARY
 #include "CoreTechHelper.h"
 
@@ -10,7 +12,6 @@
 #include "AlDagNode.h"
 #include "AlShell.h"
 #include "AlShellNode.h"
-//#include "AlCurve.h"
 #include "AlSurface.h"
 #include "AlSurfaceNode.h"
 #include "AlTrimBoundary.h"
@@ -18,9 +19,6 @@
 #include "AlTrimRegion.h"
 #include "AlTM.h"
 
-#include <vector>
-
-
 TWeakPtr<FAliasCoretechWrapper> FAliasCoretechWrapper::SharedSession;
 
 struct UVPoint
@@ -42,7 +40,7 @@ struct UVPoint4
 namespace AliasToCoreTechUtils
 {
 	template<typename Surface_T>
-	CT_OBJECT_ID CreateCTNurbs(Surface_T& Surface, bool bWorldPosition)
+	CT_OBJECT_ID CreateCTNurbs(Surface_T& Surface, EAliasObjectReference InObjectReference, AlMatrix4x4& InAlMatrix)
 	{
 		CT_UINT32 ControlPointDimension = 4;  // Control Hull Dimension (3 for non-rational or 4 for rational)
 
@@ -55,35 +53,41 @@ namespace AliasToCoreTechUtils
 		CT_UINT32 KnotSizeU = Surface.realuNumberOfKnots() + 2;
 		CT_UINT32 KnotSizeV = Surface.realvNumberOfKnots() + 2;
 
-		std::vector <CT_DOUBLE> KnotValuesU;        // Axis knot value array
-		std::vector <CT_DOUBLE> KnotValuesV;        // Axis knot value array
-		std::vector <CT_UINT32> KnotMultiplicityU;  // Axis knot multiplicity array
-		std::vector <CT_UINT32> KnotMultiplicityV;  // Axis knot multiplicity array
+		TArray<CT_DOUBLE> KnotValuesU;        // Axis knot value array
+		TArray<CT_DOUBLE> KnotValuesV;        // Axis knot value array
+		TArray<CT_UINT32> KnotMultiplicityU;  // Axis knot multiplicity array
+		TArray<CT_UINT32> KnotMultiplicityV;  // Axis knot multiplicity array
 
-		KnotValuesU.resize(KnotSizeU);
-		KnotValuesV.resize(KnotSizeV);
+		KnotValuesU.SetNumUninitialized(KnotSizeU);
+		KnotValuesV.SetNumUninitialized(KnotSizeV);
 
-		Surface.realuKnotVector(KnotValuesU.data() + 1);
-		Surface.realvKnotVector(KnotValuesV.data() + 1);
+		Surface.realuKnotVector(KnotValuesU.GetData() + 1);
+		Surface.realvKnotVector(KnotValuesV.GetData() + 1);
 
-		KnotMultiplicityU.resize(KnotSizeU, 1);
-		KnotMultiplicityV.resize(KnotSizeV, 1);
+		KnotMultiplicityU.Init(1, KnotSizeU);
+		KnotMultiplicityV.Init(1, KnotSizeV);
 
 		KnotValuesU[0] = KnotValuesU[1];
 		KnotValuesV[0] = KnotValuesV[1];
 		KnotValuesU[KnotSizeU - 1] = KnotValuesU[KnotSizeU - 2];
 		KnotValuesV[KnotSizeV - 1] = KnotValuesV[KnotSizeV - 2];
 
-		std::vector <CT_DOUBLE> ControlPoint;
+		TArray<CT_DOUBLE> ControlPoint;
 		int NbCoords = ControlPointSizeU * ControlPointSizeV * 4;
-		ControlPoint.resize(NbCoords);
-		if(bWorldPosition)
+		ControlPoint.SetNumUninitialized(NbCoords);
+
+		if (InObjectReference == EAliasObjectReference::WorldReference)
 		{
-			Surface.CVsWorldPositionInclMultiples(ControlPoint.data());
+			Surface.CVsWorldPositionInclMultiples(ControlPoint.GetData());
 		}
-		else 
+		else if (InObjectReference == EAliasObjectReference::ParentReference)
 		{
-			Surface.CVsUnaffectedPositionInclMultiples(ControlPoint.data());
+			AlTM TranformMatrix(InAlMatrix);
+			Surface.CVsAffectedPositionInclMultiples(TranformMatrix, ControlPoint.GetData());
+		}
+		else  // EAliasObjectReference::LocalReference
+		{
+			Surface.CVsUnaffectedPositionInclMultiples(ControlPoint.GetData());
 		}
 
 		CT_OBJECT_ID CTSurfaceID = 0;
@@ -91,9 +95,9 @@ namespace AliasToCoreTechUtils
 			OrderU, OrderV,
 			KnotSizeU, KnotSizeV,
 			ControlPointSizeU, ControlPointSizeV,
-			4, ControlPoint.data(),
-			KnotValuesU.data(), KnotValuesV.data(),
-			KnotMultiplicityU.data(), KnotMultiplicityV.data()
+			4, ControlPoint.GetData(),
+			KnotValuesU.GetData(), KnotValuesV.GetData(),
+			KnotMultiplicityU.GetData(), KnotMultiplicityV.GetData()
 		);
 		if (Result != IO_OK)
 		{
@@ -118,14 +122,14 @@ CT_OBJECT_ID FAliasCoretechWrapper::Add3DCurve(AlCurve& Curve)
 	CT_UINT32 Order = Curve.degree() + 1;
 	CT_UINT32 ControlPointSize = Curve.numberOfCVs();
 	CT_UINT32 RealKnotSize = Curve.realNumberOfKnots() + 2;
-	std::vector <UVPoint4> ControlPoints;
+	TArray<UVPoint4> ControlPoints;
 
-	ControlPoints.resize(ControlPointSize);
+	ControlPoints.SetNumUninitialized(ControlPointSize);
 	using UVPoint14 = double[4];
-	Curve.CVsUnaffectedPositionInclMultiples((UVPoint14*)ControlPoints.data());
+	Curve.CVsUnaffectedPositionInclMultiples((UVPoint14*)ControlPoints.GetData());
 
-	std::vector <double> CTControlPoints;
-	CTControlPoints.resize(ControlPointSize * 4);
+	TArray<double> CTControlPoints;
+	CTControlPoints.SetNumUninitialized(ControlPointSize * 4);
 	int Index = 0;
 	for (CT_UINT32 IndexPoint = 0; IndexPoint < ControlPointSize; ++IndexPoint)
 	{
@@ -135,12 +139,12 @@ CT_OBJECT_ID FAliasCoretechWrapper::Add3DCurve(AlCurve& Curve)
 		CTControlPoints[Index++] = ControlPoints[IndexPoint].w;
 	}
 
-	std::vector<CT_UINT32> KnotMult;
-	KnotMult.resize(RealKnotSize, 1);
+	TArray<CT_UINT32> KnotMult;
+	KnotMult.Init(1, RealKnotSize);
 
-	std::vector <CT_DOUBLE> Knots;
-	Knots.resize(RealKnotSize);
-	Curve.realKnotVector(Knots.data() + 1);
+	TArray<CT_DOUBLE> Knots;
+	Knots.SetNumUninitialized(RealKnotSize);
+	Curve.realKnotVector(Knots.GetData() + 1);
 	Knots[0] = Knots[1];
 	Knots[RealKnotSize - 1] = Knots[RealKnotSize - 2];
 
@@ -151,9 +155,9 @@ CT_OBJECT_ID FAliasCoretechWrapper::Add3DCurve(AlCurve& Curve)
 		RealKnotSize,           /*!< [in] Knot vector size */
 		ControlPointSize,       /*!< [in] Control Hull Size */
 		4,                      /*!< [in] Control Hull Dimension (2 for non-rational or 3 for rational) */
-		CTControlPoints.data(), /*!< [in] Control hull array */
-		Knots.data(),           /*!< [in] Value knot array */
-		KnotMult.data(),        /*!< [in] Multiplicity Knot array */
+		CTControlPoints.GetData(), /*!< [in] Control hull array */
+		Knots.GetData(),           /*!< [in] Value knot array */
+		KnotMult.GetData(),        /*!< [in] Multiplicity Knot array */
 		Knots[0],               /*!< [in] start parameter of coedge on the uv curve (t range=[knot[0], knot[knot_size-1]]) */
 		Knots[RealKnotSize - 1]   /*!< [in] end parameter of coedge on the uv curve (t range=[knot[0], knot[knot_size-1]]) */
 	);
@@ -173,42 +177,24 @@ CT_OBJECT_ID FAliasCoretechWrapper::AddTrimCurve(AlTrimCurve& TrimCurve)
 	CT_UINT32 Order = TrimCurve.degree() + 1;
 	CT_UINT32 ControlPointSize = TrimCurve.numberOfCVs();
 	CT_UINT32 RealKnotSize = TrimCurve.realNumberOfKnots() + 2;
-	std::vector <UVPoint> ControlPoints;
-	std::vector <double> Weigths;
+	TArray<UVPoint> ControlPoints;
+	TArray<double> Weigths;
 
-	CT_OBJECT_ID D3CurveId = 0;
-	/*
-    AlCurve * D3Curve = TrimCurve.unaffected3DCopyNoReverse();
-	if (D3Curve)
-	{
-		D3CurveId = Add3DCurve(*D3Curve);
-	}
-    */
-
-	ControlPoints.resize(ControlPointSize);
-	Weigths.resize(ControlPointSize);
+	ControlPoints.SetNumUninitialized(ControlPointSize);
+	Weigths.SetNumUninitialized(ControlPointSize);
 	using UVPoint2 = double[3];
-	TrimCurve.CVsUVPosition(Weigths.data(), (UVPoint2*)ControlPoints.data());
+	TrimCurve.CVsUVPosition(Weigths.GetData(), (UVPoint2*)ControlPoints.GetData());
 
-	std::vector <double> CTControlPoints;
-	CTControlPoints.resize(ControlPointSize * 3);
-	int Index = 0;
-	for (CT_UINT32 IndexPoint = 0; IndexPoint < ControlPointSize; ++IndexPoint)
-	{
-		CTControlPoints[Index++] = ControlPoints[IndexPoint].x;
-		CTControlPoints[Index++] = ControlPoints[IndexPoint].y;
-		CTControlPoints[Index++] = ControlPoints[IndexPoint].z;
-	}
-
-	std::vector<CT_UINT32> KnotMult;
-	KnotMult.resize(RealKnotSize, 1);
+	TArray<CT_UINT32> KnotMult;
+	KnotMult.Init(1, RealKnotSize);
 
-	std::vector <CT_DOUBLE> Knots;
-	Knots.resize(RealKnotSize);
-	TrimCurve.realKnotVector(Knots.data() + 1);
+	TArray<CT_DOUBLE> Knots;
+	Knots.SetNumUninitialized(RealKnotSize);
+	TrimCurve.realKnotVector(Knots.GetData() + 1);
 	Knots[0] = Knots[1];
 	Knots[RealKnotSize - 1] = Knots[RealKnotSize - 2];
 
+	CT_OBJECT_ID D3CurveId = 0;
 	CT_OBJECT_ID CoedgeID;
 	CT_IO_ERROR Result = CT_COEDGE_IO::Create(CoedgeID, bIsReversed ? CT_ORIENTATION::CT_REVERSE : CT_ORIENTATION::CT_FORWARD, D3CurveId);
 	if (Result != IO_OK)
@@ -222,9 +208,9 @@ CT_OBJECT_ID FAliasCoretechWrapper::AddTrimCurve(AlTrimCurve& TrimCurve)
 		RealKnotSize,           /*!< [in] Knot vector size */
 		ControlPointSize,       /*!< [in] Control Hull Size */
 		3,                      /*!< [in] Control Hull Dimension (2 for non-rational or 3 for rational) */
-		CTControlPoints.data(), /*!< [in] Control hull array */
-		Knots.data(),           /*!< [in] Value knot array */
-		KnotMult.data(),        /*!< [in] Multiplicity Knot array */
+		(double*) ControlPoints.GetData(), /*!< [in] Control hull array */
+		Knots.GetData(),           /*!< [in] Value knot array */
+		KnotMult.GetData(),        /*!< [in] Multiplicity Knot array */
 		Knots[0],               /*!< [in] start parameter of coedge on the uv curve (t range=[knot[0], knot[knot_size-1]]) */
 		Knots[RealKnotSize - 1]   /*!< [in] end parameter of coedge on the uv curve (t range=[knot[0], knot[knot_size-1]]) */
 	);
@@ -271,9 +257,9 @@ CT_OBJECT_ID FAliasCoretechWrapper::AddTrimBoundary(AlTrimBoundary& TrimBoundary
 	return LoopId;
 }
 
-CT_OBJECT_ID FAliasCoretechWrapper::AddTrimRegion(AlTrimRegion& TrimRegion, bool bWorldPosition, bool bOrientation)
+CT_OBJECT_ID FAliasCoretechWrapper::AddTrimRegion(AlTrimRegion& TrimRegion, EAliasObjectReference InObjectReference, AlMatrix4x4& InAlMatrix, bool bInOrientation)
 {
-	CT_OBJECT_ID NurbsId = AliasToCoreTechUtils::CreateCTNurbs(TrimRegion, bWorldPosition);
+	CT_OBJECT_ID NurbsId = AliasToCoreTechUtils::CreateCTNurbs(TrimRegion, InObjectReference, InAlMatrix);
 	if (NurbsId == 0)
 	{
 		return 0;
@@ -281,7 +267,7 @@ CT_OBJECT_ID FAliasCoretechWrapper::AddTrimRegion(AlTrimRegion& TrimRegion, bool
 
 	CT_LIST_IO Boundaries;
 	CT_OBJECT_ID FaceID;
-	CT_ORIENTATION FaceOrient = bOrientation ? CT_ORIENTATION::CT_FORWARD : CT_ORIENTATION::CT_REVERSE;
+	CT_ORIENTATION FaceOrient = bInOrientation ? CT_ORIENTATION::CT_FORWARD : CT_ORIENTATION::CT_REVERSE;
 
 	AlTrimBoundary *TrimBoundary = TrimRegion.firstBoundary();
 	while (TrimBoundary)
@@ -301,49 +287,49 @@ CT_OBJECT_ID FAliasCoretechWrapper::AddTrimRegion(AlTrimRegion& TrimRegion, bool
 	return 0;
 }
 
-void FAliasCoretechWrapper::AddFace(AlSurface& Surface, CT_LIST_IO& FaceList, bool bWorldPosition, bool bOrientation)
+void FAliasCoretechWrapper::AddFace(AlSurface& Surface, EAliasObjectReference InObjectReference, AlMatrix4x4& InAlMatrix, bool bInOrientation, CT_LIST_IO& OutFaceList)
 {
 	AlTrimRegion *TrimRegion = Surface.firstTrimRegion();
 	if (TrimRegion)
 	{
 		while (TrimRegion)
 		{
-			CT_OBJECT_ID FaceID = AddTrimRegion(*TrimRegion, bWorldPosition, bOrientation);
+			CT_OBJECT_ID FaceID = AddTrimRegion(*TrimRegion, InObjectReference, InAlMatrix, bInOrientation);
 			if (FaceID)
 			{
-				FaceList.PushBack(FaceID);
+				OutFaceList.PushBack(FaceID);
 			}
 			TrimRegion = TrimRegion->nextRegion();
 		}
 		return;
 	}
 
-	CT_OBJECT_ID NurbsId = AliasToCoreTechUtils::CreateCTNurbs(Surface, bWorldPosition);
+	CT_OBJECT_ID NurbsId = AliasToCoreTechUtils::CreateCTNurbs(Surface, InObjectReference, InAlMatrix);
 	CT_LIST_IO Boundaries;
 	CT_OBJECT_ID FaceID;
-	CT_ORIENTATION FaceOrient = bOrientation ? CT_ORIENTATION::CT_FORWARD : CT_ORIENTATION::CT_REVERSE;
+	CT_ORIENTATION FaceOrient = bInOrientation ? CT_ORIENTATION::CT_FORWARD : CT_ORIENTATION::CT_REVERSE;
 	CT_IO_ERROR Result = CT_FACE_IO::Create(FaceID, NurbsId, FaceOrient, Boundaries);
 	if (Result == IO_OK)
 	{
-		FaceList.PushBack(FaceID);
+		OutFaceList.PushBack(FaceID);
 	}
 }
 
-void FAliasCoretechWrapper::AddShell(AlShell& Shell, CT_LIST_IO& FaceList, bool bIsSymmetricBody, bool bOrientation)
+void FAliasCoretechWrapper::AddShell(AlShell& Shell, EAliasObjectReference InObjectReference, AlMatrix4x4& InAlMatrix, bool bInOrientation, CT_LIST_IO& OutFaceLis)
 {
 	AlTrimRegion *TrimRegion = Shell.firstTrimRegion();
 	while (TrimRegion)
 	{
-		CT_OBJECT_ID FaceID = AddTrimRegion(*TrimRegion, bIsSymmetricBody, bOrientation);
+		CT_OBJECT_ID FaceID = AddTrimRegion(*TrimRegion, InObjectReference, InAlMatrix, bInOrientation);
 		if (FaceID)
 		{
-			FaceList.PushBack(FaceID);
+			OutFaceLis.PushBack(FaceID);
 		}
 		TrimRegion = TrimRegion->nextRegion();
 	}
 }
 
-CT_IO_ERROR FAliasCoretechWrapper::AddBRep(TArray<AlDagNode*>& DagNodeSet, bool bWorldPosition)
+CT_IO_ERROR FAliasCoretechWrapper::AddBRep(TArray<AlDagNode*>& InDagNodeSet, EAliasObjectReference InObjectReference)
 {
 	CT_IO_ERROR Result = IO_OK;
 	if (!IsSessionValid())
@@ -351,16 +337,11 @@ CT_IO_ERROR FAliasCoretechWrapper::AddBRep(TArray<AlDagNode*>& DagNodeSet, bool
 		return IO_ERROR;
 	}
 
-	if (ImportParams.StitchingTechnique == StitchingSew)
-	{
-		bWorldPosition = true;
-	}
-
 	CT_LIST_IO FaceList;
 
 	AlEdge2CTEdge.Empty();
 
-	for (auto DagNode : DagNodeSet)
+	for (auto DagNode : InDagNodeSet)
 	{
 		if (DagNode == nullptr)
 		{
@@ -370,6 +351,12 @@ CT_IO_ERROR FAliasCoretechWrapper::AddBRep(TArray<AlDagNode*>& DagNodeSet, bool
 		DagNode->getSurfaceOrientation(bAlOrientation);
 		bool bOrientation = (bool) bAlOrientation;
 
+		AlMatrix4x4 AlMatrix;
+		if (InObjectReference == EAliasObjectReference::ParentReference)
+		{
+			DagNode->localTransformationMatrix(AlMatrix);
+		}
+
 		AlObjectType objectType = DagNode->type();
 		switch (objectType)
 		{
@@ -384,7 +371,7 @@ CT_IO_ERROR FAliasCoretechWrapper::AddBRep(TArray<AlDagNode*>& DagNodeSet, bool
 				AlShell *Shell = ShellNode->shell();
 				if (Shell)
 				{
-					AddShell(*Shell, FaceList, bWorldPosition, bOrientation);
+					AddShell(*Shell, InObjectReference, AlMatrix, bOrientation, FaceList);
 				}
 			}
 			break;
@@ -398,7 +385,7 @@ CT_IO_ERROR FAliasCoretechWrapper::AddBRep(TArray<AlDagNode*>& DagNodeSet, bool
 				AlSurface *Surface = SurfaceNode->surface();
 				if (Surface)
 				{
-					AddFace(*Surface, FaceList, bWorldPosition, bOrientation);
+					AddFace(*Surface, InObjectReference, AlMatrix, bOrientation, FaceList);
 				}
 			}
 			break;
diff --git a/Engine/Plugins/Enterprise/DatasmithCADImporter/Source/DatasmithWireTranslator/Private/AliasCoretechWrapper.h b/Engine/Plugins/Enterprise/DatasmithCADImporter/Source/DatasmithWireTranslator/Private/AliasCoretechWrapper.h
--- a/Engine/Plugins/Enterprise/DatasmithCADImporter/Source/DatasmithWireTranslator/Private/AliasCoretechWrapper.h
+++ b/Engine/Plugins/Enterprise/DatasmithCADImporter/Source/DatasmithWireTranslator/Private/AliasCoretechWrapper.h
@@ -20,6 +20,15 @@ typedef double AlMatrix4x4[4][4];
 
 using namespace CADLibrary;
 
+// Defined the reference in which the object has to be defined
+enum class EAliasObjectReference
+{
+	LocalReference,  
+	ParentReference, 
+	WorldReference, 
+};
+
+
 class FAliasCoretechWrapper : public CTSession
 {
 public:
@@ -36,7 +45,7 @@ public:
 	{
 	}
 
-	CT_IO_ERROR AddBRep(TArray<AlDagNode*>& DagNodeSet, bool bIsSymmetricBody);
+	CT_IO_ERROR AddBRep(TArray<AlDagNode*>& DagNodeSet, EAliasObjectReference ObjectReference);
 
 	static TSharedPtr<FAliasCoretechWrapper> GetSharedSession();
 
@@ -51,9 +60,9 @@ protected:
 	CT_OBJECT_ID AddTrimBoundary(AlTrimBoundary& TrimBoundary);
 	CT_OBJECT_ID Add3DCurve(AlCurve& Curve);
 
-	CT_OBJECT_ID AddTrimRegion(AlTrimRegion& TrimRegion, bool bIsSymmetricBody, bool bOrientation);
-	void AddFace(AlSurface& Surface, CT_LIST_IO& FaceLis, bool bIsSymmetricBodyt, bool bOrientation);
-	void AddShell(AlShell& Shell, CT_LIST_IO& FaceList, bool bIsSymmetricBody, bool bOrientation);
+	CT_OBJECT_ID AddTrimRegion(AlTrimRegion& InTrimRegion, EAliasObjectReference InObjectReference, AlMatrix4x4& InAlMatrix, bool bInOrientation);
+	void AddFace(AlSurface& InSurface, EAliasObjectReference InObjectReference, AlMatrix4x4& InAlMatrix, bool bInOrientation, CT_LIST_IO& OutFaceLis);
+	void AddShell(AlShell& InShell, EAliasObjectReference InObjectReference, AlMatrix4x4& InAlMatrix, bool bInOrientation, CT_LIST_IO& OutFaceLis);
 
 protected:
 	static TWeakPtr<FAliasCoretechWrapper> SharedSession;
diff --git a/Engine/Plugins/Enterprise/DatasmithCADImporter/Source/DatasmithWireTranslator/Private/DatasmithWireTranslator.cpp b/Engine/Plugins/Enterprise/DatasmithCADImporter/Source/DatasmithWireTranslator/Private/DatasmithWireTranslator.cpp
--- a/Engine/Plugins/Enterprise/DatasmithCADImporter/Source/DatasmithWireTranslator/Private/DatasmithWireTranslator.cpp
+++ b/Engine/Plugins/Enterprise/DatasmithCADImporter/Source/DatasmithWireTranslator/Private/DatasmithWireTranslator.cpp
@@ -150,7 +150,6 @@ public:
 	void SetTessellationOptions(const FDatasmithTessellationOptions& Options);
 	void SetOutputPath(const FString& Path) { OutputPath = Path; }
 
-	//double GetMetricUnit() const { return LocalSession->GetImportParameters().MetricUnit; }
 	CADLibrary::FImportParameters& GetImportParameters()
 	{
 		return LocalSession->GetImportParameters();
@@ -193,7 +192,6 @@ private:
 
  	TOptional< FMeshDescription > ImportMesh(AlMesh& Mesh, CADLibrary::FMeshParameters& MeshParameters);
 
-	void CreateAlCommonMaterial(AlShader *Shader, TSharedRef<IDatasmithUEPbrMaterialElement> MaterialElement);
 	void AddAlBlinnParameters(AlShader *Shader, TSharedRef<IDatasmithUEPbrMaterialElement> MaterialElement);
 	void AddAlLambertParameters(AlShader *Shader, TSharedRef<IDatasmithUEPbrMaterialElement> MaterialElement);
 	void AddAlLightSourceParameters(AlShader *Shader, TSharedRef<IDatasmithUEPbrMaterialElement> MaterialElement);
@@ -1820,9 +1818,17 @@ TOptional<FMeshDescription> FWireTranslatorImpl::MeshDagNodeWithExternalMesher(A
 
 	FString Filename = DagNode.name();
 
+	EAliasObjectReference ObjectReference = EAliasObjectReference::LocalReference;
+
+	if (MeshParameters.bIsSymmetric)
+	{
+		// All actors of a Alias symmetric layer are defined in the world Reference i.e. they have identity transform. So Mesh actor has to be defined in the world reference. 
+		ObjectReference = EAliasObjectReference::WorldReference;
+	}
+
 	TArray<AlDagNode*> DagNodeSet;
 	DagNodeSet.Add(&DagNode);
-	Result = LocalSession->AddBRep(DagNodeSet, MeshParameters.bIsSymmetric);
+	Result = LocalSession->AddBRep(DagNodeSet, ObjectReference);
 
 	Filename += TEXT(".ct");
 
@@ -1847,7 +1853,19 @@ TOptional<FMeshDescription> FWireTranslatorImpl::MeshDagNodeWithExternalMesher(T
 
 	LocalSession->ClearData();
 
-	Result = LocalSession->AddBRep(Body->ShellSet, MeshParameters.bIsSymmetric);
+	EAliasObjectReference ObjectReference = EAliasObjectReference::LocalReference;
+	if (MeshParameters.bIsSymmetric)
+	{
+		// All actors of a Alias symmetric layer are defined in the world Reference i.e. they have identity transform. So Mesh actor has to be defined in the world reference. 
+		ObjectReference = EAliasObjectReference::WorldReference;
+	}
+	else if (GetImportParameters().StitchingTechnique == StitchingSew)
+	{
+		// In the case of StitchingSew, AlDagNode children of a GroupNode are merged together. To be merged, they have to be defined in the reference of parent GroupNode.
+		ObjectReference = EAliasObjectReference::ParentReference;
+	}
+
+	Result = LocalSession->AddBRep(Body->ShellSet, ObjectReference);
 
 	FString Filename = Body->Label + TEXT(".ct");
 
