Rename MasterMaterial

diff --git a/Engine/Plugins/Enterprise/DatasmithCADImporter/Source/CADLibrary/Private/CoreTechHelper.cpp b/Engine/Plugins/Enterprise/DatasmithCADImporter/Source/CADLibrary/Private/CoreTechHelper.cpp
--- a/Engine/Plugins/Enterprise/DatasmithCADImporter/Source/CADLibrary/Private/CoreTechHelper.cpp
+++ b/Engine/Plugins/Enterprise/DatasmithCADImporter/Source/CADLibrary/Private/CoreTechHelper.cpp
@@ -647,7 +647,7 @@ TSharedPtr<IDatasmithUEPbrMaterialElement> CreateDefaultUEPbrMaterial()
 	ColorExpression->SetName(TEXT("Base Color"));
 	ColorExpression->GetColor() = LinearColor;
 	MaterialElement->GetBaseColor().SetExpression(ColorExpression);
-	MaterialElement->SetParentLabel(TEXT("M_CAD_Material"));
+	MaterialElement->SetParentLabel(TEXT("M_DatasmithCAD"));
 
 	return MaterialElement;
 }
@@ -678,11 +678,11 @@ TSharedPtr<IDatasmithUEPbrMaterialElement> CreateUEPbrMaterialFromColor(const FC
 		Scalar->SetName(TEXT("Opacity Level"));
 
 		MaterialElement->GetOpacity().SetExpression(Scalar);
-		MaterialElement->SetParentLabel(TEXT("M_CAD_Transparent_Material"));
+		MaterialElement->SetParentLabel(TEXT("M_DatasmithCADTransparent"));
 	}
 	else
 	{
-		MaterialElement->SetParentLabel(TEXT("M_CAD_Material"));
+		MaterialElement->SetParentLabel(TEXT("M_DatasmithCAD"));
 	}
 
 	return MaterialElement;
@@ -720,11 +720,11 @@ TSharedPtr<IDatasmithUEPbrMaterialElement> CreateUEPbrMaterialFromMaterial(FCADM
 		Scalar->GetScalar() = InMaterial.Transparency;
 		Scalar->SetName(TEXT("Opacity Level"));
 		MaterialElement->GetOpacity().SetExpression(Scalar);
-		MaterialElement->SetParentLabel(TEXT("M_CAD_Transparent_Material"));
+		MaterialElement->SetParentLabel(TEXT("M_DatasmithCADTransparent"));
 	}
 	else
 	{
-		MaterialElement->SetParentLabel(TEXT("M_CAD_Material"));
+		MaterialElement->SetParentLabel(TEXT("M_DatasmithCAD"));
 	}
 
 	return MaterialElement;
diff --git a/Engine/Plugins/Enterprise/DatasmithCADImporter/Source/DatasmithWireTranslator/Private/DatasmithWireTranslator.cpp b/Engine/Plugins/Enterprise/DatasmithCADImporter/Source/DatasmithWireTranslator/Private/DatasmithWireTranslator.cpp
--- a/Engine/Plugins/Enterprise/DatasmithCADImporter/Source/DatasmithWireTranslator/Private/DatasmithWireTranslator.cpp
+++ b/Engine/Plugins/Enterprise/DatasmithCADImporter/Source/DatasmithWireTranslator/Private/DatasmithWireTranslator.cpp
@@ -576,12 +576,17 @@ void FWireTranslatorImpl::AddAlBlinnParameters(AlShader *Shader, TSharedRef<IDat
 	MaterialElement->GetSpecular().SetExpression(FresnelMultiply);
 	MaterialElement->GetRoughness().SetExpression(RoughnessOneMinus);
 	MaterialElement->GetEmissiveColor().SetExpression(IncandescenceMultiply);
+
 	if (bIsTransparent)
 	{
 		MaterialElement->GetOpacity().SetExpression(Divide);
+		MaterialElement->SetParentLabel(TEXT("M_DatasmithAliasBlinnTransparent"));
+	}
+	else 
+	{
+		MaterialElement->SetParentLabel(TEXT("M_DatasmithAliasBlinn"));
 	}
 
-	MaterialElement->SetParentLabel(TEXT("BLINN"));
 }
 
 void FWireTranslatorImpl::AddAlLambertParameters(AlShader *Shader, TSharedRef<IDatasmithUEPbrMaterialElement> MaterialElement)
@@ -763,9 +768,12 @@ void FWireTranslatorImpl::AddAlLambertParameters(AlShader *Shader, TSharedRef<ID
 	if (bIsTransparent)
 	{
 		MaterialElement->GetOpacity().SetExpression(Divide);
+		MaterialElement->SetParentLabel(TEXT("M_DatasmithAliasLambertTransparent"));
+	}
+	else {
+		MaterialElement->SetParentLabel(TEXT("M_DatasmithAliasLambert"));
 	}
 
-	MaterialElement->SetParentLabel(TEXT("LAMBERT"));
 }
 
 void FWireTranslatorImpl::AddAlLightSourceParameters(AlShader *Shader, TSharedRef<IDatasmithUEPbrMaterialElement> MaterialElement)
@@ -917,12 +925,16 @@ void FWireTranslatorImpl::AddAlLightSourceParameters(AlShader *Shader, TSharedRe
 	// Connect material outputs
 	MaterialElement->GetBaseColor().SetExpression(BaseColorTransparencyMultiply);
 	MaterialElement->GetEmissiveColor().SetExpression(IncandescenceMultiply);
+
 	if (bIsTransparent)
 	{
 		MaterialElement->GetOpacity().SetExpression(Divide);
+		MaterialElement->SetParentLabel(TEXT("M_DatasmithAliasLightSourceTransparent"));
+	}
+	else {
+		MaterialElement->SetParentLabel(TEXT("M_DatasmithAliasLightSource"));
 	}
 
-	MaterialElement->SetParentLabel(TEXT("LIGHTSOURCE"));
 }
 
 void FWireTranslatorImpl::AddAlPhongParameters(AlShader *Shader, TSharedRef<IDatasmithUEPbrMaterialElement> MaterialElement)
@@ -1205,9 +1217,11 @@ void FWireTranslatorImpl::AddAlPhongParameters(AlShader *Shader, TSharedRef<IDat
 	if (bIsTransparent)
 	{
 		MaterialElement->GetOpacity().SetExpression(Divide);
+		MaterialElement->SetParentLabel(TEXT("M_DatasmithAliasPhongTransparent"));
+	}
+	else {
+		MaterialElement->SetParentLabel(TEXT("M_DatasmithAliasPhong"));
 	}
-
-	MaterialElement->SetParentLabel(TEXT("PHONG"));
 }
 //#endif
 
