New DatasmithSceneGraphBuilder and DatasmithMeshBuilder based on FArchive
Huge simplication of the code
Stabilisation of CAD import

diff --git a/Engine/Plugins/Enterprise/DatasmithCADImporter/Source/CADInterfaces/Private/CADSceneGraph.cpp b/Engine/Plugins/Enterprise/DatasmithCADImporter/Source/CADInterfaces/Private/CADSceneGraph.cpp
new file mode 100644
--- /dev/null
+++ b/Engine/Plugins/Enterprise/DatasmithCADImporter/Source/CADInterfaces/Private/CADSceneGraph.cpp
@@ -0,0 +1,110 @@
+// Copyright 1998-2019 Epic Games, Inc. All Rights Reserved.
+
+#include "CADSceneGraph.h"
+
+#include "Misc/FileHelper.h"
+#include "Serialization/Archive.h"
+#include "Serialization/MemoryWriter.h"
+#include "Serialization/MemoryReader.h"
+
+namespace CADLibrary
+{
+
+FArchive& operator<<(FArchive& Ar, FArchiveInstance& Instance) 
+{
+	Ar << Instance.ObjectId;
+	Ar << Instance.MetaData;
+	Ar << Instance.TransformMatrix;
+	Ar << Instance.ReferenceNodeId;
+	Ar << Instance.bIsExternalRef;
+	Ar << Instance.ExternalRef;
+	return Ar;
+}
+
+FArchive& operator<<(FArchive& Ar, FArchiveComponent& Component)
+{
+	Ar << Component.ObjectId;
+	Ar << Component.MetaData;
+	Ar << Component.Children;
+	return Ar;
+}
+
+FArchive& operator<<(FArchive& Ar, FArchiveUnloadedComponent& Unloaded) 
+{
+	Ar << Unloaded.ObjectId;
+	Ar << Unloaded.MetaData;
+	Ar << Unloaded.FileName;
+	Ar << Unloaded.FileType;
+	return Ar;
+}
+
+FArchive& operator<<(FArchive& Ar, FArchiveBody& Body) 
+{
+	Ar << Body.ObjectId;
+	Ar << Body.MetaData;
+	Ar << Body.MaterialFaceSet;
+	Ar << Body.ColorFaceSet;
+	Ar << Body.MeshActorName;
+	return Ar;
+}
+
+FArchive& operator<<(FArchive& Ar, FArchiveColor& Color) 
+{
+	Ar << Color.ObjectId;
+	Ar << Color.Color;
+	Ar << Color.UEMaterialName;
+	return Ar;
+}
+
+FArchive& operator<<(FArchive& Ar, FArchiveMaterial& Material)
+{
+	Ar << Material.ObjectId;
+	Ar << Material.Material;
+	Ar << Material.UEMaterialName;
+	return Ar;
+}
+
+FArchive& operator<<(FArchive& Ar, FArchiveMockUp& MockUp)
+{
+	Ar << MockUp.CADFile;
+	Ar << MockUp.SceneGraphArchive;
+	Ar << MockUp.FullPath;
+	Ar << MockUp.ExternalRefSet;
+
+	Ar << MockUp.ColorHIdToColor;
+	Ar << MockUp.MaterialHIdToMaterial;
+
+	Ar << MockUp.InstanceSet;
+	Ar << MockUp.ComponentSet;
+	Ar << MockUp.UnloadedComponentSet;
+	Ar << MockUp.BodySet;
+
+	Ar << MockUp.CADIdToInstanceIndex;
+	Ar << MockUp.CADIdToComponentIndex;
+	Ar << MockUp.CADIdToUnloadedComponentIndex;
+	Ar << MockUp.CADIdToBodyIndex;
+
+	return Ar;
+}
+
+void SerializeMockUp(FArchiveMockUp& MockUp, const TCHAR* Filename)
+{
+	TArray<uint8> OutBuffer;
+	FMemoryWriter ArWriter(OutBuffer);
+	ArWriter << MockUp;
+
+	FFileHelper::SaveArrayToFile(OutBuffer, Filename);
+}
+
+void DeserializeMockUpFile(const TCHAR* Filename, FArchiveMockUp& MockUp)
+{
+	TArray<uint8> Buffer;
+	FFileHelper::LoadFileToArray(Buffer, Filename);
+
+	FMemoryReader ArReader(Buffer);
+	ArReader << MockUp;
+}
+
+}
+
+
diff --git a/Engine/Plugins/Enterprise/DatasmithCADImporter/Source/CADInterfaces/Private/CoreTechFileParser.cpp b/Engine/Plugins/Enterprise/DatasmithCADImporter/Source/CADInterfaces/Private/CoreTechFileParser.cpp
--- a/Engine/Plugins/Enterprise/DatasmithCADImporter/Source/CADInterfaces/Private/CoreTechFileParser.cpp
+++ b/Engine/Plugins/Enterprise/DatasmithCADImporter/Source/CADInterfaces/Private/CoreTechFileParser.cpp
@@ -181,9 +181,7 @@ uint32 GetGeomFileHash(const uint32 InSGHash, const FImportParameters& ImportPar
 	return FileHash;
 }
 
-
-
-uint32 GetFaceTessellation(CT_OBJECT_ID FaceID, TArray<FTessellationData>& FaceTessellationSet, int32& OutRawDataSize, const FImportParameters& ImportParams)
+uint32 GetFaceTessellation(CT_OBJECT_ID FaceID, TArray<FTessellationData>& FaceTessellationSet, const FImportParameters& ImportParams)
 {
 	CT_IO_ERROR Error = IO_OK;
 
@@ -211,12 +209,10 @@ uint32 GetFaceTessellation(CT_OBJECT_ID FaceID, TArray<FTessellationData>& FaceT
 	if (Error != IO_OK || VertexArray == nullptr || IndexArray == nullptr || IndexCount == 0)
 	{
 		VertexCount = NormalCount = IndexCount = 0;
-		OutRawDataSize = 0;
 		return 0;
 	}
 
-	int32 Index = FaceTessellationSet.Emplace();
-	FTessellationData& Tessellation = FaceTessellationSet[Index];
+	FTessellationData& Tessellation = FaceTessellationSet.Emplace_GetRef();
 	if (ImportParams.bScaleUVMap && TexCoordArray != nullptr)
 	{
 		switch (TexCoordType)
@@ -239,11 +235,6 @@ uint32 GetFaceTessellation(CT_OBJECT_ID FaceID, TArray<FTessellationData>& FaceT
 	Tessellation.SizeOfNormalType = GetSize(NormalType);
 	Tessellation.SizeOfIndexType = GetSize(IndexType);
 
-	Tessellation.VertexArray.Empty(Tessellation.VertexCount * Tessellation.SizeOfVertexType);
-	Tessellation.NormalArray.Empty(Tessellation.NormalCount * Tessellation.SizeOfNormalType);
-	Tessellation.IndexArray.Empty(Tessellation.IndexCount * Tessellation.SizeOfIndexType);
-	Tessellation.TexCoordArray.Empty(Tessellation.TexCoordCount * Tessellation.SizeOfTexCoordType);
-
 	Tessellation.VertexArray.Append((uint8*) VertexArray, 3 * Tessellation.VertexCount * Tessellation.SizeOfVertexType);
 	Tessellation.NormalArray.Append((uint8*)NormalArray, 3 * Tessellation.NormalCount * Tessellation.SizeOfNormalType);
 	Tessellation.IndexArray.Append((uint8*)IndexArray, Tessellation.IndexCount * Tessellation.SizeOfIndexType);
@@ -253,8 +244,6 @@ uint32 GetFaceTessellation(CT_OBJECT_ID FaceID, TArray<FTessellationData>& FaceT
 		Tessellation.TexCoordArray.Append((uint8*)TexCoordArray, 2 * Tessellation.TexCoordCount * Tessellation.SizeOfTexCoordType);
 	}
 
-	OutRawDataSize = Tessellation.VertexArray.Num() + Tessellation.NormalArray.Num() + Tessellation.IndexArray.Num() + Tessellation.TexCoordArray.Num();
-
 	return (uint32)Tessellation.IndexCount / 3;
 }
 
@@ -266,7 +255,7 @@ void GetCTObjectDisplayDataIds(CT_OBJECT_ID ObjectID, FObjectDisplayDataId& Mate
 		CT_UINT32 MaterialId = 0;
 		if (CT_CURRENT_ATTRIB_IO::AskIntField(ITH_MATERIALID_VALUE, MaterialId) == IO_OK && MaterialId > 0)
 		{
-			Material.MaterialId = (uint32)MaterialId;
+			Material.Material = (uint32)MaterialId;
 		}
 	}
 
@@ -284,62 +273,72 @@ void GetCTObjectDisplayDataIds(CT_OBJECT_ID ObjectID, FObjectDisplayDataId& Mate
 					alpha = (uint8)int((1. - dbl_value) * 255.);
 				}
 			}
-			Material.ColorHId = BuildFastColorHash(ColorId, alpha);
+			Material.Color = BuildColorId(ColorId, alpha);
 		}
 	}
 }
 
-int32 BuildColorHash(CT_COLOR CtColor, uint8 Alpha)
+FArchiveMaterial& FCoreTechFileParser::FindOrAddMaterial(CT_MATERIAL_ID MaterialId)
 {
-	FColor Color(CtColor[0], CtColor[1], CtColor[2], Alpha);
-	return BuildColorHash(Color);
+	if (FArchiveMaterial* NewMaterial = MockUpDescription.MaterialHIdToMaterial.Find(MaterialId))
+	{
+		return *NewMaterial;
+	}
+
+	FArchiveMaterial& NewMaterial = MockUpDescription.MaterialHIdToMaterial.Emplace(MaterialId, MaterialId);
+	GetMaterial(MaterialId, NewMaterial.Material);
+	NewMaterial.UEMaterialName = BuildMaterialName(NewMaterial.Material);
+	return NewMaterial;
 }
 
-int32 BuildColorHash(uint32 ColorHId)
+FArchiveColor& FCoreTechFileParser::FindOrAddColor(uint32 ColorHId)
 {
-	uint32 ColorId;
-	uint8 Alpha;
-	UnhashFastColorHash(ColorHId, ColorId, Alpha);
-
-	// Ref. BaseHelper.cpp getColorData
-	CT_COLOR CtColor = { 200, 200, 200 };
-	if (ColorId > 0)
+	if (FArchiveColor* Color = MockUpDescription.ColorHIdToColor.Find(ColorHId))
 	{
-		CT_MATERIAL_IO::AskIndexedColor((CT_OBJECT_ID) ColorId, CtColor);
+		return *Color;
 	}
-	return BuildColorHash(CtColor, Alpha);
+
+	FArchiveColor& NewColor = MockUpDescription.ColorHIdToColor.Add(ColorHId, ColorHId);
+	GetColor(ColorHId, NewColor.Color);
+	NewColor.UEMaterialName = BuildColorName(NewColor.Color);
+	return NewColor;
 }
 
-int32 BuildMaterialHash(uint32 MaterialId)
+uint32 FCoreTechFileParser::GetObjectMaterial(ICADArchiveObject& Object)
 {
-	FCADMaterial OutMaterial;
-	bool bReturn = GetMaterial(MaterialId, OutMaterial);
-	if (!bReturn)
+	if (FString* Material = Object.MetaData.Find(TEXT("MaterialName")))
 	{
-		return 0;
+		return (uint32) FCString::Atoi64(**Material);
+	}
+	if (FString* Material = Object.MetaData.Find(TEXT("ColorName")))
+	{
+		return (uint32)FCString::Atoi64(**Material);
 	}
-	return BuildMaterialHash(OutMaterial);
+	return 0;
 }
 
-void GetColor(uint32 ColorHash, FColor& OutColor)
+bool GetColor(uint32 ColorUuid, FColor& OutColor)
 {
 	uint32 ColorId;
 	uint8 Alpha;
-	UnhashFastColorHash(ColorHash, ColorId, Alpha);
+	GetCTColorIdAlpha(ColorUuid, ColorId, Alpha);
 
 	CT_COLOR CtColor = { 200, 200, 200 };
 	if (ColorId > 0)
 	{
-		CT_MATERIAL_IO::AskIndexedColor((CT_OBJECT_ID)ColorId, CtColor);
+		if (CT_MATERIAL_IO::AskIndexedColor((CT_OBJECT_ID)ColorId, CtColor) != IO_OK)
+		{
+			return false;
+		}
 	}
 
 	OutColor.R = CtColor[0];
 	OutColor.G = CtColor[1];
 	OutColor.B = CtColor[2];
 	OutColor.A = Alpha;
+	return true;
 }
 
-
 bool GetMaterial(uint32 MaterialId, FCADMaterial& OutMaterial)
 {
 	//	// Ref. BaseHelper.cpp
@@ -366,7 +365,6 @@ bool GetMaterial(uint32 MaterialId, FCADMaterial& OutMaterial)
 		}
 	}
 
-	OutMaterial.MaterialId = MaterialId;
 	OutMaterial.MaterialName = CtName.toUnicode();
 	OutMaterial.Diffuse = FColor(CtDiffuse[0], CtDiffuse[1], CtDiffuse[2], 255);
 	OutMaterial.Ambient = FColor(CtAmbient[0], CtAmbient[1], CtAmbient[2], 255);
@@ -378,153 +376,67 @@ bool GetMaterial(uint32 MaterialId, FCADMaterial& OutMaterial)
 	return true;
 }
 
-
-
-void GetBodiesMaterials(TArray<CT_OBJECT_ID>& BodySet, TMap<uint32, uint32>& MaterialIdToHash, bool bPreferPartData)
+void FCoreTechFileParser::SetFaceMainMaterial(FObjectDisplayDataId& InFaceMaterial, FObjectDisplayDataId& InBodyMaterial, FBodyMesh& BodyMesh, int32 FaceIndex)
 {
-	CT_MATERIAL_ID BodyCtMaterialId = 0;
-	for (int Index = 0; Index < BodySet.Num(); Index++)
-	{
-		FObjectDisplayDataId BodyMaterialId;
-		GetCTObjectDisplayDataIds(BodySet[Index], BodyMaterialId);
-
-		// Loop through the face of the first body and collect material data
-		CT_LIST_IO FaceList;
-		CT_BODY_IO::AskFaces(BodySet[Index], FaceList);
-		FaceList.IteratorInitialize();
-
-		CT_OBJECT_ID FaceID;
-		while ((FaceID = FaceList.IteratorIter()) != 0)
-		{
-			FTessellationData Tessellation;
-			FObjectDisplayDataId FaceMaterialId;
-			GetCTObjectDisplayDataIds(FaceID, FaceMaterialId);
-			SetFaceMainMaterial(FaceMaterialId, BodyMaterialId, MaterialIdToHash, Tessellation);
-		}
-	}
-}
-
-
-void SetFaceMainMaterial(FObjectDisplayDataId& InFaceMaterial, FObjectDisplayDataId& InBodyMaterial, TMap<uint32, uint32>& MaterialIdToMaterialHashMap, FTessellationData& OutFaceTessellations)
-{
-	uint32 FaceMaterialId = 0; // either CT Material ID or CT Color/Alpha ID (ColorId & Alpha << 24)
-	uint32 BodyMaterialId = 0; // either CT Material ID or CT Color/Alpha ID (ColorId & Alpha << 24)
-
 	uint32 FaceMaterialHash = 0;
 	uint32 BodyMaterialHash = 0;
+	uint32 FaceColorHash = 0;
+	uint32 BodyColorHash = 0;
 
-	if (InFaceMaterial.MaterialId > 0)
-	{
-		FaceMaterialId = InFaceMaterial.MaterialId;
-		if (uint32* MHash = MaterialIdToMaterialHashMap.Find(InFaceMaterial.MaterialId))
-		{
-			FaceMaterialHash = *MHash;
-		}
-		else
-		{
-			FaceMaterialHash = BuildMaterialHash(InFaceMaterial.MaterialId);
-			MaterialIdToMaterialHashMap.Add(InFaceMaterial.MaterialId, FaceMaterialHash);
-		}
-	}
-	else if (InFaceMaterial.ColorHId > 0)
-	{
-		FaceMaterialId = InFaceMaterial.ColorHId;
-		if (uint32* MHash = MaterialIdToMaterialHashMap.Find(InFaceMaterial.ColorHId))
-		{
-			FaceMaterialHash = *MHash;
-		}
-		else
-		{
-			FaceMaterialHash = BuildColorHash(InFaceMaterial.ColorHId);
-			MaterialIdToMaterialHashMap.Add(InFaceMaterial.ColorHId, FaceMaterialHash);
-		}
-	}
+	FTessellationData& FaceTessellations = BodyMesh.Faces[FaceIndex];
 
-	if (InBodyMaterial.MaterialId > 0)
+	if (InFaceMaterial.Material > 0)
 	{
-		BodyMaterialId = InBodyMaterial.MaterialId;
-		if (uint32* MHash = MaterialIdToMaterialHashMap.Find(InBodyMaterial.MaterialId))
-		{
-			BodyMaterialHash = *MHash;
-		}
-		else
-		{
-			BodyMaterialHash = BuildMaterialHash(InBodyMaterial.MaterialId);
-			MaterialIdToMaterialHashMap.Add(InBodyMaterial.MaterialId, BodyMaterialHash);
-		}
+		FArchiveMaterial& Material = FindOrAddMaterial(InFaceMaterial.Material);
+		FaceTessellations.MaterialName = Material.UEMaterialName;
+		BodyMesh.MaterialSet.Add(Material.UEMaterialName);
 	}
-	else if (InBodyMaterial.ColorHId > 0)
+	else if (InBodyMaterial.Material > 0)
 	{
-		BodyMaterialId = InBodyMaterial.ColorHId;
-		if (uint32* MHash = MaterialIdToMaterialHashMap.Find(InBodyMaterial.ColorHId))
-		{
-			BodyMaterialHash = *MHash;
-		}
-		else
-		{
-			BodyMaterialHash = BuildColorHash(InBodyMaterial.ColorHId);
-			MaterialIdToMaterialHashMap.Add(InBodyMaterial.ColorHId, BodyMaterialHash);
-		}
+		FArchiveMaterial& Material = FindOrAddMaterial(InBodyMaterial.Material);
+		FaceTessellations.MaterialName = Material.UEMaterialName;
+		BodyMesh.MaterialSet.Add(Material.UEMaterialName);
 	}
 
-	uint32 MaterialId = 0;
-	uint32 MaterialHash = 0;
-
-	// set output
-	if (true)
+	if (InFaceMaterial.Color > 0)
 	{
-		if (BodyMaterialId)
-		{
-			MaterialId = BodyMaterialId;
-			MaterialHash = BodyMaterialHash;
-		} 
-		else
-		{
-			MaterialId = FaceMaterialId;
-			MaterialHash = FaceMaterialHash;
-		}
+		FArchiveColor& Color = FindOrAddColor(InFaceMaterial.Color);
+		FaceTessellations.ColorName = Color.UEMaterialName;
+		BodyMesh.ColorSet.Add(Color.UEMaterialName);
 	}
-	else
+	else if (InBodyMaterial.Color > 0)
 	{
-		if (FaceMaterialId)
-		{
-			MaterialId = FaceMaterialId;
-			MaterialHash = FaceMaterialHash;
-		}
-		else
-		{
-			MaterialId = BodyMaterialId;
-			MaterialHash = BodyMaterialHash;
-		}
+		FArchiveColor& Color = FindOrAddColor(InBodyMaterial.Color);
+		FaceTessellations.ColorName = Color.UEMaterialName;
+		BodyMesh.ColorSet.Add(Color.UEMaterialName);
 	}
-
-	OutFaceTessellations.MaterialId = MaterialId;
-	OutFaceTessellations.MaterialHash = MaterialHash;
-}
-
-int32 BuildColorHash(uint32 ColorId, uint8 Alpha)
-{
-	// Ref. BaseHelper.cpp getColorData
-	CT_COLOR CtColor = { 200, 200, 200 };
-	if (ColorId > 0)
+	else if(InFaceMaterial.DefaultMaterialName)
 	{
-		CT_MATERIAL_IO::AskIndexedColor(ColorId, CtColor);
+		FaceTessellations.ColorName = InFaceMaterial.DefaultMaterialName;
+		BodyMesh.ColorSet.Add(InFaceMaterial.DefaultMaterialName);
 	}
-	return BuildColorHash(CtColor, Alpha);
 }
 
 uint32 GetStaticMeshUuid(const TCHAR* OutSgFile, const int32 BodyId)
 {
 	uint32 BodyUUID = GetTypeHash(OutSgFile);
-	FString BodyStr = TEXT("B ") + FString::FromInt(BodyId);
-	BodyUUID = HashCombine(BodyUUID, GetTypeHash(*BodyStr));
+	BodyUUID = HashCombine(BodyUUID, GetTypeHash(BodyId));
 	return BodyUUID;
 }
 
-void FCoreTechFileParser::ExportFileSceneGraph()
+void FCoreTechFileParser::ExportSceneGraphFile()
 {
-	FFileHelper::SaveStringArrayToFile(SceneGraphDescription, *FPaths::Combine(CachePath, TEXT("scene"), SceneGraphFile + TEXT(".sg"))); 
-	SceneGraphDescription.Empty();
+	SerializeMockUp(MockUpDescription, *FPaths::Combine(CachePath, TEXT("scene"), MockUpDescription.SceneGraphArchive + TEXT(".sg")));
+}
+
+void FCoreTechFileParser::ExportMeshArchiveFile()
+{
+	SerializeBodyMeshSet(*MeshArchiveFilePath, BodyMeshes);
+}
+
+void FCoreTechFileParser::LoadSceneGraphArchive(const FString& SGFile)
+{
+	DeserializeMockUpFile(*SGFile, MockUpDescription);
 }
 
 uint32 FCoreTechFileParser::GetMaterialNum()
@@ -558,29 +470,9 @@ uint32 FCoreTechFileParser::GetMaterialNum()
 	return IColor + IMaterial - 2;
 }
 
-void FCoreTechFileParser::ReadColor()
-{
-	CT_COLOR CtColor;
-	CT_UINT32 IColor = 1;
-
-	SceneGraphDescription[COLORSETLINE] = FString::Printf(TEXT("Color %d"), SceneGraphDescription.Num());
-	while (true)
-	{
-		if (CT_MATERIAL_IO::AskIndexedColor(IColor, CtColor) != IO_OK)
-		{
-			break;
-		}
-		SceneGraphDescription.Add(FString::Printf(TEXT("%d %d %d %d"), IColor, CtColor[0], CtColor[1], CtColor[2]));
-		IColor++;
-	}
-	SceneGraphDescription.Add(TEXT(""));
-}
-
-void FCoreTechFileParser::ReadMaterial()
+void FCoreTechFileParser::ReadMaterials()
 {
 	CT_UINT32 MaterialId = 1;
-
-	SceneGraphDescription[MATERIALSETLINE] = FString::Printf(TEXT("Material %d"), SceneGraphDescription.Num());
 	while (true)
 	{
 		FCADMaterial Material;
@@ -590,58 +482,23 @@ void FCoreTechFileParser::ReadMaterial()
 			break;
 		}
 
-		int32 MaterialHash = BuildMaterialHash(Material);
+		FArchiveMaterial& MaterialObject = MockUpDescription.MaterialHIdToMaterial.Emplace(MaterialId, MaterialId);
+		MaterialObject.UEMaterialName = BuildMaterialName(Material);
+		MaterialObject.Material = Material; 
 
-		SceneGraphDescription.Add(FString::Printf(TEXT("%d %s"), MaterialId, *Material.MaterialName));        // COLORSETLINE 3
-		SceneGraphDescription.Add(FString::Printf(TEXT("%d %d %d %d %d %d %d %d %d %d %d %f %f %f %s"), 
-			MaterialId, MaterialHash, 
-			Material.Diffuse.R, Material.Diffuse.G, Material.Diffuse.B, 
-			Material.Ambient.R, Material.Ambient.G, Material.Ambient.B, 
-			Material.Specular.R, Material.Specular.G, Material.Specular.B, 
-			Material.Shininess, Material.Transparency, Material.Reflexion, *Material.TextureName));
-		MaterialIdToMaterialHashMap.Add(MaterialId, MaterialHash);
 		MaterialId++;
 	}
-	SceneGraphDescription.Add(TEXT(""));
 }
 
-FCoreTechFileParser::FCoreTechFileParser(const FString& InCADFullPath, const FString& InCachePath, const FImportParameters& ImportParams, bool bInPreferBodyData, bool bScaleUVMap)
+FCoreTechFileParser::FCoreTechFileParser(const FString& InCADFullPath, const FString& InCachePath, const FImportParameters& ImportParams)
 	: CachePath(InCachePath)
 	, FullPath(InCADFullPath)
 	, bNeedSaveCTFile(false)
-	, bPreferBodyData(bInPreferBodyData)
-	, bScaleUVMap(bScaleUVMap)
 	, ImportParameters(ImportParams)
 {
-	CTIdToRawLineMap.Reserve(EXTREFNUM);
 	CTKIO_InitializeKernel(ImportParameters.MetricUnit);
 }
 
-void GetRawDataFileExternalRef(const FString& InRawDataFile, TSet<FString>& ExternalReferences)
-{
-	TArray<FString> SGDescription;
-	FFileHelper::LoadFileToStringArray(SGDescription, *InRawDataFile);
-
-	if (SGDescription.Num() < 10)
-		return;
-
-	TArray<FString> HeaderSplit;
-	FString& HeaderExternalRef = SGDescription[EXTERNALREFLINE];
-	HeaderExternalRef.ParseIntoArray(HeaderSplit, TEXT(" "));
-
-	if (HeaderSplit.Num() != 3)
-	{
-		return;
-	}
-
-	int32 StartRef = FCString::Atoi(*HeaderSplit[1]);
-	int32 EndRef = StartRef + FCString::Atoi(*HeaderSplit[2]);
-	for (int32 Ref = StartRef; Ref < EndRef; ++Ref)
-	{
-		ExternalReferences.Add(SGDescription[Ref]);
-	}
-}
-
 FCoreTechFileParser::EProcessResult FCoreTechFileParser::ProcessFile()
 {
 	FileConfiguration.Empty();
@@ -664,36 +521,36 @@ FCoreTechFileParser::EProcessResult FCoreTechFileParser::ProcessFile()
 	FFileStatData FileStatData = IFileManager::Get().GetStatData(*FullPath);
 	uint32 FileHash = GetFileHash(CADFile, FileStatData, FileConfiguration, ImportParameters);
 
-	SceneGraphFile = FString::Printf(TEXT("UEx%08x"), FileHash);
+	MockUpDescription.SceneGraphArchive = FString::Printf(TEXT("UEx%08x"), FileHash);
 
-	FString RawDataFile = FPaths::Combine(CachePath, TEXT("scene"), SceneGraphFile + TEXT(".sg"));
-	FString CTFile = FPaths::Combine(CachePath, TEXT("cad"), SceneGraphFile + TEXT(".ct"));
+	FString SceneGraphArchiveFilePath = FPaths::Combine(CachePath, TEXT("scene"), MockUpDescription.SceneGraphArchive + TEXT(".sg"));
+	FString CTFilePath = FPaths::Combine(CachePath, TEXT("cad"), MockUpDescription.SceneGraphArchive + TEXT(".ct"));
 
 	uint32 MeshFileHash = GetGeomFileHash(FileHash, ImportParameters);
-	MeshFile = FString::Printf(TEXT("UEx%08x"), MeshFileHash);
-	RawDataGeom = FPaths::Combine(CachePath, TEXT("mesh"), MeshFile + TEXT(".gm"));
+	MeshArchiveFile = FString::Printf(TEXT("UEx%08x"), MeshFileHash);
+	MeshArchiveFilePath = FPaths::Combine(CachePath, TEXT("mesh"), MeshArchiveFile + TEXT(".gm"));
 
 	bool bNeedToProceed = true;
-	if (IFileManager::Get().FileExists(*RawDataFile))
+	if (IFileManager::Get().FileExists(*SceneGraphArchiveFilePath))
 	{
-		if (!IFileManager::Get().FileExists(*CTFile)) // the file is scene graph only because no CT file
+		if (!IFileManager::Get().FileExists(*CTFilePath)) // the file is scene graph only because no CT file
 		{
 			bNeedToProceed = false;
 		}
-		else if (IFileManager::Get().FileExists(*RawDataGeom)) // the file has been proceed with same meshing parameters
+		else if (IFileManager::Get().FileExists(*MeshArchiveFilePath)) // the file has been proceed with same meshing parameters
 		{
 			bNeedToProceed = false;
 		}
 		else // the file has been converted into CT file but meshed with different parameters
 		{
-			FullPath = CTFile;
+			FullPath = CTFilePath;
 		}
 	}
 
 	if (!bNeedToProceed)
 	{
 		// The file has been yet proceed, get ExternalRef
-		GetRawDataFileExternalRef(RawDataFile, ExternalRefSet);
+		LoadSceneGraphArchive(SceneGraphArchiveFilePath);
 		return EProcessResult::ProcessOk;
 	}
 
@@ -752,7 +609,12 @@ FCoreTechFileParser::EProcessResult FCoreTechFileParser::ReadFileWithKernelIO()
 	Repair(MainId, ImportParameters.StitchingTechnique);
 	SetCoreTechTessellationState(ImportParameters);
 
-	const CT_OBJECT_TYPE TypeSet[] = { CT_INSTANCE_TYPE, CT_ASSEMBLY_TYPE, CT_PART_TYPE, CT_COMPONENT_TYPE, CT_BODY_TYPE, CT_UNLOADED_COMPONENT_TYPE, CT_UNLOADED_ASSEMBLY_TYPE, CT_UNLOADED_PART_TYPE };
+	MockUpDescription.FullPath = FullPath;
+	MockUpDescription.CADFile = CADFile;
+
+	const CT_OBJECT_TYPE TypeSet[] = { CT_INSTANCE_TYPE, CT_ASSEMBLY_TYPE, CT_PART_TYPE, CT_COMPONENT_TYPE, CT_BODY_TYPE, CT_UNLOADED_COMPONENT_TYPE, CT_UNLOADED_ASSEMBLY_TYPE, CT_UNLOADED_PART_TYPE};
+	enum EObjectTypeIndex : uint8	{ CT_INSTANCE_INDEX = 0, CT_ASSEMBLY_INDEX, CT_PART_INDEX, CT_COMPONENT_INDEX, CT_BODY_INDEX, CT_UNLOADED_COMPONENT_INDEX, CT_UNLOADED_ASSEMBLY_INDEX, CT_UNLOADED_PART_INDEX };
+
 	uint32 NbElements[8], NbTotal = 10;
 	for (int32 index = 0; index < 8; index++)
 	{
@@ -760,36 +622,26 @@ FCoreTechFileParser::EProcessResult FCoreTechFileParser::ReadFileWithKernelIO()
 		NbTotal += NbElements[index];
 	}
 
-	{
-		TArray<uint8> RawData;
-		RawData.Append((uint8*)&NbElements[4], sizeof(uint32));
-		FFileHelper::SaveArrayToFile(RawData, *RawDataGeom, &IFileManager::Get());
-	}
+	BodyMeshes.Reserve(NbElements[CT_BODY_INDEX]);
 
-	SceneGraphDescription.Reset(NbTotal * 20);
-	ExternalRefSet.Empty(NbTotal);
-	CTIdToRawLineMap.Reset();
-
-	// Header
-	SceneGraphDescription.Add(CADFile);  // line 0: File name
-	SceneGraphDescription.Add(FullPath);  // line 1: File path
-	SceneGraphDescription.Add(TEXT(""));  // Line 2 is reserved
-	SceneGraphDescription.Add(TEXT("Color 0"));        // COLORSETLINE 3
-	SceneGraphDescription.Add(TEXT("Material 0"));     // MATERIALSETLINE 4
-	SceneGraphDescription.Add(TEXT(""));
-	SceneGraphDescription.Add(TEXT(""));
-	SceneGraphDescription.Add(TEXT("ExternalRef 0"));  // EXTERNALREFLINE 7
-	SceneGraphDescription.Add(TEXT("MapCTId 0"));      // MAPCTIDLINE 8
-	SceneGraphDescription.Add(TEXT(""));
+	MockUpDescription.BodySet.Reserve(NbElements[CT_BODY_INDEX]);
+	MockUpDescription.ComponentSet.Reserve(NbElements[CT_ASSEMBLY_INDEX] + NbElements[CT_PART_INDEX] + NbElements[CT_COMPONENT_INDEX]);
+	MockUpDescription.UnloadedComponentSet.Reserve(NbElements[CT_UNLOADED_COMPONENT_INDEX] + NbElements[CT_UNLOADED_ASSEMBLY_INDEX] + NbElements[CT_UNLOADED_PART_INDEX]);
+	MockUpDescription.InstanceSet.Reserve(NbElements[CT_INSTANCE_INDEX]);
+
+	MockUpDescription.CADIdToBodyIndex.Reserve(NbElements[CT_BODY_INDEX]);
+	MockUpDescription.CADIdToComponentIndex.Reserve(NbElements[CT_ASSEMBLY_INDEX] + NbElements[CT_PART_INDEX] + NbElements[CT_COMPONENT_INDEX]);
+	MockUpDescription.CADIdToUnloadedComponentIndex.Reserve(NbElements[CT_UNLOADED_COMPONENT_INDEX] + NbElements[CT_UNLOADED_ASSEMBLY_INDEX] + NbElements[CT_UNLOADED_PART_INDEX]);
+	MockUpDescription.CADIdToInstanceIndex.Reserve(NbElements[CT_INSTANCE_INDEX]);
 
 	uint32 MaterialNum = GetMaterialNum();
-	MaterialIdToMaterialHashMap.Empty(MaterialNum);
+	MockUpDescription.MaterialHIdToMaterial.Reserve(MaterialNum);
 
-	ReadColor();
-	ReadMaterial();
+	ReadMaterials();
 
 	// Parse the file
-	bool bReadNodeSucceed = ReadNode(MainId);
+	uint32 DefaultMaterialHash = 0;
+	bool bReadNodeSucceed = ReadNode(MainId, DefaultMaterialHash);
 	// End of parsing
 
 	if (bNeedSaveCTFile)
@@ -797,7 +649,7 @@ FCoreTechFileParser::EProcessResult FCoreTechFileParser::ReadFileWithKernelIO()
 		CT_LIST_IO ObjectList;
 		ObjectList.PushBack(MainId);
 
-		CT_KERNEL_IO::SaveFile(ObjectList, *FPaths::Combine(CachePath, TEXT("cad"), SceneGraphFile+TEXT(".ct")), L"Ct");
+		CT_KERNEL_IO::SaveFile(ObjectList, *FPaths::Combine(CachePath, TEXT("cad"), MockUpDescription.SceneGraphArchive + TEXT(".ct")), L"Ct");
 	}
 
 	CT_KERNEL_IO::UnloadModel();
@@ -807,38 +659,8 @@ FCoreTechFileParser::EProcessResult FCoreTechFileParser::ReadFileWithKernelIO()
 		return EProcessResult::ProcessFailed;
 	}
 
-	uint32 LastLine = SceneGraphDescription.Num();
-	if (ExternalRefSet.Num() > 0)
-	{
-		SceneGraphDescription[EXTERNALREFLINE] = FString::Printf(TEXT("ExternalRef %d %d"), LastLine, ExternalRefSet.Num());
-		for (auto ExternalFile : ExternalRefSet)
-		{
-			SceneGraphDescription.Add(ExternalFile);
-		}
-	}
-	SceneGraphDescription.Add("");
-
-	LastLine = SceneGraphDescription.Num();
-	SceneGraphDescription[MAPCTIDLINE] = FString::Printf(TEXT("MapCTId %d"), LastLine);
-
-	FString MapCTId;
-	uint32 SizeId = (uint32) log(CTIdToRawLineMap.Num());
-	uint32 SizeLine = (uint32)log(LastLine);
-	uint32 NbChar = (SizeId > SizeLine ? SizeId : SizeLine) +1;
-	MapCTId.Reserve(CTIdToRawLineMap.Num() * NbChar * 2);
-
-	for (auto ID : CTIdToRawLineMap)
-	{
-		MapCTId.AppendInt(ID.Key);
-		MapCTId.Append(TEXT(" "));
-		MapCTId.AppendInt(ID.Value);
-		MapCTId.Append(TEXT(" "));
-	}
-	SceneGraphDescription.Add(MapCTId);
-	SceneGraphDescription.Add(TEXT(""));
-	SceneGraphDescription.Add(TEXT("F"));
-	
-	ExportFileSceneGraph();
+	ExportSceneGraphFile();
+	ExportMeshArchiveFile();
 
 	return EProcessResult::ProcessOk;
 }
@@ -875,34 +697,44 @@ CT_FLAGS FCoreTechFileParser::SetCoreTechImportOption(const FString& MainFileExt
 	return Flags;
 }
 
-bool FCoreTechFileParser::ReadNode(CT_OBJECT_ID NodeId)
+bool FCoreTechFileParser::ReadNode(CT_OBJECT_ID NodeId, uint32 DefaultMaterialHash)
 {
 	CT_OBJECT_TYPE Type;
 	CT_OBJECT_IO::AskType(NodeId, Type);
 
-	uint32* LineNumber = CTIdToRawLineMap.Find(NodeId);
-	if (LineNumber != nullptr) 
-	{
-		return true;
-	}
-
 	switch (Type)
 	{
 	case CT_INSTANCE_TYPE:
-		return ReadInstance(NodeId);
+		if (int32* Index = MockUpDescription.CADIdToInstanceIndex.Find(NodeId))
+		{
+			return true;
+		}
+		return ReadInstance(NodeId, DefaultMaterialHash);
 
 	case CT_ASSEMBLY_TYPE:
 	case CT_PART_TYPE:
 	case CT_COMPONENT_TYPE:
-		return ReadComponent(NodeId);
+		if (int32* Index = MockUpDescription.CADIdToComponentIndex.Find(NodeId))
+		{
+			return true;
+		}
+		return ReadComponent(NodeId, DefaultMaterialHash);
 
 	case CT_UNLOADED_ASSEMBLY_TYPE:
 	case CT_UNLOADED_COMPONENT_TYPE:
 	case CT_UNLOADED_PART_TYPE:
+		if (int32* Index = MockUpDescription.CADIdToUnloadedComponentIndex.Find(NodeId))
+		{
+			return true;
+		}
 		return ReadUnloadedComponent(NodeId);
 
 	case CT_BODY_TYPE:
-		return ReadBody(NodeId);
+		if (int32* Index = MockUpDescription.CADIdToBodyIndex.Find(NodeId))
+		{
+			return true;
+		}
+		return ReadBody(NodeId, DefaultMaterialHash);
 
 		//Treat all CT_CURVE_TYPE :
 	case CT_CURVE_TYPE:
@@ -926,105 +758,102 @@ bool FCoreTechFileParser::ReadNode(CT_OBJECT_ID NodeId)
 
 bool FCoreTechFileParser::ReadUnloadedComponent(CT_OBJECT_ID ComponentId)
 {
-	CTIdToRawLineMap.Add(ComponentId, SceneGraphDescription.Num());
-
-	FString Data = TEXT("U ") + FString::FromInt(ComponentId);
-	SceneGraphDescription.Add(*Data);
-
-	ReadNodeMetaDatas(ComponentId);
-
 	CT_STR Filename, FileType;
-	CT_COMPONENT_IO::AskExternalDefinition(ComponentId, Filename, FileType);
-	Data = TEXT("ext ");
-	Data += Filename.toUnicode();
-	Data += TEXT(" ");
-	Data += FileType.toUnicode();
-	SceneGraphDescription.Add(*Data);
+	CT_IO_ERROR Error = CT_COMPONENT_IO::AskExternalDefinition(ComponentId, Filename, FileType);
+	if (Error != IO_OK)
+	{
+		return false;
+	}
+
+	int32 Index = MockUpDescription.UnloadedComponentSet.Emplace(ComponentId);
+	MockUpDescription.CADIdToUnloadedComponentIndex.Add(ComponentId, Index);
+	ReadNodeMetaData(ComponentId, MockUpDescription.UnloadedComponentSet[Index].MetaData);
 
-	SceneGraphDescription.Add(TEXT(""));
+	MockUpDescription.UnloadedComponentSet[Index].FileName = Filename.toUnicode();
+	MockUpDescription.UnloadedComponentSet[Index].FileType = FileType.toUnicode();
 
 	return true;
 }
 
-bool FCoreTechFileParser::ReadComponent(CT_OBJECT_ID ComponentId)
+bool FCoreTechFileParser::ReadComponent(CT_OBJECT_ID ComponentId, uint32 DefaultMaterialHash)
 {
-	CTIdToRawLineMap.Add(ComponentId, SceneGraphDescription.Num());
+	int32 Index = MockUpDescription.ComponentSet.Emplace(ComponentId);
+	MockUpDescription.CADIdToComponentIndex.Add(ComponentId, Index);
+	ReadNodeMetaData(ComponentId, MockUpDescription.ComponentSet[Index].MetaData);
 
-	FString Data = TEXT("C ") + FString::FromInt(ComponentId);
-	SceneGraphDescription.Add(*Data);
-
-	ReadNodeMetaDatas(ComponentId);
+	if (uint32 MaterialHash = GetObjectMaterial(MockUpDescription.ComponentSet[Index]))
+	{
+		DefaultMaterialHash = MaterialHash;
+	}
 
 	CT_LIST_IO Children;
 	CT_COMPONENT_IO::AskChildren(ComponentId, Children);
 
 	// Parse children
-	Data = TEXT("children ") + FString::FromInt(Children.Count());
-	SceneGraphDescription.Add(*Data);
-
 	Children.IteratorInitialize();
-	CT_OBJECT_ID CurrentObjectId;
-	while ((CurrentObjectId = Children.IteratorIter()) != 0)
+	CT_OBJECT_ID ChildId;
+	while ((ChildId = Children.IteratorIter()) != 0)
 	{
-		SceneGraphDescription.Add(*FString::FromInt(CurrentObjectId));
-	}
-	SceneGraphDescription.Add(TEXT(""));
-
-	Children.IteratorInitialize();
-	while ((CurrentObjectId = Children.IteratorIter()) != 0)
-	{
-		ReadNode(CurrentObjectId);
+		MockUpDescription.ComponentSet[Index].Children.Add(ChildId);
+		ReadNode(ChildId, DefaultMaterialHash);
 	}
 
 	return true;
 }
 
-bool FCoreTechFileParser::ReadInstance(CT_OBJECT_ID InstanceNodeId)
+bool FCoreTechFileParser::ReadInstance(CT_OBJECT_ID InstanceNodeId, uint32 DefaultMaterialHash)
 {
-	CTIdToRawLineMap.Add(InstanceNodeId, SceneGraphDescription.Num());
+	NodeConfiguration.Empty();
 
-	FString Data = TEXT("I ") + FString::FromInt(InstanceNodeId);
-	SceneGraphDescription.Add(*Data);
+	int32 Index = MockUpDescription.InstanceSet.Emplace(InstanceNodeId);
+	MockUpDescription.CADIdToInstanceIndex.Add(InstanceNodeId, Index);
+	ReadNodeMetaData(InstanceNodeId, MockUpDescription.InstanceSet[Index].MetaData);
 
-	ReadNodeMetaDatas(InstanceNodeId);
+	if (uint32 MaterialHash = GetObjectMaterial(MockUpDescription.InstanceSet[Index]))
+	{
+		DefaultMaterialHash = MaterialHash;
+	}
 
 	// Ask the transformation of the instance
-	double CTMatrix[16];
-	CT_INSTANCE_IO::AskTransformation(InstanceNodeId, CTMatrix);
-	SceneGraphDescription.Add(FString::Printf(TEXT("matrix %f %f %f %f %f %f %f %f %f %f %f %f %f %f %f %f"), CTMatrix[0], CTMatrix[1], CTMatrix[2], CTMatrix[3], CTMatrix[4], CTMatrix[5], CTMatrix[6], CTMatrix[7], CTMatrix[8], CTMatrix[9], CTMatrix[10], CTMatrix[11], CTMatrix[12], CTMatrix[13], CTMatrix[14], CTMatrix[15]));
-
+	double Matrix[16];
+	if (CT_INSTANCE_IO::AskTransformation(InstanceNodeId, Matrix) == IO_OK)
+	{
+		float* MatrixFloats = (float*)MockUpDescription.InstanceSet[Index].TransformMatrix.M;
+		for (int32 index = 0; index < 16; index++)
+		{
+			MatrixFloats[index] = (float) Matrix[index];
+		}
+	}
+	
 	// Ask the reference
 	CT_OBJECT_ID ReferenceNodeId;
 	CT_IO_ERROR CTReturn = CT_INSTANCE_IO::AskChild(InstanceNodeId, ReferenceNodeId);
 	if (CTReturn != CT_IO_ERROR::IO_OK)
 		return false;
+	MockUpDescription.InstanceSet[Index].ReferenceNodeId = ReferenceNodeId;
 
-	NodeConfiguration.Empty(); // #ueent_CAD
 	CT_OBJECT_TYPE type;
 	CT_OBJECT_IO::AskType(ReferenceNodeId, type);
 	if (type == CT_UNLOADED_PART_TYPE || type == CT_UNLOADED_COMPONENT_TYPE || type == CT_UNLOADED_ASSEMBLY_TYPE)
 	{
+		MockUpDescription.InstanceSet[Index].bIsExternalRef = true;
+
 		CT_STR ComponentFile, FileType;
 		CT_COMPONENT_IO::AskExternalDefinition(ReferenceNodeId, ComponentFile, FileType);
-		FString ExternalRef = FPaths::GetCleanFilename(ComponentFile.toUnicode());
+		MockUpDescription.InstanceSet[Index].ExternalRef = FPaths::GetCleanFilename(ComponentFile.toUnicode());
 		if(!NodeConfiguration.IsEmpty())
 		{
-			ExternalRef += TEXT("|") + NodeConfiguration;
+			MockUpDescription.InstanceSet[Index].ExternalRef += TEXT("|") + NodeConfiguration;
 		}
 
-		Data = TEXT("ext ") + FString::FromInt(ReferenceNodeId) + TEXT(" ") + ExternalRef;
-		SceneGraphDescription.Add(*Data);
-
-		ExternalRefSet.Add(ExternalRef);
+		MockUpDescription.ExternalRefSet.Add(MockUpDescription.InstanceSet[Index].ExternalRef);
 	}
 	else
 	{
-		Data = TEXT("ref ") + FString::FromInt(ReferenceNodeId);
-		SceneGraphDescription.Add(Data);
+		MockUpDescription.InstanceSet[Index].bIsExternalRef = false;
 	}
-	SceneGraphDescription.Add(TEXT(""));
 
-	return ReadNode(ReferenceNodeId);
+	return ReadNode(ReferenceNodeId, DefaultMaterialHash);
 }
 
 
@@ -1041,111 +870,78 @@ uint32 GetBodiesFaceSetNum(TArray<CT_OBJECT_ID>& BodySet)
 	return size;
 }
 
-uint32 GetBodiesTessellations(TArray<CT_OBJECT_ID>& BodySet, TArray<FTessellationData>& FaceTessellationSet, TMap<uint32, uint32>& MaterialIdToMaterialHash, int32& OutRawDataSize, const FImportParameters& ImportParams)
+void FCoreTechFileParser::GetBodyTessellation(CT_OBJECT_ID BodyId, FBodyMesh& OutBodyMesh, const FImportParameters& ImportParams, uint32 DefaultMaterialHash)
 {
-	uint32 FaceSize = GetBodiesFaceSetNum(BodySet);
+	CT_LIST_IO FaceList;
+	CT_BODY_IO::AskFaces(BodyId, FaceList);
 
-	// Allocate memory space for tessellation data
-	FaceTessellationSet.Reserve(FaceSize);
+	uint32 FaceSize = FaceList.Count();
 
-	uint32 GlobalTriangleCount = 0;
+	// Allocate memory space for tessellation data
+	OutBodyMesh.Faces.Reserve(FaceSize);
+	OutBodyMesh.ColorSet.Reserve(FaceSize);
+	OutBodyMesh.MaterialSet.Reserve(FaceSize);
 
 	CT_MATERIAL_ID BodyCtMaterialId = 0;
-	for (int Index = 0; Index < BodySet.Num(); Index++)
-	{
-		FObjectDisplayDataId BodyMaterial;
-		GetCTObjectDisplayDataIds(BodySet[Index], BodyMaterial);
-
-		// Loop through the face of the first body and collect all tessellation data
-		CT_LIST_IO FaceList;
-		CT_BODY_IO::AskFaces(BodySet[Index], FaceList);
-		FaceList.IteratorInitialize();
 
-		int32 LastIndex = FaceTessellationSet.Num();
-		int32 FaceIndex = LastIndex;
+	FObjectDisplayDataId BodyMaterial;
+	BodyMaterial.DefaultMaterialName = DefaultMaterialHash;
+	GetCTObjectDisplayDataIds(BodyId, BodyMaterial);
 
-		CT_OBJECT_ID FaceID;
-		while ((FaceID = FaceList.IteratorIter()) != 0)
-		{
-			int32 FaceRawSize = 0;
-			uint32 TriangleCount = GetFaceTessellation(FaceID, FaceTessellationSet, FaceRawSize, ImportParams);
+	// Loop through the face of the first body and collect all tessellation data
+	FaceList.IteratorInitialize();
 
-			if (TriangleCount == 0)
-			{
-				continue;
-			}
+	int32 FaceIndex = 0;
 
-			GlobalTriangleCount += TriangleCount;
-			OutRawDataSize += (FaceRawSize + ECTTessellationLine::LastLine * sizeof(uint32));
-			FObjectDisplayDataId FaceMaterial;
-			GetCTObjectDisplayDataIds(FaceID, FaceMaterial);
-			SetFaceMainMaterial(FaceMaterial, BodyMaterial, MaterialIdToMaterialHash, FaceTessellationSet[FaceIndex]);
-			FaceIndex++;
-		}
+	CT_OBJECT_ID FaceID;
+	while ((FaceID = FaceList.IteratorIter()) != 0)
+	{
+		uint32 TriangleNum = GetFaceTessellation(FaceID, OutBodyMesh.Faces, ImportParams);
 
-		int32 FaceNum = FaceTessellationSet.Num() - LastIndex;
-		for (FaceIndex = LastIndex; FaceIndex < FaceTessellationSet.Num(); FaceIndex++)
+		if (TriangleNum == 0)
 		{
-			FaceTessellationSet[FaceIndex].BodyId = (uint32)BodySet[Index];
-			FaceTessellationSet[FaceIndex].BodyFaceNum = FaceNum;
+			continue;
 		}
+
+		OutBodyMesh.TriangleCount += TriangleNum;
+
+		FObjectDisplayDataId FaceMaterial;
+		GetCTObjectDisplayDataIds(FaceID, FaceMaterial);
+		SetFaceMainMaterial(FaceMaterial, BodyMaterial, OutBodyMesh, FaceIndex);
+		FaceIndex++;
 	}
-	return GlobalTriangleCount;
 }
 
-
-bool FCoreTechFileParser::ReadBody(CT_OBJECT_ID BodyId)
+bool FCoreTechFileParser::ReadBody(CT_OBJECT_ID BodyId, uint32 DefaultMaterialHash)
 {
-	CTIdToRawLineMap.Add(BodyId, SceneGraphDescription.Num());
 
-	bNeedSaveCTFile = true;
+	int32 Index = MockUpDescription.BodySet.Emplace(BodyId);
+	MockUpDescription.CADIdToBodyIndex.Add(BodyId, Index);
+	ReadNodeMetaData(BodyId, MockUpDescription.BodySet[Index].MetaData);
 
-	FString Data = TEXT("B ") + FString::FromInt(BodyId);
-	SceneGraphDescription.Add(*Data);
-	
-	ReadNodeMetaDatas(BodyId);
+	int32 BodyMeshIndex = BodyMeshes.Emplace(BodyId);
 
-	// Parse Body to set Get Mesh
-	TArray<CT_OBJECT_ID> BodySet;
-	BodySet.Add(BodyId);
-	TArray<FTessellationData> FaceTessellationSet;
+	if (uint32 MaterialHash = GetObjectMaterial(MockUpDescription.BodySet[Index]))
+	{
+		DefaultMaterialHash = MaterialHash;
+	}
 
-	uint32 BodyUuId = GetStaticMeshUuid(*SceneGraphFile, BodyId);
-                
-	int32 BodyRawDataSize = 0;
-	uint32 NbTriangles = GetBodiesTessellations(BodySet, FaceTessellationSet, MaterialIdToMaterialHashMap, BodyRawDataSize, ImportParameters);
-	TArray<uint8> GlobalRawData;
+	bNeedSaveCTFile = true;
 
-	GlobalRawData.Reserve(BodyRawDataSize);
+	MockUpDescription.BodySet[Index].MeshActorName = GetStaticMeshUuid(*MockUpDescription.SceneGraphArchive, BodyId);
+	BodyMeshes[BodyMeshIndex].MeshActorName = MockUpDescription.BodySet[Index].MeshActorName;
 
-	int32 FaceNum = FaceTessellationSet.Num();
-	for (FTessellationData& Tessellation : FaceTessellationSet)
-	{
-		Tessellation.BodyUuId = BodyUuId;
-		Tessellation.BodyFaceNum = FaceNum;
-		WriteTessellationInRawData(Tessellation, GlobalRawData);
-	}
+	GetBodyTessellation(BodyId, BodyMeshes[BodyMeshIndex], ImportParameters, DefaultMaterialHash);
 
-	bool bWriteOk = FFileHelper::SaveArrayToFile(GlobalRawData, *RawDataGeom, &IFileManager::Get(), FILEWRITE_Append);
-	ensure(bWriteOk);
-	GlobalRawData.Empty();
+	MockUpDescription.BodySet[Index].ColorFaceSet = BodyMeshes[BodyMeshIndex].ColorSet;
+	MockUpDescription.BodySet[Index].MaterialFaceSet = BodyMeshes[BodyMeshIndex].MaterialSet;
 
+	// Save Body in CT file for re-tessellation
 	CT_LIST_IO ObjectList;
 	ObjectList.PushBack(BodyId);
-	FString BodyFile = FString::Printf(TEXT("UEx%08x"), BodyUuId);
+	FString BodyFile = FString::Printf(TEXT("UEx%08x"), MockUpDescription.BodySet[Index].MeshActorName);
 	CT_KERNEL_IO::SaveFile(ObjectList, *FPaths::Combine(CachePath, TEXT("body"), BodyFile + TEXT(".ct")), L"Ct");
 
-	FString MapCTId;
-	MapCTId.Reserve((MaterialIdToMaterialHashMap.Num()+1) * 22);
-
-	MapCTId.Append(TEXT("materialMap "));
-	for (auto Material : MaterialIdToMaterialHashMap)
-	{
-		MapCTId.Append(FString::Printf(TEXT("%u %u "), Material.Key, Material.Value));
-	}
-	SceneGraphDescription.Add(MapCTId);
-	SceneGraphDescription.Add(TEXT(""));
-
 	return true;
 }
 
@@ -1159,44 +955,41 @@ void FCoreTechFileParser::GetAttributeValue(CT_ATTRIB_TYPE AttributType, int Ith
 	if (CT_ATTRIB_DEFINITION_IO::AskFieldDefinition(AttributType, IthField, FieldType, FieldName) != IO_OK) return;
 
 	switch (FieldType) {
-	case CT_ATTRIB_FIELD_UNKNOWN:
-	{
-		break;
-	}
-	case CT_ATTRIB_FIELD_INTEGER:
-	{
-		int IValue;
-		if (CT_CURRENT_ATTRIB_IO::AskIntField(IthField, IValue) != IO_OK) break;
-		Value = FString::FromInt(IValue);
-		break;
-	}
-	case CT_ATTRIB_FIELD_DOUBLE:
-	{
-		double DValue;
-		if (CT_CURRENT_ATTRIB_IO::AskDblField(IthField, DValue) != IO_OK) break;
-		Value = FString::Printf(TEXT("%lf"), DValue);
-		break;
-	}
-	case CT_ATTRIB_FIELD_STRING:
-	{
-		CT_STR StrValue;
-		if (CT_CURRENT_ATTRIB_IO::AskStrField(IthField, StrValue) != IO_OK) break;
-		Value = StrValue.toUnicode();
-		break;
-	}
-	case CT_ATTRIB_FIELD_POINTER:
-	{
-		break;
-	}
+		case CT_ATTRIB_FIELD_UNKNOWN:
+		{
+			break;
+		}
+		case CT_ATTRIB_FIELD_INTEGER:
+		{
+			int IValue;
+			if (CT_CURRENT_ATTRIB_IO::AskIntField(IthField, IValue) != IO_OK) break;
+			Value = FString::FromInt(IValue);
+			break;
+		}
+		case CT_ATTRIB_FIELD_DOUBLE:
+		{
+			double DValue;
+			if (CT_CURRENT_ATTRIB_IO::AskDblField(IthField, DValue) != IO_OK) break;
+			Value = FString::Printf(TEXT("%lf"), DValue);
+			break;
+		}
+		case CT_ATTRIB_FIELD_STRING:
+		{
+			CT_STR StrValue;
+			if (CT_CURRENT_ATTRIB_IO::AskStrField(IthField, StrValue) != IO_OK) break;
+			Value = StrValue.toUnicode();
+			break;
+		}
+		case CT_ATTRIB_FIELD_POINTER:
+		{
+			break;
+		}
 	}
 }
 
-void FCoreTechFileParser::ReadNodeMetaDatas(CT_OBJECT_ID NodeId)
+void FCoreTechFileParser::ReadNodeMetaData(CT_OBJECT_ID NodeId, TMap<FString, FString>& OutMetaData)
 {
 	const FString ConfigName = TEXT("Configuration Name");
-	FString AttributeNameLigne = TEXT("M");
-	SceneGraphDescription.Add(AttributeNameLigne);
-	int32 AttributIndex = SceneGraphDescription.Num();
 
 	if (CT_COMPONENT_IO::IsA(NodeId, CT_COMPONENT_TYPE))
 	{
@@ -1204,8 +997,7 @@ void FCoreTechFileParser::ReadNodeMetaDatas(CT_OBJECT_ID NodeId)
 		CT_COMPONENT_IO::AskExternalDefinition(NodeId, FileName, FileType);
 		if (!FileName.IsEmpty())
 		{
-			SceneGraphDescription.Add(TEXT("ExternalDefinition"));
-			SceneGraphDescription.Add(FileName.toUnicode());
+			OutMetaData.Add(TEXT("ExternalDefinition"), FileName.toUnicode());
 		}
 	}
 
@@ -1219,16 +1011,13 @@ void FCoreTechFileParser::ReadNodeMetaDatas(CT_OBJECT_ID NodeId)
 		switch (IsShow)
 		{
 		case CT_SHOW:
-			SceneGraphDescription.Add(TEXT("ShowAttribute"));
-			SceneGraphDescription.Add(TEXT("show"));
+			OutMetaData.Add(TEXT("ShowAttribute"), TEXT("show"));
 			break;
 		case CT_NOSHOW:
-			SceneGraphDescription.Add(TEXT("ShowAttribute"));
-			SceneGraphDescription.Add(TEXT("noShow"));
+			OutMetaData.Add(TEXT("ShowAttribute"), TEXT("noShow"));
 			break;
 		case CT_UNKNOWN:
-			SceneGraphDescription.Add(TEXT("ShowAttribute"));
-			SceneGraphDescription.Add(TEXT("unknown"));
+			OutMetaData.Add(TEXT("ShowAttribute"), TEXT("unknown"));
 			break;
 		}
 	}
@@ -1250,75 +1039,59 @@ void FCoreTechFileParser::ReadNodeMetaDatas(CT_OBJECT_ID NodeId)
 		switch (AttributeType) {
 
 		case CT_ATTRIB_SPLT:
-			//if (CT_CURRENT_ATTRIB_IO::AskStrField(ITH_SPLT_VALUE, field_strValue) != IO_OK) break;
-			//nodeAttributeSet[TEXT("SPLT")] = field_strValue.toUnicode();
 			break;
 
 		case CT_ATTRIB_NAME:
 			if (CT_CURRENT_ATTRIB_IO::AskStrField(ITH_NAME_VALUE, FieldStrValue) != IO_OK) break;
 			if (FieldStrValue.IsEmpty()) { break; }
-			SceneGraphDescription.Add(TEXT("CTName"));
-			SceneGraphDescription.Add(FieldStrValue.toUnicode());
+			OutMetaData.Add(TEXT("CTName"), FieldStrValue.toUnicode());
 			break;
 
 		case CT_ATTRIB_ORIGINAL_NAME:
 			if (CT_CURRENT_ATTRIB_IO::AskStrField(ITH_NAME_VALUE, FieldStrValue) != IO_OK) break;
 			if (FieldStrValue.IsEmpty()) { break; }
-			SceneGraphDescription.Add(TEXT("Name"));
-			SceneGraphDescription.Add(FieldStrValue.toUnicode());
+			OutMetaData.Add(TEXT("Name"), FieldStrValue.toUnicode());
 			break;
 
 		case CT_ATTRIB_ORIGINAL_FILENAME:
 			if (CT_CURRENT_ATTRIB_IO::AskStrField(ITH_FILENAME_VALUE, FieldStrValue) != IO_OK) break;
 			if (FieldStrValue.IsEmpty()) { break; }
-			SceneGraphDescription.Add(TEXT("FileName"));
-			SceneGraphDescription.Add(FieldStrValue.toUnicode());
+			OutMetaData.Add(TEXT("FileName"), FieldStrValue.toUnicode());
 			break;
 
 		case CT_ATTRIB_UUID:
 			if (CT_CURRENT_ATTRIB_IO::AskStrField(ITH_UUID_VALUE, FieldStrValue) != IO_OK) break;
 			if (FieldStrValue.IsEmpty()) { break; }
-			SceneGraphDescription.Add(TEXT("UUID"));
-			SceneGraphDescription.Add(FieldStrValue.toUnicode());
+			OutMetaData.Add(TEXT("UUID"), FieldStrValue.toUnicode());
 			break;
 
 		case CT_ATTRIB_INPUT_FORMAT_AND_EMETTOR:
 			if (CT_CURRENT_ATTRIB_IO::AskStrField(ITH_INPUT_FORMAT_AND_EMETTOR, FieldStrValue) != IO_OK) break;
 			if (FieldStrValue.IsEmpty()) { break; }
-			SceneGraphDescription.Add(TEXT("Input_Format_and_Emitter"));
-			SceneGraphDescription.Add(FieldStrValue.toUnicode());
+			OutMetaData.Add(TEXT("Input_Format_and_Emitter"), FieldStrValue.toUnicode());
 			break;
 
 		case CT_ATTRIB_CONFIGURATION_NAME:
 			if (CT_CURRENT_ATTRIB_IO::AskStrField(ITH_NAME_VALUE, FieldStrValue) != IO_OK) break;
 			if (FieldStrValue.IsEmpty()) { break; }
-			SceneGraphDescription.Add(TEXT("ConfigurationName"));
-			SceneGraphDescription.Add(FieldStrValue.toUnicode());
+			OutMetaData.Add(TEXT("ConfigurationName"), FieldStrValue.toUnicode());
 			break;
 
 		case CT_ATTRIB_LAYERID:
 			GetAttributeValue(AttributeType, ITH_LAYERID_VALUE, FieldValue);
-			SceneGraphDescription.Add(TEXT("LayerId"));
-			SceneGraphDescription.Add(FieldValue);
+			OutMetaData.Add(TEXT("LayerId"), FieldValue);
 			GetAttributeValue(AttributeType, ITH_LAYERID_NAME, FieldValue);
-			SceneGraphDescription.Add(TEXT("LayerName"));
-			SceneGraphDescription.Add(FieldValue);
+			OutMetaData.Add(TEXT("LayerName"), FieldValue);
 			GetAttributeValue(AttributeType, ITH_LAYERID_FLAG, FieldValue);
-			SceneGraphDescription.Add(TEXT("LayerFlag"));
-			SceneGraphDescription.Add(FieldValue);
+			OutMetaData.Add(TEXT("LayerFlag"), FieldValue);
 			break;
 
 		case CT_ATTRIB_COLORID:
 			{
 				if (CT_CURRENT_ATTRIB_IO::AskIntField(ITH_COLORID_VALUE, FieldIntValue) != IO_OK) break;
 				uint32 ColorId = FieldIntValue;
-				uint8 Alpha = 255;
-				SceneGraphDescription.Add(TEXT("ColorId"));
-				SceneGraphDescription.Add(FString::FromInt(FieldIntValue));
-
-				CT_COLOR CtColor;
-				if (CT_MATERIAL_IO::AskIndexedColor(FieldIntValue, CtColor) != IO_OK) break;
 
+				uint8 Alpha = 255;
 				if (CT_OBJECT_IO::SearchAttribute(NodeId, CT_ATTRIB_TRANSPARENCY) == IO_OK)
 				{
 					if (CT_CURRENT_ATTRIB_IO::AskDblField(0, FieldDoubleValue0) == IO_OK)
@@ -1326,30 +1099,30 @@ void FCoreTechFileParser::ReadNodeMetaDatas(CT_OBJECT_ID NodeId)
 						Alpha = FMath::Max((1. - FieldDoubleValue0), FieldDoubleValue0) * 255.;
 					}
 				}
-				FString colorHexa = FString::Printf(TEXT("%02x%02x%02x%02x"), CtColor[0], CtColor[1], CtColor[2], Alpha);
-				SceneGraphDescription.Add(TEXT("ColorValue"));
-				SceneGraphDescription.Add(colorHexa);
 
-				SceneGraphDescription.Add(TEXT("ColorUEId"));
-				uint32 ColorUuId = BuildFastColorHash(ColorId, Alpha);
-				SceneGraphDescription.Add(FString::Printf(TEXT("%u"), ColorUuId));
+				uint32 ColorHId = BuildColorId(ColorId, Alpha);
+				FArchiveColor& ColorArchive = FindOrAddColor(ColorHId);
+				OutMetaData.Add(TEXT("ColorName"), FString::FromInt(ColorArchive.UEMaterialName));
 
-				int32 ColorHash = BuildColorHash(CtColor, Alpha);
-				MaterialIdToMaterialHashMap.Add(ColorUuId, ColorHash);
+				FString colorHexa = FString::Printf(TEXT("%02x%02x%02x%02x"), ColorArchive.Color.R, ColorArchive.Color.G, ColorArchive.Color.B, ColorArchive.Color.A);
+				OutMetaData.Add(TEXT("ColorValue"), colorHexa);
 			}
 			break;
 
 		case CT_ATTRIB_MATERIALID:
+		{
 			if (CT_CURRENT_ATTRIB_IO::AskIntField(ITH_MATERIALID_VALUE, FieldIntValue) != IO_OK) break;
-			SceneGraphDescription.Add(TEXT("MaterialId"));
-			SceneGraphDescription.Add(FString::FromInt(FieldIntValue));
+			if (FArchiveMaterial* Material = MockUpDescription.MaterialHIdToMaterial.Find(FieldIntValue))
+			{
+				OutMetaData.Add(TEXT("MaterialName"), FString::FromInt(Material->UEMaterialName));
+			}
 			break;
+		}
 
 		case CT_ATTRIB_TRANSPARENCY:
 			if (CT_CURRENT_ATTRIB_IO::AskDblField(ITH_TRANSPARENCY_VALUE, FieldDoubleValue0) != IO_OK) break;
 			FieldIntValue = FMath::Max((1. - FieldDoubleValue0), FieldDoubleValue0) * 255.;
-			SceneGraphDescription.Add(TEXT("Transparency"));
-			SceneGraphDescription.Add(FString::FromInt(FieldIntValue));
+			OutMetaData.Add(TEXT("Transparency"), FString::FromInt(FieldIntValue));
 			break;
 
 		case CT_ATTRIB_COMMENT:
@@ -1358,34 +1131,26 @@ void FCoreTechFileParser::ReadNodeMetaDatas(CT_OBJECT_ID NodeId)
 
 		case CT_ATTRIB_REFCOUNT:
 			if (CT_CURRENT_ATTRIB_IO::AskIntField(ITH_REFCOUNT_VALUE, FieldIntValue) != IO_OK) break;
-			//NodeAttributeSet.Add(TEXT("RefCount"));
-			SceneGraphDescription.Add(FString::FromInt(FieldIntValue));
+			//OutMetaData.Add(TEXT("RefCount"), FString::FromInt(FieldIntValue));
 			break;
 
 		case CT_ATTRIB_TESS_PARAMS:
-			//ITH_TESS_PARAMS_MAXSAG, ITH_TESS_PARAMS_MAXANGLE, ITH_TESS_PARAMS_MAXLENGTH
-			break;
-
 		case CT_ATTRIB_COMPARE_RESULT:
-			//ITH_COMPARE_TYPE, ITH_COMPARE_RESULT_MAXDIST, ITH_COMPARE_RESULT_MINDIST, ITH_COMPARE_THICKNESS, ITH_COMPARE_BACKLASH, ITH_COMPARE_CLEARANCE, ITH_COMPARE_MAXTHICKNESS, ITH_COMPARE_RESULT_NEAREST, ITH_COMPARE_RESULT_ORIGIN1_X, ITH_COMPARE_RESULT_ORIGIN1_Y, ITH_COMPARE_RESULT_ORIGIN1_Z, ITH_COMPARE_RESULT_DIRECTION1_X, ITH_COMPARE_RESULT_DIRECTION1_Y, ITH_COMPARE_RESULT_DIRECTION1_Z, ITH_COMPARE_RESULT_RADIUS1_1, ITH_COMPARE_RESULT_RADIUS2_1, ITH_COMPARE_RESULT_HALFHANGLE1, ITH_COMPARE_RESULT_ORIGIN2_X, ITH_COMPARE_RESULT_ORIGIN2_Y, ITH_COMPARE_RESULT_ORIGIN2_Z, ITH_COMPARE_RESULT_DIRECTION2_X, ITH_COMPARE_RESULT_DIRECTION2_Y, ITH_COMPARE_RESULT_DIRECTION2_Z, ITH_COMPARE_RESULT_RADIUS1_2, ITH_COMPARE_RESULT_RADIUS2_2, ITH_COMPARE_RESULT_HALFHANGLE2, ITH_COMPARE_RESULT_COLOR1, ITH_COMPARE_RESULT_COLOR2, ITH_COMPARE_DIST_VECTOR1_X, ITH_COMPARE_DIST_VECTOR1_Y, ITH_COMPARE_DIST_VECTOR1_Z, ITH_COMPARE_DIST_VECTOR2_X, ITH_COMPARE_DIST_VECTOR2_Y, ITH_COMPARE_DIST_VECTOR2_Z
 			break;
+
 		case CT_ATTRIB_DENSITY:
 			//ITH_VOLUME_DENSITY_VALUE
 			break;
 
 		case CT_ATTRIB_MASS_PROPERTIES:
 			if (CT_CURRENT_ATTRIB_IO::AskDblField(ITH_MASS_PROPERTIES_AREA, FieldDoubleValue0) != IO_OK) break;
-			SceneGraphDescription.Add(TEXT("Area"));
-			SceneGraphDescription.Add(FString::Printf(TEXT("%lf"), FieldDoubleValue0));
+			OutMetaData.Add(TEXT("Area"), FString::Printf(TEXT("%lf"), FieldDoubleValue0));
 			if (CT_CURRENT_ATTRIB_IO::AskDblField(ITH_MASS_PROPERTIES_VOLUME, FieldDoubleValue0) != IO_OK) break;
-			SceneGraphDescription.Add(TEXT("Volume"));
-			SceneGraphDescription.Add(FString::Printf(TEXT("%lf"), FieldDoubleValue0));
+			OutMetaData.Add(TEXT("Volume"), FString::Printf(TEXT("%lf"), FieldDoubleValue0));
 			if (CT_CURRENT_ATTRIB_IO::AskDblField(ITH_MASS_PROPERTIES_MASS, FieldDoubleValue0) != IO_OK) break;
-			SceneGraphDescription.Add(TEXT("Mass"));
-			SceneGraphDescription.Add(FString::Printf(TEXT("%lf"), FieldDoubleValue0));
+			OutMetaData.Add(TEXT("Mass"), FString::Printf(TEXT("%lf"), FieldDoubleValue0));
 			if (CT_CURRENT_ATTRIB_IO::AskDblField(ITH_MASS_PROPERTIES_LENGTH, FieldDoubleValue0) != IO_OK) break;
-			SceneGraphDescription.Add(TEXT("Length"));
-			SceneGraphDescription.Add(FString::Printf(TEXT("%lf"), FieldDoubleValue0));
+			OutMetaData.Add(TEXT("Length"), FString::Printf(TEXT("%lf"), FieldDoubleValue0));
 			//ITH_MASS_PROPERTIES_COGX, ITH_MASS_PROPERTIES_COGY, ITH_MASS_PROPERTIES_COGZ
 			//ITH_MASS_PROPERTIES_M1, ITH_MASS_PROPERTIES_M2, ITH_MASS_PROPERTIES_M3
 			//ITH_MASS_PROPERTIES_IXXG,ITH_MASS_PROPERTIES_IYYG, ITH_MASS_PROPERTIES_IZZG, ITH_MASS_PROPERTIES_IXYG, ITH_MASS_PROPERTIES_IYZG, ITH_MASS_PROPERTIES_IZXG
@@ -1400,15 +1165,13 @@ void FCoreTechFileParser::ReadNodeMetaDatas(CT_OBJECT_ID NodeId)
 		case CT_ATTRIB_INTEGER_METADATA:
 			if (CT_CURRENT_ATTRIB_IO::AskStrField(ITH_INTEGER_METADATA_NAME, FieldName) != IO_OK) break;
 			if (CT_CURRENT_ATTRIB_IO::AskIntField(ITH_INTEGER_METADATA_VALUE, FieldIntValue) != IO_OK) break;
-			SceneGraphDescription.Add(FieldName.toUnicode());
-			SceneGraphDescription.Add(FString::FromInt(FieldIntValue));
+			OutMetaData.Add(FieldName.toUnicode(), FString::FromInt(FieldIntValue));
 			break;
 
 		case CT_ATTRIB_DOUBLE_METADATA:
 			if (CT_CURRENT_ATTRIB_IO::AskStrField(ITH_DOUBLE_METADATA_NAME, FieldName) != IO_OK) break;
 			if (CT_CURRENT_ATTRIB_IO::AskDblField(ITH_DOUBLE_METADATA_VALUE, FieldDoubleValue0) != IO_OK) break;
-			SceneGraphDescription.Add(FieldName.toUnicode());
-			SceneGraphDescription.Add(FString::Printf(TEXT("%lf"), FieldDoubleValue0));
+			OutMetaData.Add(FieldName.toUnicode(), FString::Printf(TEXT("%lf"), FieldDoubleValue0));
 			break;
 
 		case CT_ATTRIB_STRING_METADATA:
@@ -1419,115 +1182,84 @@ void FCoreTechFileParser::ReadNodeMetaDatas(CT_OBJECT_ID NodeId)
 			{
 				NodeConfiguration = FieldStrValue.toUnicode();
 			}
-			SceneGraphDescription.Add(FieldName.toUnicode());
-			SceneGraphDescription.Add(FieldStrValue.toUnicode());
+			OutMetaData.Add(FieldName.toUnicode(), FieldStrValue.toUnicode());
 			break;
 
 		case CT_ATTRIB_ORIGINAL_UNITS:
 			if (CT_CURRENT_ATTRIB_IO::AskDblField(ITH_ORIGINAL_UNITS_MASS, FieldDoubleValue0) != IO_OK) break;
 			if (CT_CURRENT_ATTRIB_IO::AskDblField(ITH_ORIGINAL_UNITS_LENGTH, FieldDoubleValue1) != IO_OK) break;
 			if (CT_CURRENT_ATTRIB_IO::AskDblField(ITH_ORIGINAL_UNITS_DURATION, FieldDoubleValue2) != IO_OK) break;
-			SceneGraphDescription.Add(TEXT("OriginalUnitsMass"));
-			SceneGraphDescription.Add(FString::Printf(TEXT("%lf"), FieldDoubleValue0));
-			SceneGraphDescription.Add(TEXT("OriginalUnitsLength"));
-			SceneGraphDescription.Add(FString::Printf(TEXT("%lf"), FieldDoubleValue1));
-			SceneGraphDescription.Add(TEXT("OriginalUnitsDuration"));
-			SceneGraphDescription.Add(FString::Printf(TEXT("%lf"), FieldDoubleValue2));
+			OutMetaData.Add(TEXT("OriginalUnitsMass"), FString::Printf(TEXT("%lf"), FieldDoubleValue0));
+			OutMetaData.Add(TEXT("OriginalUnitsLength"), FString::Printf(TEXT("%lf"), FieldDoubleValue1));
+			OutMetaData.Add(TEXT("OriginalUnitsDuration"), FString::Printf(TEXT("%lf"), FieldDoubleValue2));
 			break;
 
 		case CT_ATTRIB_ORIGINAL_TOLERANCE:
-			//ITH_ORIGINAL_TOLERANCE_LENGTH, ITH_ORIGINAL_TOLERANCE_MAXCOORD, ITH_ORIGINAL_TOLERANCE_ANGLE
-			break;
 		case CT_ATTRIB_IGES_PARAMETERS:
-			//ITH_IGES_DELIMITOR, ITH_IGES_ENDDELIMITOR, ITH_IGES_EMETTORID, ITH_IGES_CREATION_FILENAME, ITH_IGES_EMETTOR_NAME, ITH_IGES_EMETTOR_VERSION, ITH_IGES_NBOFBITS, ITH_IGES_MAXPOW_SINGLE_PREC, ITH_IGES_NB_DIGITS_SINGLE_PRE, ITH_IGES_MAXPOW_DOUBLE_PREC, ITH_IGES_NB_DIGITS_DOUBLE_PRE, ITH_IGES_RECEPTOR_ID, ITH_IGES_SCALE, ITH_IGES_UNIT, ITH_IGES_UNIT_NAME, ITH_IGES_MAXLINEGRADATIONS, ITH_IGES_MAXLINEWIDTH, ITH_IGES_CREATIONFILEDATE, ITH_IGES_MINIRESOLUTION, ITH_IGES_MAXCOORD, ITH_IGES_AUTHORNAME, ITH_IGES_ORGANIZATION, ITH_IGES_IGESVERSION, ITH_IGES_DRAFTINGSTANDARD, ITH_IGES_MODELDATE, ITH_IGES_APPLICATIONPROTOCOL
-			break;
 		case CT_ATTRIB_READ_V4_MARKER:
-			//ITH_READ_V4_MARKER_VALUE
 			break;
 
 		case CT_ATTRIB_PRODUCT:
 			if (CT_CURRENT_ATTRIB_IO::AskStrField(ITH_PRODUCT_REVISION, FieldStrValue) != IO_OK) break;
 			if (FieldStrValue.IsEmpty()) { break; }
-			SceneGraphDescription.Add(TEXT("ProductRevision"));
-			SceneGraphDescription.Add(FieldStrValue.toUnicode());
+			OutMetaData.Add(TEXT("ProductRevision"), FieldStrValue.toUnicode());
 			if (CT_CURRENT_ATTRIB_IO::AskStrField(ITH_PRODUCT_DEFINITION, FieldStrValue) != IO_OK) break;
 			if (FieldStrValue.IsEmpty()) { break; }
-			SceneGraphDescription.Add(TEXT("ProductDefinition"));
-			SceneGraphDescription.Add(FieldStrValue.toUnicode());
+			OutMetaData.Add(TEXT("ProductDefinition"), FieldStrValue.toUnicode());
 			if (CT_CURRENT_ATTRIB_IO::AskStrField(ITH_PRODUCT_NOMENCLATURE, FieldStrValue) != IO_OK) break;
 			if (FieldStrValue.IsEmpty()) { break; }
-			SceneGraphDescription.Add(TEXT("ProductNomenclature"));
-			SceneGraphDescription.Add(FieldStrValue.toUnicode());
+			OutMetaData.Add(TEXT("ProductNomenclature"), FieldStrValue.toUnicode());
 			if (CT_CURRENT_ATTRIB_IO::AskStrField(ITH_PRODUCT_SOURCE, FieldStrValue) != IO_OK) break;
 			if (FieldStrValue.IsEmpty()) { break; }
-			SceneGraphDescription.Add(TEXT("ProductSource"));
-			SceneGraphDescription.Add(FieldStrValue.toUnicode());
+			OutMetaData.Add(TEXT("ProductSource"), FieldStrValue.toUnicode());
 			if (CT_CURRENT_ATTRIB_IO::AskStrField(ITH_PRODUCT_DESCRIPTION, FieldStrValue) != IO_OK) break;
 			if (FieldStrValue.IsEmpty()) { break; }
-			SceneGraphDescription.Add(TEXT("ProductDescription"));
-			SceneGraphDescription.Add(FieldStrValue.toUnicode());
+			OutMetaData.Add(TEXT("ProductDescription"), FieldStrValue.toUnicode());
 			break;
 
 		case CT_ATTRIB_SIMPLIFY:
-			//ITH_SIMPLIFY_CURRENT_ID, ITH_SIMPLIFY_CONNECTED_ID
-			break;
 		case CT_ATTRIB_MIDFACE:
-			//ITH_MIDFACE_TYPE, ITH_MIDFACE_THICKNESS1, ITH_MIDFACE_THICKNESS2
-			break;
 		case CT_ATTRIB_DEBUG_STRING:
-			//ITH_DEBUG_STRING_VALUE
-			break;
 		case CT_ATTRIB_DEFEATURING:
-			break;
 		case CT_ATTRIB_BREPLINKID:
-			//ITH_BREPLINKID_BRANCHID, ITH_BREPLINKID_FACEID
-			break;
 		case CT_ATTRIB_MARKUPS_REF:
-			//ITH_MARKUPS_REF
-			break;
 		case CT_ATTRIB_COLLISION:
-			//ITH_COLLISION_ID, ITH_COLLISION_MATID
 			break;
+
 		case CT_ATTRIB_EXTERNAL_ID:
 			//ITH_EXTERNAL_ID_VALUE
 			break;
+
 		case CT_ATTRIB_MODIFIER:
-			//ITH_MODIFIER_TYPE, ITH_MODIFIER_INTVALUE, ITH_MODIFIER_DBLVALUE, ITH_MODIFIER_STRVALUE
-			break;
 		case CT_ATTRIB_ORIGINAL_SURF_OLD:
-			break;
 		case CT_ATTRIB_RESULT_BREPLINKID:
 			break;
+
 		case CT_ATTRIB_AREA:
 			//ITH_AREA_VALUE
 			break;
+
 		case CT_ATTRIB_ACIS_SG_PIDNAME:
-			//ITH_ACIS_SG_PIDNAME_BASE_NAME, ITH_ACIS_SG_PIDNAME_TIME_VAL, ITH_ACIS_SG_PIDNAME_INDEX, ITH_ACIS_SG_PIDNAME_COPY_NUM
-			break;
 		case CT_ATTRIB_CURVE_ORIGINAL_BOUNDARY_PARAMS:
-			//ITH_CURVE_ORIGINAL_BOUNDARY_PARAMS_START, ITH_CURVE_ORIGINAL_BOUNDARY_PARAMS_END, ITH_CURVE_ORIGINAL_BOUNDARY_PARAMS_SCALE
 			break;
 
 		case CT_ATTRIB_INTEGER_PARAMETER:
 			if (CT_CURRENT_ATTRIB_IO::AskStrField(ITH_INTEGER_PARAMETER_NAME, FieldName) != IO_OK) break;
 			if (CT_CURRENT_ATTRIB_IO::AskIntField(ITH_INTEGER_PARAMETER_VALUE, FieldIntValue) != IO_OK) break;
-			SceneGraphDescription.Add(FieldName.toUnicode());
-			SceneGraphDescription.Add(FString::FromInt(FieldIntValue));
+			OutMetaData.Add(FieldName.toUnicode(), FString::FromInt(FieldIntValue));
 			break;
 
 		case CT_ATTRIB_DOUBLE_PARAMETER:
 			if (CT_CURRENT_ATTRIB_IO::AskStrField(ITH_DOUBLE_PARAMETER_NAME, FieldName) != IO_OK) break;
 			if (CT_CURRENT_ATTRIB_IO::AskDblField(ITH_DOUBLE_PARAMETER_VALUE, FieldDoubleValue0) != IO_OK) break;
-			SceneGraphDescription.Add(FieldName.toUnicode());
-			SceneGraphDescription.Add(FString::Printf(TEXT("%lf"), FieldDoubleValue0));
+			OutMetaData.Add(FieldName.toUnicode(), FString::Printf(TEXT("%lf"), FieldDoubleValue0));
 			break;
 
 		case CT_ATTRIB_STRING_PARAMETER:
 			if (CT_CURRENT_ATTRIB_IO::AskStrField(ITH_STRING_PARAMETER_NAME, FieldName) != IO_OK) break;
 			if (CT_CURRENT_ATTRIB_IO::AskStrField(ITH_STRING_PARAMETER_VALUE, FieldStrValue) != IO_OK) break;
 			if (FieldStrValue.IsEmpty()) { break; }
-			SceneGraphDescription.Add(FieldName.toUnicode());
-			SceneGraphDescription.Add(FieldStrValue.toUnicode());
+			OutMetaData.Add(FieldName.toUnicode(), FieldStrValue.toUnicode());
 			break;
 
 		case CT_ATTRIB_PARAMETER_ARRAY:
@@ -1539,37 +1271,34 @@ void FCoreTechFileParser::ReadNodeMetaDatas(CT_OBJECT_ID NodeId)
 		case CT_ATTRIB_SAVE_OPTION:
 			if (CT_CURRENT_ATTRIB_IO::AskStrField(ITH_SAVE_OPTION_AUTHOR, FieldStrValue) != IO_OK) break;
 			if (FieldStrValue.IsEmpty()) { break; }
-			//NodeAttributeSet.Add(TEXT("SaveOptionAuthor"));
-			SceneGraphDescription.Add(FieldStrValue.toUnicode());
+			OutMetaData.Add(TEXT("SaveOptionAuthor"), FieldStrValue.toUnicode());
+
 			if (CT_CURRENT_ATTRIB_IO::AskStrField(ITH_SAVE_OPTION_ORGANIZATION, FieldStrValue) != IO_OK) break;
 			if (FieldStrValue.IsEmpty()) { break; }
-			//NodeAttributeSet.Add(TEXT("SaveOptionOrganization"));
-			SceneGraphDescription.Add(FieldStrValue.toUnicode());
+			OutMetaData.Add(TEXT("SaveOptionOrganization"), FieldStrValue.toUnicode());
+	
 			if (CT_CURRENT_ATTRIB_IO::AskStrField(ITH_SAVE_OPTION_FILE_DESCRIPTION, FieldStrValue) != IO_OK) break;
 			if (FieldStrValue.IsEmpty()) { break; }
-			//NodeAttributeSet.Add(TEXT("SaveOptionFileDescription"));
-			SceneGraphDescription.Add(FieldStrValue.toUnicode());
+			OutMetaData.Add(TEXT("SaveOptionFileDescription"), FieldStrValue.toUnicode());
+
 			if (CT_CURRENT_ATTRIB_IO::AskStrField(ITH_SAVE_OPTION_AUTHORISATION, FieldStrValue) != IO_OK) break;
 			if (FieldStrValue.IsEmpty()) { break; }
-			//NodeAttributeSet.Add(TEXT("SaveOptionAuthorisation"));
-			SceneGraphDescription.Add(FieldStrValue.toUnicode());
+			OutMetaData.Add(TEXT("SaveOptionAuthorisation"), FieldStrValue.toUnicode());
+
 			if (CT_CURRENT_ATTRIB_IO::AskStrField(ITH_SAVE_OPTION_PREPROCESSOR, FieldStrValue) != IO_OK) break;
 			if (FieldStrValue.IsEmpty()) { break; }
-			//NodeAttributeSet.Add(TEXT("SaveOptionPreprocessor"));
-			SceneGraphDescription.Add(FieldStrValue.toUnicode());
+			OutMetaData.Add(TEXT("SaveOptionPreprocessor"), FieldStrValue.toUnicode());
 			break;
 
 		case CT_ATTRIB_ORIGINAL_ID:
 			GetAttributeValue(AttributeType, ITH_ORIGINAL_ID_VALUE, FieldValue);
-			SceneGraphDescription.Add(TEXT("OriginalId"));
-			SceneGraphDescription.Add(FieldValue);
+			OutMetaData.Add(TEXT("OriginalId"), FieldValue);
 			break;
 
 		case CT_ATTRIB_ORIGINAL_ID_STRING:
 			if (CT_CURRENT_ATTRIB_IO::AskStrField(ITH_ORIGINAL_ID_VALUE_STRING, FieldStrValue) != IO_OK) break;
 			if (FieldStrValue.IsEmpty()) { break; }
-			SceneGraphDescription.Add(TEXT("OriginalIdStr"));
-			SceneGraphDescription.Add(FieldStrValue.toUnicode());
+			OutMetaData.Add(TEXT("OriginalIdStr"), FieldStrValue.toUnicode());
 			break;
 
 		case CT_ATTRIB_COLOR_RGB_DOUBLE:
@@ -1577,123 +1306,73 @@ void FCoreTechFileParser::ReadNodeMetaDatas(CT_OBJECT_ID NodeId)
 			if (CT_CURRENT_ATTRIB_IO::AskDblField(ITH_ATTRIB_COLOR_G_DOUBLE, FieldDoubleValue1) != IO_OK) break;
 			if (CT_CURRENT_ATTRIB_IO::AskDblField(ITH_ATTRIB_COLOR_B_DOUBLE, FieldDoubleValue2) != IO_OK) break;
 			FieldValue = FString::Printf(TEXT("%lf"), FieldDoubleValue0) + TEXT(", ") + FString::Printf(TEXT("%lf"), FieldDoubleValue1) + TEXT(", ") + FString::Printf(TEXT("%lf"), FieldDoubleValue2);
-			//OutMetaDataSet.Add(TEXT("ColorRGBDouble"))
-			//OutMetaDataSet.Add(FieldValue);
+			//OutMetaData.Add(TEXT("ColorRGBDouble"), FieldValue);
 			break;
 
 		case CT_ATTRIB_REVERSE_COLORID:
-			break;
 		case CT_ATTRIB_INITIAL_FILTER:
-			break;
 		case CT_ATTRIB_ORIGINAL_SURF:
-			break;
 		case CT_ATTRIB_LINKMANAGER_BRANCH_FACE:
-			//ITH_LINKMANAGER_BRANCH_FACE_PART_ID, ITH_LINKMANAGER_BRANCH_FACE_SOLID_ID, ITH_LINKMANAGER_BRANCH_FACE_BRANCH_ID, ITH_LINKMANAGER_BRANCH_FACE_FACE_ID
-			break;
 		case CT_ATTRIB_LINKMANAGER_PMI:
-			//ITH_LINKMANAGER_PMI_ID, ITH_LINKMANAGER_PMI_SENSE, ITH_LINKMANAGER_PMI_INSTANCE
-			break;
 		case CT_ATTRIB_NULL:
-			//ITH_NULL_TITLE
-			break;
 		case CT_ATTRIB_MEASURE_VALIDATION_ATTRIBUTE:
-			//ITH_MEASURE_VALIDATION_NAME
-			//ITH_MEASURE_VALIDATION_UNIT
-			//ITH_MEASURE_VALIDATION_VALUE
 			break;
 
 		case CT_ATTRIB_INTEGER_VALIDATION_ATTRIBUTE:
 			if (CT_CURRENT_ATTRIB_IO::AskStrField(ITH_INTEGER_VALIDATION_NAME, FieldName) != IO_OK) break;
 			if (CT_CURRENT_ATTRIB_IO::AskIntField(ITH_INTEGER_VALIDATION_VALUE, FieldIntValue) != IO_OK) break;
-			SceneGraphDescription.Add(FieldName.toUnicode());
-			SceneGraphDescription.Add(FString::FromInt(FieldIntValue));
+			OutMetaData.Add(FieldName.toUnicode(), FString::FromInt(FieldIntValue));
 			break;
 
 		case CT_ATTRIB_DOUBLE_VALIDATION_ATTRIBUTE:
 			if (CT_CURRENT_ATTRIB_IO::AskStrField(ITH_DOUBLE_VALIDATION_NAME, FieldName) != IO_OK) break;
 			if (CT_CURRENT_ATTRIB_IO::AskDblField(ITH_DOUBLE_VALIDATION_VALUE, FieldDoubleValue0) != IO_OK) break;
-			SceneGraphDescription.Add(FieldName.toUnicode());
-			SceneGraphDescription.Add(FString::Printf(TEXT("%lf"), FieldDoubleValue0));
+			OutMetaData.Add(FieldName.toUnicode(), FString::Printf(TEXT("%lf"), FieldDoubleValue0));
 			break;
 
 		case CT_ATTRIB_STRING_VALIDATION_ATTRIBUTE:
 			if (CT_CURRENT_ATTRIB_IO::AskStrField(ITH_STRING_VALIDATION_NAME, FieldName) != IO_OK) break;
 			if (CT_CURRENT_ATTRIB_IO::AskStrField(ITH_STRING_VALIDATION_VALUE, FieldStrValue) != IO_OK) break;
 			if (FieldStrValue.IsEmpty()) { break; }
-			SceneGraphDescription.Add(FieldName.toUnicode());
-			SceneGraphDescription.Add(FieldStrValue.toUnicode());
+			OutMetaData.Add(FieldName.toUnicode(), FieldStrValue.toUnicode());
 			break;
 
 		case CT_ATTRIB_BOUNDING_BOX:
 			//ITH_BOUNDING_BOX_XMIN, ITH_BOUNDING_BOX_YMIN, ITH_BOUNDING_BOX_ZMIN, ITH_BOUNDING_BOX_XMAX, ITH_BOUNDING_BOX_YMAX, ITH_BOUNDING_BOX_ZMAX
 			break;
+
 		case CT_ATTRIB_DATABASE:
-			//ITH_DATABASE_NAME, ITH_DATABASE_TABLE
-			break;
 		case CT_ATTRIB_CURVE_FONT:
-			//ITH_ATTRIB_CURVE_FONT_VALUE
-			break;
 		case CT_ATTRIB_CURVE_WEIGHT:
-			//ITH_ATTRIB_CURVE_WEIGHT_VALUE
-			break;
 		case CT_ATTRIB_COMPARE_TOPO:
-			break;
 		case CT_ATTRIB_MONIKER_GUID_TABLE:
-			//ITH_MONIKER_GUID_TABLE_NAME, ITH_MONIKER_GUID_TABLE_GUID, ITH_MONIKER_GUID_TABLE_INDEX, ITH_MONIKER_GUID_TABLE_APPLICATION_NAME
-			break;
 		case CT_ATTRIB_MONIKER_DATA:
-			//ITH_MONIKER_DATA_NAME, ITH_MONIKER_DATA_TABLE_INDEX, ITH_MONIKER_DATA_ENTITY_ID, ITH_MONIKER_DATA_LABEL, ITH_MONIKER_DATA_BODY_GUID
-			break;
 		case CT_ATTRIB_MONIKER_BODY_ID:
-			//ITH_MONIKER_BODY_ID_NAME, ITH_MONIKER_BODY_ID_ENTITY_ID, ITH_MONIKER_BODY_ID_APPLICATION_NAME
-			break;
 		case CT_ATTRIB_NO_INSTANCE:
 			break;
 
 		case CT_ATTRIB_GROUPNAME:
 			if (CT_CURRENT_ATTRIB_IO::AskStrField(ITH_GROUPNAME_VALUE, FieldStrValue) != IO_OK) break;
 			if (FieldStrValue.IsEmpty()) { break; }
-			SceneGraphDescription.Add(TEXT("GroupName"));
-			SceneGraphDescription.Add(FieldStrValue.toUnicode());
+			OutMetaData.Add(TEXT("GroupName"), FieldStrValue.toUnicode());
 			break;
 
 		case CT_ATTRIB_ANALYZE_ID:
-			//ITH_ANALYZE_OBJECT_ID
-			break;
 		case CT_ATTRIB_ANALYZER_DISPLAY_MODE:
-			//ITH_ANALYZER_DISPLAY_MODE_VALUE
-			break;
 		case CT_ATTRIB_ANIMATION_ID:
-			//ITH_ANIMATION_ID_VALUE
-			break;
 		case CT_ATTRIB_PROJECTED_SURFACE_ID:
-			//ITH_PROJECTED_SURFACE_ID_VALUE
-			break;
 		case CT_ATTRIB_ANALYZE_LINK:
-			//ITH_ANALYZE_LINK_ID
-			break;
 		case CT_ATTRIB_TOPO_EVENT_ID:
-			//ITH_TOPO_EVENT_ID_VALUE
-			break;
 		case CT_ATTRIB_ADDITIVE_MANUFACTURING:
-			//ITH_AM_MODE, ITH_AM_ROUGHNESS_MIN, ITH_AM_SECTION_DIST, ITH_AM_FUNCTIONAL_TYPE, ITH_AM_LATTICE_TYPE
-			break;
 		case CT_ATTRIB_MOLDING_RESULT:
-			//ITH_MOLDING_MODE, ITH_MOLDING_MEASURE_MIN, ITH_MOLDING_MEASURE_MAX
-			break;
 		case CT_ATTRIB_AMF_ID:
-			//ITH_AMF_ID_VALUE, ITH_AMF_TYPE
-			break;
 		case CT_ATTRIB_PARAMETER_LINK:
-			//ITH_PARAMETER_LINK_ID, ITH_PARAMETER_LINK_PARENT_ID, ITH_PARAMETER_LINK_SENSE
 			break;
+
 		default:
 			break;
 		}
 	}
-
-	AttributeNameLigne = "M " + FString::FromInt(SceneGraphDescription.Num() - AttributIndex);
-	SceneGraphDescription[AttributIndex-1] = *AttributeNameLigne;
 }
 
 uint32 GetSize(CT_TESS_DATA_TYPE type)
diff --git a/Engine/Plugins/Enterprise/DatasmithCADImporter/Source/CADInterfaces/Public/CADSceneGraph.h b/Engine/Plugins/Enterprise/DatasmithCADImporter/Source/CADInterfaces/Public/CADSceneGraph.h
new file mode 100644
--- /dev/null
+++ b/Engine/Plugins/Enterprise/DatasmithCADImporter/Source/CADInterfaces/Public/CADSceneGraph.h
@@ -0,0 +1,154 @@
+// Copyright 1998-2019 Epic Games, Inc. All Rights Reserved.
+#pragma once
+
+#include "CoreMinimal.h"
+
+#include "CADData.h"
+
+class FArchive;
+
+namespace CADLibrary
+{
+
+class CADINTERFACES_API ICADArchiveObject
+{
+public:
+	ICADArchiveObject(CadId Id = 0)
+		: ObjectId(Id)
+	{
+	}
+
+	virtual ~ICADArchiveObject() = default;
+
+public:
+	uint32 ObjectId;
+	TMap<FString, FString> MetaData;
+};
+
+class CADINTERFACES_API FArchiveInstance : public ICADArchiveObject
+{
+public:
+	FArchiveInstance(CadId Id = 0)
+		: ICADArchiveObject(Id)
+	{
+	}
+
+	friend FArchive& operator<<(FArchive& Ar, FArchiveInstance& C);
+
+public:
+	FMatrix TransformMatrix = FMatrix::Identity;
+	CadId ReferenceNodeId = 0;
+	bool bIsExternalRef = false;
+	FString ExternalRef;
+};
+
+class CADINTERFACES_API FArchiveComponent : public ICADArchiveObject
+{
+public:
+	FArchiveComponent(CadId Id = 0)
+		: ICADArchiveObject(Id)
+	{
+	}
+
+	friend FArchive& operator<<(FArchive& Ar, FArchiveComponent& C);
+
+public:
+	TArray<CadId> Children;
+};
+
+class CADINTERFACES_API FArchiveUnloadedComponent : public FArchiveComponent
+{
+public:
+	FArchiveUnloadedComponent(CadId Id = 0)
+		: FArchiveComponent(Id)
+	{
+	}
+
+	friend FArchive& operator<<(FArchive& Ar, FArchiveUnloadedComponent& C);
+
+public:
+	FString FileName;
+	FString FileType;
+};
+
+class CADINTERFACES_API FArchiveBody : public ICADArchiveObject
+{
+public:
+	FArchiveBody(CadId Id = 0)
+		: ICADArchiveObject(Id)
+	{
+	}
+
+	friend FArchive& operator<<(FArchive& Ar, FArchiveBody& C);
+
+public:
+	CADUUID MeshActorName;
+
+	TSet<MaterialId> MaterialFaceSet;
+	TSet<ColorId> ColorFaceSet;
+
+};
+
+class CADINTERFACES_API FArchiveColor
+{
+public:
+	FArchiveColor(ColorId Id = 0)
+		: ObjectId(Id)
+	{
+	}
+
+	friend FArchive& operator<<(FArchive& Ar, FArchiveColor& C);
+
+public:
+	ColorId ObjectId;
+	FColor Color;
+	CADUUID UEMaterialName;
+};
+
+class CADINTERFACES_API FArchiveMaterial
+{
+public:
+	FArchiveMaterial(MaterialId Id = 0)
+		: ObjectId(Id)
+	{
+	}
+
+	friend FArchive& operator<<(FArchive& Ar, FArchiveMaterial& C);
+
+public:
+	MaterialId ObjectId;
+	CADUUID UEMaterialName;
+	FCADMaterial Material;
+};
+
+class CADINTERFACES_API FArchiveMockUp
+{
+public:
+	friend FArchive& operator<<(FArchive& Ar, FArchiveMockUp& C);
+
+public:
+	FString CADFile;
+	FString SceneGraphArchive;
+	FString FullPath;
+	TSet<FString> ExternalRefSet;
+
+	TMap<ColorId, FArchiveColor> ColorHIdToColor;
+	TMap<MaterialId, FArchiveMaterial> MaterialHIdToMaterial;
+
+	TArray<FArchiveBody> BodySet;
+	TArray<FArchiveComponent> ComponentSet;
+	TArray<FArchiveUnloadedComponent> UnloadedComponentSet;
+	TArray<FArchiveInstance> InstanceSet;
+
+	TMap<CadId, int32> CADIdToBodyIndex;
+	TMap<CadId, int32> CADIdToComponentIndex;
+	TMap<CadId, int32> CADIdToUnloadedComponentIndex;
+	TMap<CadId, int32> CADIdToInstanceIndex;
+};
+
+CADINTERFACES_API void SerializeMockUp(FArchiveMockUp& Object, const TCHAR* Filename);
+CADINTERFACES_API void DeserializeMockUpFile(const TCHAR* Filename, FArchiveMockUp& MockUp);
+
+}
+
+
diff --git a/Engine/Plugins/Enterprise/DatasmithCADImporter/Source/CADInterfaces/Public/CoreTechFileParser.h b/Engine/Plugins/Enterprise/DatasmithCADImporter/Source/CADInterfaces/Public/CoreTechFileParser.h
--- a/Engine/Plugins/Enterprise/DatasmithCADImporter/Source/CADInterfaces/Public/CoreTechFileParser.h
+++ b/Engine/Plugins/Enterprise/DatasmithCADImporter/Source/CADInterfaces/Public/CoreTechFileParser.h
@@ -5,6 +5,7 @@
 
 #include "CADData.h"
 #include "CADOptions.h"
+#include "CADSceneGraph.h"
 #include "Containers/Map.h"
 #include "Containers/Queue.h"
 #include "Misc/Paths.h"
@@ -18,6 +19,7 @@
 #include "kernel_io/attribute_io/attribute_enum.h"
 #include "kernel_io/ct_types.h"
 #include "kernel_io/kernel_io_type.h"
+#include "kernel_io/material_io/material_io.h"
 #endif // CAD_INTERFACE
 
 
@@ -38,76 +40,72 @@ enum class ECoretechParsingResult
 
 #ifdef CAD_INTERFACE
 
-CADINTERFACES_API void GetColor(uint32 ColorHash, FColor& OutMaterial);
+CADINTERFACES_API bool GetColor(uint32 ColorHash, FColor& OutColor);
 CADINTERFACES_API bool GetMaterial(uint32 MaterialID, FCADMaterial& OutMaterial);
 
-CADINTERFACES_API int32 BuildColorHash(uint32 ColorHId);
-CADINTERFACES_API int32 BuildMaterialHash(uint32 MaterialId);
-
-CADINTERFACES_API uint32 GetFaceTessellation(CT_OBJECT_ID FaceID, TArray<FTessellationData>& FaceTessellationSet, int32& OutRawDataSize, const float ScaleFactor);
-CADINTERFACES_API uint32 GetBodiesTessellations(TArray<CT_OBJECT_ID>& BodySet, TArray<FTessellationData>& FaceTessellations, TMap<uint32, uint32>& MaterialIdToMaterialHashMap, int32& OutRawDataSize, const FImportParameters&);
 CADINTERFACES_API uint32 GetBodiesFaceSetNum(TArray<CT_OBJECT_ID>& BodySet);
 CADINTERFACES_API void GetCTObjectDisplayDataIds(CT_OBJECT_ID ObjectID, FObjectDisplayDataId& Material);
-CADINTERFACES_API void GetBodiesMaterials(TArray<CT_OBJECT_ID>& BodySet, TMap<uint32, uint32>& MaterialIdToHash, bool bPreferPartData);
 
-uint32 GetFileHash(const FString& FileName, const FFileStatData& FileStatData, const FString& Config);
 uint32 GetGeomFileHash(const uint32 InSGHash, CADLibrary::FImportParameters& ImportParam);
 
-/**
- * According to the preference (bPreferPartData), select the best material to use
- */
-CADINTERFACES_API void SetFaceMainMaterial(FObjectDisplayDataId& InFaceMaterial, FObjectDisplayDataId& InBodyMaterial, TMap<uint32, uint32>& MaterialIdToMaterialHashMap, FTessellationData& OutFaceTessellations);
-
 uint32 GetSize(CT_TESS_DATA_TYPE type);
 
-
 class CADINTERFACES_API FCoreTechFileParser
 {
 public:
 	using EProcessResult = ECoretechParsingResult;
 
-	FCoreTechFileParser(const FString& InCTFullPath, const FString& InCachePath, const FImportParameters& ImportParams, bool bPreferBodyData, bool bScaleUVMap);
+	FCoreTechFileParser(const FString& InCTFullPath, const FString& InCachePath, const FImportParameters& ImportParams);
 
 	double GetMetricUnit() const { return 0.01; }
 
 	EProcessResult ProcessFile();
+	void GetBodyTessellation(CT_OBJECT_ID BodyId, FBodyMesh& OutBodyMesh, const FImportParameters& ImportParams, uint32 ParentMaterialHash);
+
 
 	const TSet<FString>& GetExternalRefSet()
 	{
-		return ExternalRefSet;
+		return MockUpDescription.ExternalRefSet;
 	}
 
 	const FString& GetSceneGraphFile()
 	{
-		return SceneGraphFile;
+		return MockUpDescription.SceneGraphArchive;
 	}
-	const FString& GetMeshFile()
+	const FString& GetMeshFileName()
 	{
-		return MeshFile;
+		return MeshArchiveFile;
 	}
-	const FString& GetFileName()
+	const FString& GetCADFileName()
 	{
 		return CADFile;
 	}
 
 private:
 	EProcessResult ReadFileWithKernelIO();
-	bool ReadNode(CT_OBJECT_ID NodeId);
-	bool ReadInstance(CT_OBJECT_ID NodeId);
-	bool ReadComponent(CT_OBJECT_ID NodeId);
+	bool ReadNode(CT_OBJECT_ID NodeId, uint32 ParentMaterialHash);
+	bool ReadInstance(CT_OBJECT_ID NodeId, uint32 ParentMaterialHash);
+	bool ReadComponent(CT_OBJECT_ID NodeId, uint32 ParentMaterialHash);
 	bool ReadUnloadedComponent(CT_OBJECT_ID NodeId);
-	bool ReadBody(CT_OBJECT_ID NodeId);
+	bool ReadBody(CT_OBJECT_ID NodeId, uint32 ParentMaterialHash);
+
+	void LoadSceneGraphArchive(const FString& SceneGraphFilePath);
 
 	uint32 GetMaterialNum();
-	void ReadColor();
-	void ReadMaterial();
+	void ReadMaterials();
+
+	uint32 GetObjectMaterial(ICADArchiveObject& Object);
+	FArchiveColor& FindOrAddColor(uint32 ColorHId);
+	FArchiveMaterial& FindOrAddMaterial(CT_MATERIAL_ID MaterialId);
+	void SetFaceMainMaterial(FObjectDisplayDataId& InFaceMaterial, FObjectDisplayDataId& InBodyMaterial, FBodyMesh& BodyMesh, int32 FaceIndex);
 
-	void ReadNodeMetaDatas(CT_OBJECT_ID NodeId);
+	void ReadNodeMetaData(CT_OBJECT_ID NodeId, TMap<FString, FString>& OutMetaData);
 
 	CT_FLAGS SetCoreTechImportOption(const FString& MainFileExt);
 	void GetAttributeValue(CT_ATTRIB_TYPE attrib_type, int ith_field, FString& value);
 
-	void ExportFileSceneGraph();
+	void ExportSceneGraphFile();
+	void ExportMeshArchiveFile();
 
 protected:
 	const FString& CachePath;
@@ -118,20 +116,13 @@ protected:
 	FString FileConfiguration;
 	FString NodeConfiguration;
 
-	FString SceneGraphFile;
-	FString RawDataGeom;
-	FString MeshFile;
-
-	TSet<FString> ExternalRefSet;
-
-	TArray<FString> SceneGraphDescription;
-	TMap<uint32, uint32> CTIdToRawLineMap;
+	FArchiveMockUp MockUpDescription;
 
-	TMap<uint32, uint32> MaterialIdToMaterialHashMap;
+	FString MeshArchiveFilePath;
+	FString MeshArchiveFile;
+	TArray<FBodyMesh> BodyMeshes;
 
 	bool bNeedSaveCTFile;
-	bool bPreferBodyData;
-	bool bScaleUVMap;
 
 	const FImportParameters& ImportParameters;
 };
diff --git a/Engine/Plugins/Enterprise/DatasmithCADImporter/Source/CADLibrary/Private/CoreTechHelper.cpp b/Engine/Plugins/Enterprise/DatasmithCADImporter/Source/CADLibrary/Private/CoreTechHelper.cpp
--- a/Engine/Plugins/Enterprise/DatasmithCADImporter/Source/CADLibrary/Private/CoreTechHelper.cpp
+++ b/Engine/Plugins/Enterprise/DatasmithCADImporter/Source/CADLibrary/Private/CoreTechHelper.cpp
@@ -276,7 +276,7 @@ bool FillMesh(const FMeshParameters& MeshParameters, const FImportParameters& Im
 		for (FTessellationData& Tessellation : FaceTessellations)
 		{
 			// Get the polygonGroup
-			const FPolygonGroupID* PolygonGroupID = MaterialToPolygonGroupMapping.Find(Tessellation.MaterialHash);
+			const FPolygonGroupID* PolygonGroupID = MaterialToPolygonGroupMapping.Find(Tessellation.ColorName);
 			if (PolygonGroupID == nullptr)
 			{
 				continue;
@@ -391,47 +391,35 @@ bool FillMesh(const FMeshParameters& MeshParameters, const FImportParameters& Im
 			}
 		}
 	}
-
-
 	return true;
 }
 
-bool ConvertCTBodySetToMeshDescription(const FImportParameters& ImportParams, const FMeshParameters& MeshParameters, TArray<CT_OBJECT_ID>& BodySet, FMeshDescription& MeshDescription)
-{
-	TArray<FTessellationData> FaceTessellationSet;
-	TMap<uint32, uint32> MaterialIdToMaterialHash;
-
-	int32 BodyRawSize = 0;
-	uint32 TriangleCount = GetBodiesTessellations(BodySet, (TArray<FTessellationData>&) FaceTessellationSet, MaterialIdToMaterialHash, BodyRawSize, ImportParams);
-
-	return ConvertCTBodySetToMeshDescription(ImportParams, MeshParameters, TriangleCount, FaceTessellationSet, MaterialIdToMaterialHash, MeshDescription);
-}
-
-bool ConvertCTBodySetToMeshDescription(const FImportParameters& ImportParams, const FMeshParameters& MeshParameters, uint32 TriangleCount, TArray<FTessellationData>& FaceTessellationSet, TMap<uint32, uint32>& MaterialIdToMaterialHash, FMeshDescription& MeshDescription)
+bool ConvertCTBodySetToMeshDescription(const FImportParameters& ImportParams, const FMeshParameters& MeshParameters, FBodyMesh& Body, FMeshDescription& MeshDescription)
 {
 	// Ref. CreateMesh(UDatasmithCADImportOptions* CADOptions, CTMesh& Mesh)
 	MeshDescription.EdgeAttributes().RegisterAttribute<bool>(MeshAttribute::Edge::IsUVSeam, 1, false);
 
 	// in a closed big mesh VertexCount ~ TriangleCount / 2, EdgeCount ~ 1.5* TriangleCount
-	MeshDescription.ReserveNewVertexInstances(TriangleCount*3);
-	MeshDescription.ReserveNewPolygons(TriangleCount);
-	MeshDescription.ReserveNewEdges(TriangleCount*3);
+	MeshDescription.ReserveNewVertexInstances(Body.TriangleCount*3);
+	MeshDescription.ReserveNewPolygons(Body.TriangleCount);
+	MeshDescription.ReserveNewEdges(Body.TriangleCount*3);
 
 	// CoreTech is generating position duplicates. make sure to remove them before filling the mesh description
 	TArray<FVertexID> RemapVertexPosition;
-	FillVertexPosition(ImportParams, MeshParameters, TriangleCount, FaceTessellationSet, MeshDescription);
+	FillVertexPosition(ImportParams, MeshParameters, Body.TriangleCount, Body.Faces, MeshDescription);
 
 	TMap<uint32, FPolygonGroupID> MaterialToPolygonGroupMapping;
-	for (const FTessellationData& FaceTessellation : FaceTessellationSet)
+	for (const FTessellationData& FaceTessellation : Body.Faces)
 	{
-		MaterialToPolygonGroupMapping.Add(FaceTessellation.MaterialHash, FPolygonGroupID::Invalid);
+		// we assume that face has only color
+		MaterialToPolygonGroupMapping.Add(FaceTessellation.ColorName, FPolygonGroupID::Invalid);
 	}
 
 	// Add the mesh's materials as polygon groups
 	TPolygonGroupAttributesRef<FName> PolygonGroupImportedMaterialSlotNames = MeshDescription.PolygonGroupAttributes().GetAttributesRef<FName>(MeshAttribute::PolygonGroup::ImportedMaterialSlotName);
 	UpdatePolygonGroup(MaterialToPolygonGroupMapping, PolygonGroupImportedMaterialSlotNames, MeshDescription);
 
-	if (!FillMesh(MeshParameters, ImportParams, FaceTessellationSet, MaterialToPolygonGroupMapping, MeshDescription))
+	if (!FillMesh(MeshParameters, ImportParams, Body.Faces, MaterialToPolygonGroupMapping, MeshDescription))
 	{
 		return false;
 	}
@@ -445,8 +433,6 @@ bool ConvertCTBodySetToMeshDescription(const FImportParameters& ImportParams, co
 	return MeshDescription.Polygons().Num() > 0;
 }
 
-
-
 double Distance(CT_COORDINATE Point1, CT_COORDINATE Point2)
 {
 	return sqrt((Point2.xyz[0] - Point1.xyz[0]) * (Point2.xyz[0] - Point1.xyz[0]) + (Point2.xyz[1] - Point1.xyz[1]) * (Point2.xyz[1] - Point1.xyz[1]) + (Point2.xyz[2] - Point1.xyz[2]) * (Point2.xyz[2] - Point1.xyz[2]));
@@ -577,65 +563,6 @@ void ScaleUV(CT_OBJECT_ID FaceID, FTessellationData& Tessellation, UVType Scale)
 	}
 }
 
-void SetFaceMainMaterial(FObjectDisplayDataId& InFaceMaterial, FObjectDisplayDataId& InBodyMaterial, TMap<uint32, uint32> MaterialIdToMaterialHashMap, FTessellationData& OutFaceTessellations)
-{
-	uint32 MaterialId = 0; // either CT Material ID or CT Color/Alpha ID (ColorId & Alpha << 24)
-	uint32 MaterialHash = 0;
-	if (InFaceMaterial.MaterialId > 0)
-	{
-		MaterialId = InFaceMaterial.MaterialId;
-		if (uint32* MHash = MaterialIdToMaterialHashMap.Find(InFaceMaterial.MaterialId))
-		{
-			MaterialHash = *MHash;
-		}
-		else
-		{
-			MaterialHash = BuildMaterialHash(InFaceMaterial.MaterialId);
-		}
-	}
-	else if (InFaceMaterial.ColorHId > 0)
-	{
-		MaterialId = InFaceMaterial.ColorHId;
-		if (uint32* MHash = MaterialIdToMaterialHashMap.Find(InFaceMaterial.ColorHId))
-		{
-			MaterialHash = *MHash;
-		}
-		else
-		{
-			MaterialHash = BuildColorHash(InFaceMaterial.ColorHId);
-		}
-	}
-	else if (InBodyMaterial.MaterialId > 0)
-	{
-		MaterialId = InBodyMaterial.MaterialId;
-		if (uint32* MHash = MaterialIdToMaterialHashMap.Find(InBodyMaterial.MaterialId))
-		{
-			MaterialHash = *MHash;
-		}
-		else
-		{
-			MaterialHash = BuildMaterialHash(InBodyMaterial.MaterialId);
-		}
-	}
-	else if (InBodyMaterial.ColorHId > 0)
-	{
-		MaterialId = InBodyMaterial.ColorHId;
-		if (uint32* MHash = MaterialIdToMaterialHashMap.Find(InBodyMaterial.ColorHId))
-		{
-			MaterialHash = *MHash;
-		}
-		else
-		{
-			MaterialHash = BuildColorHash(InBodyMaterial.ColorHId);
-		}
-	}
-
-	// set output
-	MaterialIdToMaterialHashMap.Add(MaterialId, MaterialHash);
-	OutFaceTessellations.MaterialId = MaterialId;
-	OutFaceTessellations.MaterialHash = MaterialHash;
-}
-
 TSharedPtr<IDatasmithUEPbrMaterialElement> CreateDefaultUEPbrMaterial()
 {
 	// Take the Material diffuse color and connect it to the BaseColor of a UEPbrMaterial
@@ -654,7 +581,7 @@ TSharedPtr<IDatasmithUEPbrMaterialElement> CreateDefaultUEPbrMaterial()
 
 TSharedPtr<IDatasmithUEPbrMaterialElement> CreateUEPbrMaterialFromColor(const FColor& InColor)
 {
-	FString Name = FString::FromInt(BuildColorHash(InColor));
+	FString Name = FString::FromInt(BuildColorName(InColor));
 	FString Label = FString::Printf(TEXT("color_%02x%02x%02x%02x"), InColor.R, InColor.G, InColor.B, InColor.A);
 
 	// Take the Material diffuse color and connect it to the BaseColor of a UEPbrMaterial
@@ -690,7 +617,7 @@ TSharedPtr<IDatasmithUEPbrMaterialElement> CreateUEPbrMaterialFromColor(const FC
 
 TSharedPtr<IDatasmithUEPbrMaterialElement> CreateUEPbrMaterialFromMaterial(FCADMaterial& InMaterial, TSharedRef<IDatasmithScene> Scene)
 {
-	FString Name = FString::FromInt(BuildMaterialHash(InMaterial));
+	FString Name = FString::FromInt(BuildMaterialName(InMaterial));
 
 	// Take the Material diffuse color and connect it to the BaseColor of a UEPbrMaterial
 	TSharedRef<IDatasmithUEPbrMaterialElement> MaterialElement = FDatasmithSceneFactory::CreateUEPbrMaterial(*Name);
@@ -801,24 +728,21 @@ CT_IO_ERROR Tessellate(CT_OBJECT_ID MainObjectId, const FImportParameters& Impor
 		return Result;
 
 	FImportParameters FixedParams = GetAdjustedTessellationParameters(MainObjectId, ImportParams);
-
 	SetCoreTechTessellationState(FixedParams);
 
-	TArray<FTessellationData> FaceTessellations;
-	TArray<uint32> FaceColors;
+	FString FullPath;
+	FString CachePath;
+	FCoreTechFileParser Parser = FCoreTechFileParser(FullPath, CachePath, ImportParams);
 
-	size_t VertexCount = 0;
-	size_t IndexCount = 0;
+	FBodyMesh BodyMesh;
+	BodyMesh.BodyID = 1;
 
-	TArray<CT_OBJECT_ID> BodySet;
-	BodySet.Reserve(Objects.Count());
-	while (CT_OBJECT_ID CurrentObj = Objects.IteratorIter())
+	while (CT_OBJECT_ID BodyId = Objects.IteratorIter())
 	{
-		BodySet.Add(CurrentObj);
+		Parser.GetBodyTessellation(BodyId, BodyMesh, ImportParams, 0);
 	}
 
-	bool bPreferBodyData = true;
-	bool bTessellated = ConvertCTBodySetToMeshDescription(ImportParams, MeshParameters, BodySet, MeshDesc);
+	bool bTessellated = ConvertCTBodySetToMeshDescription(ImportParams, MeshParameters, BodyMesh, MeshDesc);
 
 	CheckedCTError ConversionResult;
 	if (!bTessellated)
diff --git a/Engine/Plugins/Enterprise/DatasmithCADImporter/Source/CADLibrary/Private/CoreTechMeshLoader.cpp b/Engine/Plugins/Enterprise/DatasmithCADImporter/Source/CADLibrary/Private/CoreTechMeshLoader.cpp
--- a/Engine/Plugins/Enterprise/DatasmithCADImporter/Source/CADLibrary/Private/CoreTechMeshLoader.cpp
+++ b/Engine/Plugins/Enterprise/DatasmithCADImporter/Source/CADLibrary/Private/CoreTechMeshLoader.cpp
@@ -41,8 +41,7 @@ bool CoreTechMeshLoader::LoadFile(const FString& FileName, FMeshDescription& Mes
 		return false;
 	}
 
-	ConvertCTBodySetToMeshDescription(ImportParameters, MeshParameters, BodySet, MeshDescription);
-
+	Tessellate(MainObjectID, ImportParameters, MeshDescription, MeshParameters);
 	return true;
 }
 
diff --git a/Engine/Plugins/Enterprise/DatasmithCADImporter/Source/CADLibrary/Public/CoreTechHelper.h b/Engine/Plugins/Enterprise/DatasmithCADImporter/Source/CADLibrary/Public/CoreTechHelper.h
--- a/Engine/Plugins/Enterprise/DatasmithCADImporter/Source/CADLibrary/Public/CoreTechHelper.h
+++ b/Engine/Plugins/Enterprise/DatasmithCADImporter/Source/CADLibrary/Public/CoreTechHelper.h
@@ -50,8 +50,8 @@ namespace CADLibrary
 
 	uint32 GetSize(CT_TESS_DATA_TYPE type);
 
-	CADLIBRARY_API bool ConvertCTBodySetToMeshDescription(const FImportParameters& ImportParams, const FMeshParameters& MeshParameters, TArray<CT_OBJECT_ID>& BodySet, FMeshDescription& OutMeshDescription);
-	CADLIBRARY_API bool ConvertCTBodySetToMeshDescription(const FImportParameters& ImportParams, const FMeshParameters& MeshParameters, uint32 TriangleCount, TArray<FTessellationData>& FacePTessellationSet, TMap<uint32, uint32>& MaterialIdToMaterialHash, FMeshDescription& MeshDescription);
+	CADLIBRARY_API bool ConvertCTBodySetToMeshDescription(const FImportParameters& ImportParams, const FMeshParameters& MeshParameters, FBodyMesh& Body, FMeshDescription& MeshDescription);
+
 	CADLIBRARY_API CT_IO_ERROR Tessellate(CT_OBJECT_ID MainObjectId, const FImportParameters& ImportParams, FMeshDescription& Mesh, FMeshParameters& MeshParameters);
 
 	using TColorMap = TMap<uint32, uint32>;
diff --git a/Engine/Plugins/Enterprise/DatasmithCADImporter/Source/CADTools/Private/CADData.cpp b/Engine/Plugins/Enterprise/DatasmithCADImporter/Source/CADTools/Private/CADData.cpp
--- a/Engine/Plugins/Enterprise/DatasmithCADImporter/Source/CADTools/Private/CADData.cpp
+++ b/Engine/Plugins/Enterprise/DatasmithCADImporter/Source/CADTools/Private/CADData.cpp
@@ -4,11 +4,13 @@
 #include "CADOptions.h"
 
 #include "Misc/FileHelper.h"
-
+#include "Serialization/Archive.h"
+#include "Serialization/MemoryWriter.h"
+#include "Serialization/MemoryReader.h"
 
 namespace CADLibrary
 {
-uint32 BuildFastColorHash(uint32 ColorId, uint8 Alpha)
+uint32 BuildColorId(uint32 ColorId, uint8 Alpha)
 {
 	if (Alpha == 0)
 	{
@@ -17,196 +19,113 @@ uint32 BuildFastColorHash(uint32 ColorId, uint8 Alpha)
 	return ColorId | Alpha << 24;
 }
 
-void UnhashFastColorHash(uint32 ColorHash, uint32& ColorId, uint8& Alpha)
+void GetCTColorIdAlpha(ColorId ColorId, uint32& CTColorId, uint8& Alpha)
 {
-	ColorId = ColorHash & 0x00ffffff;
-	Alpha = (uint8)((ColorHash & 0xff000000) >> 24);
+	CTColorId = ColorId & 0x00ffffff;
+	Alpha = (uint8)((ColorId & 0xff000000) >> 24);
 }
 
-int32 BuildColorHash(const FColor& Color)
+int32 BuildColorName(const FColor& Color)
 {
-	FString HashString = FString::Printf(TEXT("%02x%02x%02x%02x"), Color.R, Color.G, Color.B, Color.A);
-	return FGenericPlatformMath::Abs((int32)GetTypeHash(HashString));
+	FString Name = FString::Printf(TEXT("%02x%02x%02x%02x"), Color.R, Color.G, Color.B, Color.A);
+	return FGenericPlatformMath::Abs((int32)GetTypeHash(Name));
 }
 
-int32 BuildMaterialHash(const FCADMaterial& Material)
+int32 BuildMaterialName(const FCADMaterial& Material)
 {
-	FString Hash;
+	FString Name;
 	if (!Material.MaterialName.IsEmpty())
 	{
-		Hash += Material.MaterialName;  // we add material name because it could be used by the end user so two material with same parameters but different name are different.
+		Name += Material.MaterialName;  // we add material name because it could be used by the end user so two material with same parameters but different name are different.
 	}
-	Hash += FString::Printf(TEXT("%02x%02x%02x "), Material.Diffuse.R, Material.Diffuse.G, Material.Diffuse.B);
-	Hash += FString::Printf(TEXT("%02x%02x%02x "), Material.Ambient.R, Material.Ambient.G, Material.Ambient.B);
-	Hash += FString::Printf(TEXT("%02x%02x%02x "), Material.Specular.R, Material.Specular.G, Material.Specular.B);
-	Hash += FString::Printf(TEXT("%02x%02x%02x"), (int)(Material.Shininess * 255.0), (int)(Material.Transparency * 255.0), (int)(Material.Reflexion * 255.0));
+	Name += FString::Printf(TEXT("%02x%02x%02x "), Material.Diffuse.R, Material.Diffuse.G, Material.Diffuse.B);
+	Name += FString::Printf(TEXT("%02x%02x%02x "), Material.Ambient.R, Material.Ambient.G, Material.Ambient.B);
+	Name += FString::Printf(TEXT("%02x%02x%02x "), Material.Specular.R, Material.Specular.G, Material.Specular.B);
+	Name += FString::Printf(TEXT("%02x%02x%02x"), (int)(Material.Shininess * 255.0), (int)(Material.Transparency * 255.0), (int)(Material.Reflexion * 255.0));
 
 	if (!Material.TextureName.IsEmpty())
 	{
-		Hash += Material.TextureName;
+		Name += Material.TextureName;
 	}
-	return FMath::Abs((int32)GetTypeHash(Hash));
+	return FMath::Abs((int32)GetTypeHash(Name));
 }
 
-
-FBodyMesh::FBodyMesh(uint32 InBodyID, int32 FaceNum)
-	: TriangleCount(0)
-	, BodyID(InBodyID)
+FArchive& operator<<(FArchive& Ar, FCADMaterial& Material)
 {
-	FaceTessellationSet.Reserve(FaceNum);
+	Ar << Material.MaterialName;
+	Ar << Material.Diffuse;
+	Ar << Material.Ambient;
+	Ar << Material.Specular;
+	Ar << Material.Shininess;
+	Ar << Material.Transparency;
+	Ar << Material.Reflexion;
+	Ar << Material.TextureName;
+	return Ar;
 }
 
-
-FCADMeshFile::FCADMeshFile(FString& InFileName, TMap< uint32, FBodyMesh* >& InBodyUuidToCTBodyMap)
-	: BodyUuidToBodyMap(InBodyUuidToCTBodyMap)
+FArchive& operator<<(FArchive& Ar, FTessellationData& TessellationData)
 {
-	uint32 NbBody = 0;
-
-	TArray<uint8> MeshData;
-	FFileHelper::LoadFileToArray(MeshData, *InFileName);
-
-	ReadTessellationSetFromFile(MeshData, NbBody, CTBodySet);
-
-	for(FBodyMesh& Body : CTBodySet)
-	{
-		BodyUuidToBodyMap.Add(Body.GetBodyUuid(), &Body);
-	}
+	Ar << TessellationData.VertexArray;
+	Ar << TessellationData.NormalArray;
+	Ar << TessellationData.IndexArray;
+	Ar << TessellationData.TexCoordArray;
+	Ar << TessellationData.VertexCount;
+	Ar << TessellationData.NormalCount;
+	Ar << TessellationData.IndexCount;
+	Ar << TessellationData.TexCoordCount;
+
+	Ar << TessellationData.StartVertexIndex;
+
+	Ar << TessellationData.SizeOfVertexType;
+	Ar << TessellationData.SizeOfTexCoordType;
+	Ar << TessellationData.SizeOfNormalType;
+	Ar << TessellationData.SizeOfIndexType;
+
+	Ar << TessellationData.ColorName;
+	Ar << TessellationData.MaterialName;
+
+	return Ar;
 }
 
-enum ETessellationContent
+FArchive& operator<<(FArchive& Ar, FBodyMesh& BodyMesh)
 {
-	HasVertex = 1 << 1,
-	HasNormal = 1 << 2,
-	HasIndex = 1 << 3,
-	HasTexCoord = 1 << 4,
+	Ar << BodyMesh.TriangleCount;
+	Ar << BodyMesh.BodyID;
+	Ar << BodyMesh.MeshActorName;
+	Ar << BodyMesh.Faces;
 
-	Default = HasVertex | HasNormal | HasIndex | HasTexCoord
-};
+	Ar << BodyMesh.MaterialSet;
+	Ar << BodyMesh.ColorSet;
 
-void WriteTessellationInRawData(FTessellationData& Tessellation, TArray<uint8>& GlobalRawData)
-{
-	int32 INFO[LastLine];
-
-	int32 Types = ETessellationContent::Default;
-	uint32 VertexSize = Tessellation.VertexArray.Num();
-	uint32 NormalSize = Tessellation.NormalArray.Num();
-	uint32 IndexSize = Tessellation.IndexArray.Num();
-	uint32 TexCoordSize = Tessellation.TexCoordArray.Num();
-
-	uint32 TessellationRawDataSize = VertexSize + NormalSize + IndexSize + TexCoordSize + sizeof(INFO); // VertexCount, VertexType, NormalCount, NormalType, IndexCount, IndexType, TexCoordType;
-
-	INFO[LineType] = Types;
-	INFO[LineRawSize] = TessellationRawDataSize;
-	INFO[LineBodyId] = Tessellation.BodyId;
-	INFO[LineBodyUuId] = Tessellation.BodyUuId;
-	INFO[LineBodyFaceNum] = Tessellation.BodyFaceNum;
-	INFO[LineMaterialId] = Tessellation.MaterialId;
-	INFO[LineMaterialHash] = Tessellation.MaterialHash;
-	INFO[LineVertexCount] = Tessellation.VertexCount;
-	INFO[LineVertexType] = Tessellation.SizeOfVertexType;
-	INFO[LineNormalCount] = Tessellation.NormalCount;
-	INFO[LineNormalType] = Tessellation.SizeOfNormalType;
-	INFO[LineIndexCount] = Tessellation.IndexCount;
-	INFO[LineIndexType] = Tessellation.SizeOfIndexType;
-	INFO[LineTexCoordCount] = Tessellation.TexCoordCount;
-	INFO[LineTexCoordType] = Tessellation.SizeOfTexCoordType;
-
-	GlobalRawData.Append((uint8*)INFO, sizeof(INFO));
-	GlobalRawData.Append(Tessellation.VertexArray);
-	GlobalRawData.Append(Tessellation.NormalArray);
-	GlobalRawData.Append(Tessellation.IndexArray);
-	GlobalRawData.Append(Tessellation.TexCoordArray);
+	return Ar;
 }
 
-bool ReadTessellationSetFromFile(TArray<uint8>& RawData, uint32& OutNbBodies, TArray<FBodyMesh>& CTBodySet)
+void SerializeBodyMeshSet(const TCHAR* Filename, TArray<FBodyMesh>& InBodySet)
 {
-	int32 RawDataOffset = 0;
-	memcpy(&OutNbBodies, RawData.GetData() + RawDataOffset, sizeof(uint32));
-	RawDataOffset += sizeof(uint32);
+	TArray<uint8> OutBuffer;
+	FMemoryWriter ArWriter(OutBuffer);
+	uint32 type = 234561;
+	ArWriter << type;
 
-	CTBodySet.Reserve(OutNbBodies);
+	ArWriter << InBodySet;
 
-	FBodyMesh* CurrentBody = nullptr;
+	FFileHelper::SaveArrayToFile(OutBuffer, Filename);
+}
 
-	int32 CurrentBodyId = -1;
-	int32 CurrentBodyIndex = -1;
-	uint32 BodyHeader[LastLine];
-	const uint32 BodyHeaderByteSize = sizeof(BodyHeader);
+void DeserializeBodyMeshFile(const TCHAR* Filename, TArray<FBodyMesh>& OutBodySet)
+{
+	TArray<uint8> InBuffer;
+	FFileHelper::LoadFileToArray(InBuffer, Filename);
 
-	while (RawDataOffset + int32(BodyHeaderByteSize) <= RawData.Num())
+	FMemoryReader ArReader(InBuffer);
+	uint32 type = 0;
+	ArReader << type;
+	if (type != 234561)
 	{
-		memcpy(BodyHeader, RawData.GetData() + RawDataOffset, BodyHeaderByteSize);
-
-		int32 BlocByteSize  = BodyHeader[LineRawSize];
-		if (RawDataOffset + BlocByteSize > RawData.Num())
-		{
-			return false;
-		}
-
-		RawDataOffset += BodyHeaderByteSize;
-
-		int32 Types = BodyHeader[LineType];
-		int32 RawSize = BodyHeader[LineRawSize];
-
-		if (Types != ETessellationContent::Default)
-		{
-			return false;
-		}
-
-		if (BodyHeader[LineBodyId] != CurrentBodyId)
-		{
-			CurrentBodyId = BodyHeader[LineBodyId];
-			CTBodySet.Emplace(BodyHeader[LineBodyId], BodyHeader[LineBodyFaceNum]);
-			CurrentBodyIndex++;
-
-			CurrentBody = &CTBodySet[CurrentBodyIndex];
-			CurrentBody->SetBodyId(BodyHeader[LineBodyId]);
-			CurrentBody->SetBodyUuid(BodyHeader[LineBodyUuId]);
-			CurrentBody->GetTessellationSet().Reserve(BodyHeader[LineBodyFaceNum]);
-		}
-
-		if (CurrentBody == nullptr)
-		{
-			return false;
-		}
-
-		FTessellationData& Tessellation = CurrentBody->GetTessellationSet().Emplace_GetRef();
-
-		Tessellation.BodyId = BodyHeader[LineBodyId];
-		Tessellation.BodyUuId = BodyHeader[LineBodyUuId];
-		Tessellation.BodyFaceNum = BodyHeader[LineBodyFaceNum];
-		Tessellation.MaterialId = BodyHeader[LineMaterialId];
-		Tessellation.MaterialHash = BodyHeader[LineMaterialHash];
-		Tessellation.VertexCount = BodyHeader[LineVertexCount];
-		Tessellation.SizeOfVertexType = BodyHeader[LineVertexType];
-		Tessellation.NormalCount = BodyHeader[LineNormalCount];
-		Tessellation.SizeOfNormalType = BodyHeader[LineNormalType];
-		Tessellation.IndexCount = BodyHeader[LineIndexCount];
-		Tessellation.SizeOfIndexType = BodyHeader[LineIndexType];
-		Tessellation.TexCoordCount = BodyHeader[LineTexCoordCount];
-		Tessellation.SizeOfTexCoordType = BodyHeader[LineTexCoordType];
-
-		uint32 VertexSize = Tessellation.VertexCount * Tessellation.SizeOfVertexType * 3;
-		uint32 NormalSize = Tessellation.NormalCount * Tessellation.SizeOfNormalType * 3;
-		uint32 IndexSize = Tessellation.IndexCount * Tessellation.SizeOfIndexType;
-		uint32 TexCoordSize = Tessellation.TexCoordCount * Tessellation.SizeOfTexCoordType * 2;
-
-		Tessellation.VertexArray.Empty(VertexSize);
-		Tessellation.NormalArray.Empty(NormalSize);
-		Tessellation.IndexArray.Empty(IndexSize);
-		Tessellation.TexCoordArray.Empty(TexCoordSize);
-
-		Tessellation.VertexArray.Append(RawData.GetData() + RawDataOffset, VertexSize);
-		RawDataOffset += VertexSize;
-		Tessellation.NormalArray.Append(RawData.GetData() + RawDataOffset, NormalSize);
-		RawDataOffset += NormalSize;
-		Tessellation.IndexArray.Append(RawData.GetData() + RawDataOffset, IndexSize);
-		RawDataOffset += IndexSize;
-		Tessellation.TexCoordArray.Append(RawData.GetData() + RawDataOffset, TexCoordSize);
-		RawDataOffset += TexCoordSize;
+		return;
 	}
-	return true;
-}
 
+	ArReader << OutBodySet;
+}
 
 }
diff --git a/Engine/Plugins/Enterprise/DatasmithCADImporter/Source/CADTools/Public/CADData.h b/Engine/Plugins/Enterprise/DatasmithCADImporter/Source/CADTools/Public/CADData.h
--- a/Engine/Plugins/Enterprise/DatasmithCADImporter/Source/CADTools/Public/CADData.h
+++ b/Engine/Plugins/Enterprise/DatasmithCADImporter/Source/CADTools/Public/CADData.h
@@ -6,32 +6,23 @@
 #include "Containers/Array.h"
 #include "Math/Color.h"
 
+class FArchive;
+
+using CadId = uint32; // Identifier defined in the input CAD file
+using ColorId = uint32; // Identifier defined in the input CAD file
+using MaterialId = uint32; // Identifier defined in the input CAD file
+using CADUUID = uint32;  // Universal unique identifier that be used for the unreal asset name (Actor, Material)
+
+
 namespace CADLibrary
 {
 
-enum ECTTessellationLine
+class CADTOOLS_API FCADMaterial
 {
-	LineType = 0,
-	LineRawSize,
-	LineBodyId,
-	LineBodyUuId,
-	LineBodyFaceNum,
-	LineMaterialId,
-	LineMaterialHash,
-	LineVertexCount,
-	LineVertexType,
-	LineNormalCount,
-	LineNormalType,
-	LineIndexCount,
-	LineIndexType,
-	LineTexCoordType,
-	LineTexCoordCount,
-	LastLine
-};
+public:
+	friend CADTOOLS_API FArchive& operator<<(FArchive& Ar, FCADMaterial& Material);
 
-struct CADTOOLS_API FCADMaterial
-{
-	int32 MaterialId;  //#ueent_CAD
+public:
 	FString MaterialName;
 	FColor Diffuse;
 	FColor Ambient;
@@ -44,8 +35,9 @@ struct CADTOOLS_API FCADMaterial
 
 struct CADTOOLS_API FObjectDisplayDataId
 {
-	uint32 MaterialId = 0;
-	uint32 ColorHId = 0; // => FastHash == ColorId+Transparency
+	CADUUID DefaultMaterialName = 0;
+	MaterialId Material = 0;
+	ColorId Color = 0; // => FastHash == ColorId+Transparency
 };
 
 /**
@@ -53,6 +45,8 @@ struct CADTOOLS_API FObjectDisplayDataId
  */
 struct CADTOOLS_API FTessellationData
 {
+	friend CADTOOLS_API FArchive& operator<<(FArchive& Ar, FTessellationData& Tessellation);
+
 	TArray<uint8> VertexArray;
 	TArray<uint8> NormalArray;
 	TArray<uint8> IndexArray;
@@ -64,97 +58,46 @@ struct CADTOOLS_API FTessellationData
 
 	uint32    StartVertexIndex = 0;
 
-	uint8 SizeOfVertexType;
-	uint8 SizeOfTexCoordType;
-	uint8 SizeOfNormalType;
-	uint8 SizeOfIndexType;
-	bool  bHasRGBColor;
+	uint8 SizeOfVertexType = 0;
+	uint8 SizeOfTexCoordType = 0;
+	uint8 SizeOfNormalType = 0;
+	uint8 SizeOfIndexType = 0;
 
-	uint32 BodyId = 0;
-	uint32 BodyUuId = 0;
-	uint32 BodyFaceNum = 0;
-	uint32 MaterialId = 0;
-	uint32 MaterialHash = 0;
+	CADUUID ColorName = 0;
+	CADUUID MaterialName = 0;
 
-	TArray<int32> VertexIdSet;  // StaticMesh FVertexID
-	TArray<int32> SymVertexIdSet; // StaticMesh FVertexID for sym part
+	TArray<int32> VertexIdSet;  // StaticMesh FVertexID NO Serialize
+	TArray<int32> SymVertexIdSet; // StaticMesh FVertexID for sym part NO Serialize
 };
 
 class CADTOOLS_API FBodyMesh
 {
 public:
-	FBodyMesh(uint32 BodyID, int32 FaceNum);
-
-	TArray<FTessellationData>& GetTessellationSet()
-	{
-		return FaceTessellationSet;
-	}
-
-	uint32 GetTriangleCount() const
-	{
-		return TriangleCount;
-	}
-
-	uint32 GetBodyId() const
+	FBodyMesh(CadId InBodyID = 0) : BodyID(InBodyID)
 	{
-		return BodyID;
 	}
 
-	uint32 GetBodyUuid() const 
-	{
-		return BodyUuid;
-	}
-
-	void SetBodyId(uint32 Id)
-	{
-		BodyID = Id;
-	}
-
-	void SetBodyUuid(uint32 Uuid)
-	{
-		BodyUuid = Uuid;
-	}
+	friend FArchive& operator<<(FArchive& Ar, FBodyMesh& BodyMesh);
 
-protected:
-
-	/**
-	 * FCTTessellation is an elementary structure of Mesh
-	 */
-	TArray<FTessellationData> FaceTessellationSet;
-
-	/**
-	 * Structure use to map Material ID to Material Hash. Material hash is an unique value by material
-	 */
-	TMap<uint32, uint32> MaterialIdToMaterialHash;
-	uint32 TriangleCount;
-	uint32 BodyID;
-	uint32 BodyUuid;
-};
-
-class CADTOOLS_API FCADMeshFile
-{
 public:
-	FCADMeshFile(FString& InFileName, TMap< uint32, FBodyMesh* >& InBodyUuidToCTBodyMap);
+	TArray<FTessellationData> Faces;
 
-	FString& GetFileName()
-	{
-		return FileName;
-	}
+	uint32 TriangleCount = 0;
+	CadId BodyID = 0;
+	CADUUID MeshActorName = 0;
 
-protected:
-	FString FileName;
-	TArray<FBodyMesh> CTBodySet;
-	TMap<uint32, FBodyMesh* >& BodyUuidToBodyMap;
+	TSet<uint32> MaterialSet;
+	TSet<uint32> ColorSet;
 };
 
-CADTOOLS_API uint32 BuildFastColorHash(uint32 ColorId, uint8 Alpha);
-CADTOOLS_API void UnhashFastColorHash(uint32 ColorHash, uint32& OutColorId, uint8& OutAlpha);
+CADTOOLS_API uint32 BuildColorId(uint32 ColorId, uint8 Alpha);
+CADTOOLS_API void GetCTColorIdAlpha(uint32 ColorHash, uint32& OutColorId, uint8& OutAlpha);
 
-CADTOOLS_API int32 BuildColorHash(const FColor& Color);
-CADTOOLS_API int32 BuildMaterialHash(const FCADMaterial& Material);
+CADTOOLS_API int32 BuildColorName(const FColor& Color);
+CADTOOLS_API int32 BuildMaterialName(const FCADMaterial& Material);
 
-CADTOOLS_API void WriteTessellationInRawData(FTessellationData& Tessellation, TArray<uint8>& GlobalRawData);
-CADTOOLS_API bool ReadTessellationSetFromFile(TArray<uint8>& RawData, uint32& OutNbBody, TArray<FBodyMesh>& FaceTessellations);
+CADTOOLS_API void SerializeBodyMeshSet(const TCHAR* Filename, TArray<FBodyMesh>& InBodySet);
+CADTOOLS_API void DeserializeBodyMeshFile(const TCHAR* Filename, TArray<FBodyMesh>& OutBodySet);
 
 inline void CopyValue(const uint8* Source, int Offset, uint8 Size, bool bIs3D, FVector& Dest, const FMatrix* Matrix = nullptr)
 {
diff --git a/Engine/Plugins/Enterprise/DatasmithCADImporter/Source/DatasmithCADTranslator/Private/DatasmithCADTranslator.cpp b/Engine/Plugins/Enterprise/DatasmithCADImporter/Source/DatasmithCADTranslator/Private/DatasmithCADTranslator.cpp
--- a/Engine/Plugins/Enterprise/DatasmithCADImporter/Source/DatasmithCADTranslator/Private/DatasmithCADTranslator.cpp
+++ b/Engine/Plugins/Enterprise/DatasmithCADImporter/Source/DatasmithCADTranslator/Private/DatasmithCADTranslator.cpp
@@ -11,10 +11,13 @@
 #include "Misc/FileHelper.h"
 
 
+#define CAD_CACHE_VERSION 1
+
 void FDatasmithCADTranslator::Initialize(FDatasmithTranslatorCapabilities& OutCapabilities)
 {
+#ifndef CAD_TRANSLATOR_DEBUG
 	OutCapabilities.bParallelLoadStaticMeshSupported = true;
-
+#endif
 	OutCapabilities.SupportedFileFormats.Add(FFileFormatInfo{ TEXT("CATPart"), TEXT("CATIA Part files") });
 	OutCapabilities.SupportedFileFormats.Add(FFileFormatInfo{ TEXT("CATProduct"), TEXT("CATIA Product files") });
 	OutCapabilities.SupportedFileFormats.Add(FFileFormatInfo{ TEXT("cgr"), TEXT("CATIA Graphical Representation V5 files") });
@@ -52,6 +55,8 @@ void FDatasmithCADTranslator::Initialize(FDatasmithTranslatorCapabilities& OutCa
 
 	OutCapabilities.SupportedFileFormats.Add(FFileFormatInfo{ TEXT("dwg"), TEXT("AutoCAD, Model files") });
 	OutCapabilities.SupportedFileFormats.Add(FFileFormatInfo{ TEXT("dgn"), TEXT("MicroStation files") });
+
+	OutCapabilities.SupportedFileFormats.Add(FFileFormatInfo{ TEXT("ct"), TEXT("Kernel_IO files") });
 }
 
 bool FDatasmithCADTranslator::LoadScene(TSharedRef<IDatasmithScene> DatasmithScene)
@@ -82,20 +87,24 @@ bool FDatasmithCADTranslator::LoadScene(TSharedRef<IDatasmithScene> DatasmithSce
 	FString CachePath = FPaths::ConvertRelativePathToFull(FDatasmithCADTranslatorModule::Get().GetCacheDir());
 
 	TMap<FString, FString> CADFileToUE4FileMap;
-
 	int32 NumCores = FPlatformMisc::NumberOfCores();
 	{
 		DatasmithDispatcher::FDatasmithDispatcher Dispatcher(ImportParameters, CachePath, NumCores, CADFileToUE4FileMap, CADFileToUE4GeomMap);
 		Dispatcher.AddTask(FPaths::ConvertRelativePathToFull(GetSource().GetSourceFile()));
 
 		bool bWithProcessor = true;
+
+#ifdef CAD_TRANSLATOR_DEBUG
+		bWithProcessor = false;
+#endif //CAD_TRANSLATOR_DEBUG
+		
 		Dispatcher.Process(bWithProcessor);
 	}
 
-	FDatasmithSceneGraphBuilder SceneGraphBuilder(CADFileToUE4FileMap, MeshElementToCADBRepUuidMap,  CachePath, DatasmithScene, GetSource(), ImportParameters);
+	FDatasmithSceneGraphBuilder SceneGraphBuilder(CADFileToUE4FileMap, CachePath, DatasmithScene, GetSource(), ImportParameters);
 	SceneGraphBuilder.Build();
 
-	MeshBuilderPtr = MakeUnique<FDatasmithMeshBuilder>(CADFileToUE4GeomMap, MeshElementToCADBRepUuidMap, CachePath, ImportParameters);
+	MeshBuilderPtr = MakeUnique<FDatasmithMeshBuilder>(CADFileToUE4GeomMap, CachePath, ImportParameters);
 
 	return true;
 }
@@ -105,7 +114,6 @@ void FDatasmithCADTranslator::UnloadScene()
 	MeshBuilderPtr = nullptr;
 
 	CADFileToUE4GeomMap.Empty();
-	MeshElementToCADBRepUuidMap.Empty();
 }
 
 bool FDatasmithCADTranslator::LoadStaticMesh(const TSharedRef<IDatasmithMeshElement> MeshElement, FDatasmithMeshElementPayload& OutMeshPayload)
diff --git a/Engine/Plugins/Enterprise/DatasmithCADImporter/Source/DatasmithCADTranslator/Private/DatasmithCADTranslator.h b/Engine/Plugins/Enterprise/DatasmithCADImporter/Source/DatasmithCADTranslator/Private/DatasmithCADTranslator.h
--- a/Engine/Plugins/Enterprise/DatasmithCADImporter/Source/DatasmithCADTranslator/Private/DatasmithCADTranslator.h
+++ b/Engine/Plugins/Enterprise/DatasmithCADImporter/Source/DatasmithCADTranslator/Private/DatasmithCADTranslator.h
@@ -37,7 +37,6 @@ public:
 
 private:
 	TMap<FString, FString> CADFileToUE4GeomMap;
-	TMap< TSharedPtr< IDatasmithMeshElement >, uint32 > MeshElementToCADBRepUuidMap;
 
 	CADLibrary::FImportParameters ImportParameters;
 
diff --git a/Engine/Plugins/Enterprise/DatasmithCADImporter/Source/DatasmithCADTranslator/Private/DatasmithCADTranslatorModule.cpp b/Engine/Plugins/Enterprise/DatasmithCADImporter/Source/DatasmithCADTranslator/Private/DatasmithCADTranslatorModule.cpp
--- a/Engine/Plugins/Enterprise/DatasmithCADImporter/Source/DatasmithCADTranslator/Private/DatasmithCADTranslatorModule.cpp
+++ b/Engine/Plugins/Enterprise/DatasmithCADImporter/Source/DatasmithCADTranslator/Private/DatasmithCADTranslatorModule.cpp
@@ -14,7 +14,11 @@ IMPLEMENT_MODULE(FDatasmithCADTranslatorModule, DatasmithCADTranslator);
 void FDatasmithCADTranslatorModule::StartupModule()
 {
 	// Create temporary directory which will be used by CoreTech to store tessellation data
-	const TCHAR* CacheVersion = TEXT("0");
+	const TCHAR* OldCacheVersion = TEXT("0");
+	FString OldCacheDir = FPaths::ConvertRelativePathToFull(FPaths::Combine(FPaths::ProjectIntermediateDir(), TEXT("DatasmithCADCache"), OldCacheVersion));
+	IFileManager::Get().DeleteDirectory(*OldCacheDir, true, true);
+
+	const TCHAR* CacheVersion = TEXT("1");
 	CacheDir = FPaths::ConvertRelativePathToFull(FPaths::Combine(FPaths::ProjectIntermediateDir(), TEXT("DatasmithCADCache"), CacheVersion));
 	IFileManager::Get().MakeDirectory(*CacheDir);
 
diff --git a/Engine/Plugins/Enterprise/DatasmithCADImporter/Source/DatasmithCADTranslator/Private/DatasmithMeshBuilder.cpp b/Engine/Plugins/Enterprise/DatasmithCADImporter/Source/DatasmithCADTranslator/Private/DatasmithMeshBuilder.cpp
--- a/Engine/Plugins/Enterprise/DatasmithCADImporter/Source/DatasmithCADTranslator/Private/DatasmithMeshBuilder.cpp
+++ b/Engine/Plugins/Enterprise/DatasmithCADImporter/Source/DatasmithCADTranslator/Private/DatasmithMeshBuilder.cpp
@@ -15,10 +15,9 @@
 
 using namespace CADLibrary;
 
-FDatasmithMeshBuilder::FDatasmithMeshBuilder(TMap<FString, FString>& InCADFileToUE4GeomMap, TMap< TSharedPtr< IDatasmithMeshElement >, uint32 >& InMeshElementToCTBodyUuidMap, const FString& InCachePath, const CADLibrary::FImportParameters& InImportParameters)
+FDatasmithMeshBuilder::FDatasmithMeshBuilder(TMap<FString, FString>& InCADFileToUE4GeomMap, const FString& InCachePath, const CADLibrary::FImportParameters& InImportParameters)
 	: CachePath(InCachePath)
-	, CADFileToMeshFileMap(InCADFileToUE4GeomMap)
-	, MeshElementToBodyUuidMap(InMeshElementToCTBodyUuidMap)
+	, CADFileToMeshFile(InCADFileToUE4GeomMap)
 	, ImportParameters(InImportParameters)
 {
 	LoadMeshFiles();
@@ -26,29 +25,35 @@ FDatasmithMeshBuilder::FDatasmithMeshBuilder(TMap<FString, FString>& InCADFileTo
 
 void FDatasmithMeshBuilder::LoadMeshFiles()
 {
-	Meshes.Reserve(CADFileToMeshFileMap.Num());
-	for (const auto& FilePair : CADFileToMeshFileMap)
+	BodyMeshes.Reserve(CADFileToMeshFile.Num());
+
+	for (const auto& FilePair : CADFileToMeshFile)
 	{
 		FString MeshFile = FPaths::Combine(CachePath, TEXT("mesh"), FilePair.Value + TEXT(".gm"));
 		if (!IFileManager::Get().FileExists(*MeshFile))
 		{
 			continue;
 		}
-
-		Meshes.Emplace(MeshFile, BodyUuidToMeshMap);
+		TArray<CADLibrary::FBodyMesh>& BodyMeshSet = BodyMeshes.Emplace_GetRef();
+		DeserializeBodyMeshFile(*MeshFile, BodyMeshSet);
+		for (FBodyMesh& Body : BodyMeshSet)
+		{
+			MeshActorNameToBodyMesh.Emplace(Body.MeshActorName, &Body);
+		}
 	}
 }
 
 TOptional<FMeshDescription> FDatasmithMeshBuilder::GetMeshDescription(TSharedRef<IDatasmithMeshElement> OutMeshElement, CADLibrary::FMeshParameters& OutMeshParameters)
 {
 #ifdef CAD_INTERFACE
-	uint32* BodyUuid = MeshElementToBodyUuidMap.Find(OutMeshElement);
-	if (BodyUuid == nullptr)
+	const TCHAR* NameLabel = OutMeshElement->GetName();
+	CADUUID BodyUuid = (CADUUID) FCString::Atoi64(OutMeshElement->GetName()+2);  // +2 to remove 2 first char (Ox)
+	if (BodyUuid == 0)
 	{
 		return TOptional<FMeshDescription>();
 	}
 
-	FBodyMesh*const* PPBody = BodyUuidToMeshMap.Find(*BodyUuid);
+	FBodyMesh** PPBody = MeshActorNameToBodyMesh.Find(BodyUuid);
 	if(PPBody == nullptr || *PPBody == nullptr)
 	{
 		return TOptional<FMeshDescription>();
@@ -56,12 +61,10 @@ TOptional<FMeshDescription> FDatasmithMeshBuilder::GetMeshDescription(TSharedRef
 
 	FBodyMesh& Body = **PPBody;
 
-	TMap<uint32, uint32> MaterialIdToMaterialHash;
-
 	FMeshDescription MeshDescription;
 	DatasmithMeshHelper::PrepareAttributeForStaticMesh(MeshDescription);
 
-	if (ConvertCTBodySetToMeshDescription(ImportParameters, OutMeshParameters, Body.GetTriangleCount(), Body.GetTessellationSet(), MaterialIdToMaterialHash, MeshDescription))
+	if (ConvertCTBodySetToMeshDescription(ImportParameters, OutMeshParameters, Body, MeshDescription))
 	{
 		return MoveTemp(MeshDescription);
 	}
diff --git a/Engine/Plugins/Enterprise/DatasmithCADImporter/Source/DatasmithCADTranslator/Private/DatasmithMeshBuilder.h b/Engine/Plugins/Enterprise/DatasmithCADImporter/Source/DatasmithCADTranslator/Private/DatasmithMeshBuilder.h
--- a/Engine/Plugins/Enterprise/DatasmithCADImporter/Source/DatasmithCADTranslator/Private/DatasmithMeshBuilder.h
+++ b/Engine/Plugins/Enterprise/DatasmithCADImporter/Source/DatasmithCADTranslator/Private/DatasmithMeshBuilder.h
@@ -15,7 +15,7 @@ class IDatasmithMeshElement;
 class FDatasmithMeshBuilder
 {
 public:
-	FDatasmithMeshBuilder(TMap<FString, FString>& InCADFileToMeshFileMap, TMap< TSharedPtr< IDatasmithMeshElement >, uint32 >& InMeshElementToCADBodyUuidMap, const FString& InCachePath, const CADLibrary::FImportParameters& InImportParameters);
+	FDatasmithMeshBuilder(TMap<FString, FString>& InCADFileToMeshFileMap, const FString& InCachePath, const CADLibrary::FImportParameters& InImportParameters);
 
 	TOptional<FMeshDescription> GetMeshDescription(TSharedRef<IDatasmithMeshElement> OutMeshElement, CADLibrary::FMeshParameters& OutMeshParameters);
 
@@ -24,10 +24,9 @@ protected:
 
 	void LoadMeshFiles();
 
-	TArray<CADLibrary::FCADMeshFile> Meshes;
-	TMap<FString, FString>& CADFileToMeshFileMap;
-	TMap< TSharedPtr< IDatasmithMeshElement >, uint32 >& MeshElementToBodyUuidMap;
-	TMap< uint32, CADLibrary::FBodyMesh* > BodyUuidToMeshMap;
+	TMap<FString, FString>& CADFileToMeshFile;
+	TArray<TArray<CADLibrary::FBodyMesh>> BodyMeshes;
+	TMap<CADUUID, CADLibrary::FBodyMesh*> MeshActorNameToBodyMesh;
 
 	CADLibrary::FImportParameters ImportParameters;
 };
diff --git a/Engine/Plugins/Enterprise/DatasmithCADImporter/Source/DatasmithCADTranslator/Private/DatasmithSceneGraphBuilder.cpp b/Engine/Plugins/Enterprise/DatasmithCADImporter/Source/DatasmithCADTranslator/Private/DatasmithSceneGraphBuilder.cpp
--- a/Engine/Plugins/Enterprise/DatasmithCADImporter/Source/DatasmithCADTranslator/Private/DatasmithSceneGraphBuilder.cpp
+++ b/Engine/Plugins/Enterprise/DatasmithCADImporter/Source/DatasmithCADTranslator/Private/DatasmithSceneGraphBuilder.cpp
@@ -2,7 +2,9 @@
 #include "DatasmithSceneGraphBuilder.h"
 
 #ifdef CAD_INTERFACE
+
 #include "CADData.h"
+#include "CADSceneGraph.h"
 #include "CoreTechFileParser.h"
 #include "CoreTechHelper.h"
 #include "DatasmithSceneFactory.h"
@@ -15,409 +17,49 @@
 
 namespace
 {
-	void GetMainMaterial(const TMap<FString, FString>& InNodeMetaDataMap, FSceneNodeDescription* InNode, ActorData& OutNodeData, bool bMaterialPropagationIsTopDown)
+	void GetMainMaterial(const TMap<FString, FString>& InNodeMetaDataMap, ActorData& OutNodeData, bool bMaterialPropagationIsTopDown)
 	{
-		if (const FString* MaterialIdStr = InNodeMetaDataMap.Find(TEXT("MaterialId")))
+		if (const FString* MaterialNameStr = InNodeMetaDataMap.Find(TEXT("MaterialName")))
 		{
 			if (!bMaterialPropagationIsTopDown || !OutNodeData.MaterialUuid)
 			{
-				uint32 MaterialId = FCString::Atoi64(**MaterialIdStr);
-				CADLibrary::FCADMaterial Material;
-				if (InNode->GetMaterial(MaterialId, Material))
-				{
-					OutNodeData.Material = Material;
-					OutNodeData.MaterialUuid = CADLibrary::BuildMaterialHash(OutNodeData.Material);
-				}
+				OutNodeData.MaterialUuid = FCString::Atoi64(**MaterialNameStr);
 			}
 		}
 
-		if (const FString* ColorIdStr = InNodeMetaDataMap.Find(TEXT("ColorUEId")))
+		if (const FString* ColorHashStr = InNodeMetaDataMap.Find(TEXT("ColorName")))
 		{
 			if (bMaterialPropagationIsTopDown || !OutNodeData.ColorUuid)
 			{
-				uint32 ColorHId = FCString::Atoi64(**ColorIdStr);
-				FColor Color;
-				if (InNode->GetColor(ColorHId, Color))
-				{
-					OutNodeData.Color = Color;
-					OutNodeData.ColorUuid = CADLibrary::BuildColorHash(OutNodeData.Color);
-				}
+				OutNodeData.ColorUuid = FCString::Atoi64(**ColorHashStr);
 			}
 		}
 	}
-}
-
-
-FCADSceneGraphDescriptionFile::FCADSceneGraphDescriptionFile(const FString& InFileName) 
-	: FileName(InFileName)
-{
-	ReadFile();
-}
-
-void FCADSceneGraphDescriptionFile::GetMaterialDescription(int32 LineNumber, CADLibrary::FCADMaterial& Material) const
-{
-	FString OutMaterialName;
-	uint8 OutDiffuse[3];
-	uint8 OutAmbient[3];
-	uint8 OutSpecular[3];
-	float OutShininess, OutTransparency, OutReflexion;
-	FString OutTextureName;
-
-	OutDiffuse[0] = OutDiffuse[1] = OutDiffuse[2] = 255;
-	OutAmbient[0] = OutAmbient[1] = OutAmbient[2] = 255;
-	OutSpecular[0] = OutSpecular[1] = OutSpecular[2] = 255;
-	OutShininess = OutReflexion = 0;
-	OutTransparency = 255;
-
-	const FString& MaterialLine1 = GetString(LineNumber);
-	FString MaterialIdStr;
-	MaterialLine1.Split(SPACE, &MaterialIdStr, &OutMaterialName);
-
-	const FString& MaterialLine2 = GetString(LineNumber+1);
-	TArray<FString> MaterialParameters;
-	MaterialLine2.ParseIntoArray(MaterialParameters, *SPACE);
-
-	for (int32 index = 0; index < 3; index++)
-	{
-		OutDiffuse[index] = FCString::Atoi(*MaterialParameters[index + 2]);
-		OutAmbient[index] = FCString::Atoi(*MaterialParameters[index + 5]);
-		OutSpecular[index] = FCString::Atoi(*MaterialParameters[index + 8]);
-	}
-	OutShininess = FCString::Atof(*MaterialParameters[11]);
-	OutTransparency = FCString::Atof(*MaterialParameters[12]);
-	OutReflexion = FCString::Atof(*MaterialParameters[13]);
-	if (MaterialParameters.Num() > 14)
-	{
-		OutTextureName = MaterialParameters[14];
-	}
-
-	Material.MaterialId = FCString::Atoi(*MaterialParameters[0]);
-	Material.MaterialName = OutMaterialName;
-	Material.Diffuse= FColor(OutDiffuse[0], OutDiffuse[1], OutDiffuse[2]);
-	Material.Ambient = FColor(OutAmbient[0], OutAmbient[1], OutAmbient[2]);
-	Material.Specular = FColor(OutSpecular[0], OutSpecular[1], OutSpecular[2]);
-	Material.Shininess = OutShininess;
-	Material.Transparency = OutTransparency;
-	Material.Reflexion = OutReflexion;
-	Material.TextureName = OutTextureName;
-}
 
-bool FCADSceneGraphDescriptionFile::GetColor(uint32 ColorHId, FColor& Color) const
-{
-	uint32 ColorId;
-	uint8 Alpha;
-	CADLibrary::UnhashFastColorHash(ColorHId, ColorId, Alpha);
-	const FColor* ColorOpac = ColorIdToColor.Find(ColorId);
-	if (ColorOpac)
-	{
-		Color = *ColorOpac;
-		Color.A = Alpha;
-		return true;
-	}
-	return false;
-}
-
-bool FCADSceneGraphDescriptionFile::GetMaterial(int32 MaterialId, CADLibrary::FCADMaterial& Material) const
-{
-	const CADLibrary::FCADMaterial* PMaterial = MaterialIdToMaterial.Find(MaterialId);
-	if (PMaterial)
+	void AddTransformToActor(CADLibrary::FArchiveInstance& Instance, TSharedPtr< IDatasmithActorElement > Actor, const CADLibrary::FImportParameters& ImportParameters)
 	{
-		Material = *PMaterial;
-		return true;
-	}
-	return false;
-}
-
-void FCADSceneGraphDescriptionFile::ReadFile()
-{
-	FFileHelper::LoadFileToStringArray(SceneGraphDescriptionData, *FileName);
-	if (SceneGraphDescriptionData.Num() < 10)
-		return;
-
-	FileName = FPaths::GetBaseFilename(FileName);
-
-	FString Left, Right;
-	FString HeaderCTId = SceneGraphDescriptionData[MAPCTIDLINE];
-	HeaderCTId.Split(SPACE, &Left, &Right);
-	int32 Line = FCString::Atoi(*Right);
-
-	TArray<FString> HeaderCTIdSplit;
-	SceneGraphDescriptionData[Line].ParseIntoArray(HeaderCTIdSplit, TEXT(" "));
-
-	if (HeaderCTIdSplit.Num() % 2 != 0)
-	{
-		return;
-	}
-
-	for (int32 index = 0; index < HeaderCTIdSplit.Num(); index += 2)
-	{
-		SceneNodeIdToSceneNodeDescriptionMap.Emplace(FCString::Atoi(*HeaderCTIdSplit[index]), FSceneNodeDescription(*this, FCString::Atoi(*HeaderCTIdSplit[index + 1])));
-	}
-
-	// Parse color
-	HeaderCTId = SceneGraphDescriptionData[COLORSETLINE];
-	HeaderCTId.Split(SPACE, &Left, &Right);
-	if (Right != TEXT("0"))
-	{
-		Line = FCString::Atoi(*Right);
-		FString ColorString = SceneGraphDescriptionData[Line];
-		while (!ColorString.IsEmpty())
-		{
-			TArray<FString> ColorData;
-			ColorString.ParseIntoArray(ColorData, TEXT(" "));
-			ColorIdToColor.Emplace(FCString::Atoi(*ColorData[0]), FColor(FCString::Atoi(*ColorData[1]), FCString::Atoi(*ColorData[2]), FCString::Atoi(*ColorData[3])));
-			Line++;
-			ColorString = SceneGraphDescriptionData[Line];
-		}
-	}
-	
-	// Parse material
-	HeaderCTId = SceneGraphDescriptionData[MATERIALSETLINE];
-	HeaderCTId.Split(SPACE, &Left, &Right);
-	if (Right != TEXT("0"))
-	{
-		Line = FCString::Atoi(*Right);
-		FString MaterialString = SceneGraphDescriptionData[Line];
-		while (!MaterialString.IsEmpty())
+		if (!Actor.IsValid())
 		{
-			CADLibrary::FCADMaterial Material;
-			GetMaterialDescription(Line, Material);
-			MaterialIdToMaterial.Add(Material.MaterialId, Material);
-			Line += 2;
-			MaterialString = SceneGraphDescriptionData[Line];
+			return;
 		}
-	}
-}
 
-void FCADSceneGraphDescriptionFile::SetMaterialMaps(TMap< uint32, FColor>& MaterialUuidToColor, TMap< uint32, CADLibrary::FCADMaterial>& MaterialUuidToMaterial)
-{
-	for (const auto& Color : ColorIdToColor)
-	{
-		uint32 ColorHash = CADLibrary::BuildColorHash(Color.Value);
-		MaterialUuidToColor.Add(ColorHash, Color.Value);
-	}
-	for (const auto& Material : MaterialIdToMaterial)
-	{
-		uint32 MaterialHash = CADLibrary::BuildMaterialHash(Material.Value);
-		MaterialUuidToMaterial.Add(MaterialHash, Material.Value);
-	}
-}
+		FTransform LocalTransform(Instance.TransformMatrix);
+		FTransform LocalUETransform = FDatasmithUtils::ConvertTransform((FDatasmithUtils::EModelCoordSystem) ImportParameters.ModelCoordSys, LocalTransform);
 
-void FSceneNodeDescription::GetMetaData(TMap<FString, FString>& MetaDataMap)
-{
-	const FString& RawDataString = SceneGraphDescription.GetString(Line+1);
-	if (RawDataString[0] != L'M')
-	{
-		return;
-	}
-
-	FString Left, Right;
-	RawDataString.Split(SPACE, &Left, &Right);
-	StartMeta = Line + 2;
-	EndMeta = FCString::Atoi(*Right);
-	MetaDataMap.Reserve(EndMeta / 2);
-	EndMeta += StartMeta;
-	for (int32 LineIndex = StartMeta; LineIndex < EndMeta; LineIndex += 2)
-	{
-		MetaDataMap.Add(SceneGraphDescription.GetString(LineIndex), SceneGraphDescription.GetString(LineIndex + 1));
-	}
-}
-
-void FSceneNodeDescription::GetMaterialSet(TMap<uint32, uint32>& MaterialIdToMaterialHashMap)
-{
-	const FString& MaterialString = SceneGraphDescription.GetString(EndMeta);
-	TArray<FString> MaterialIDToHashSet;
-	MaterialString.ParseIntoArray(MaterialIDToHashSet, TEXT(" "));
-	if (MaterialIDToHashSet[0] != TEXT("materialMap"))
-	{
-		return;
-	}
-
-	for (int32 index = 1; index < MaterialIDToHashSet.Num(); index++)
-	{
-		uint32 MaterialId = FCString::Atoi64(*MaterialIDToHashSet[index]); 
-		index++;
-		uint32 MaterialHash = FCString::Atoi64(*MaterialIDToHashSet[index]);
-		MaterialIdToMaterialHashMap.Add(MaterialId, MaterialHash);
-	}
-}
-
-void FSceneNodeDescription::GetChildren(TArray<int32>& Children)
-{
-	if(EndMeta == -1)
-	{
-		TMap<FString, FString> TempMap;
-		GetMetaData(TempMap);
-	}
-
-	const FString& DescriptionString = SceneGraphDescription.GetString(EndMeta);
-	if (DescriptionString[0] != L'c')
-	{
-		return;
-	}
-
-	FString Left, Right;
-	DescriptionString.Split(SPACE, &Left, &Right);
-	int32 nbChildren = FCString::Atoi(*Right);
-	Children.Reserve(nbChildren);
-
-	int32 StartChildren = EndMeta + 1;
-	int32 EndChildren = StartChildren + nbChildren;
-	for (int32 LineIndex = StartChildren; LineIndex < EndChildren; LineIndex += 1)
-	{
-		Children.Add(FCString::Atoi(*SceneGraphDescription.GetString(LineIndex)));
-	}
-}
-
-void FSceneNodeDescription::GetNodeReference(int32& RefId, FString& ExternalFile, NODE_TYPE& NodeType)
-{
-	RefId = 0;
-	ExternalFile = "";
-
-	if (EndMeta == -1)
-	{
-		TMap<FString, FString> TempMap;
-		GetMetaData(TempMap);
+		Actor->SetTranslation(LocalUETransform.GetTranslation() * ImportParameters.ScaleFactor);
+		Actor->SetScale(LocalUETransform.GetScale3D());
+		Actor->SetRotation(LocalUETransform.GetRotation());
 	}
 
-	const FString& DescriptionString = SceneGraphDescription.GetString(EndMeta + 1);
-
-	FString TypeStr, Right;
-	DescriptionString.Split(SPACE, &TypeStr, &Right);
-
-	if (TypeStr == TEXT("ref"))
-	{
-		RefId = FCString::Atoi(*Right);
-		NodeType = COMPONENT;
-	}
-	else if (TypeStr == TEXT("ext"))
-	{
-		FString RefIdString;
-		Right.Split(SPACE, &RefIdString, &ExternalFile);
-		RefId = FCString::Atoi(*RefIdString);
-		NodeType = EXTERNALCOMPONENT;
-	}
-	else
-	{
-		NodeType = UNDEFINED;
-		return;
-	}
 }
 
-void FSceneNodeDescription::AddTransformToActor(TSharedPtr< IDatasmithActorElement > Actor, const CADLibrary::FImportParameters& ImportParameters)
-{
-	if (!Actor.IsValid())
-	{
-		return;
-	}
-
-	const FString& RawDataString = SceneGraphDescription.GetString(EndMeta);
-	TArray<FString> TransformString;
-	RawDataString.ParseIntoArray(TransformString, TEXT(" "));
-
-	if (TransformString[0] != TEXT("matrix"))
-	{
-		return;
-	}
-
-	if (TransformString.Num() != 17)
-	{
-		return;
-	}
-
-	FMatrix Matrix;
-	float* MatrixFloats = (float*)Matrix.M;
-	for(int32 index = 0; index < 16; index++)
-	{
-		MatrixFloats[index] = FCString::Atof(*TransformString[index + 1]);
-	}
-
-	FTransform LocalTransform(Matrix);
-	FTransform LocalUETransform = FDatasmithUtils::ConvertTransform((FDatasmithUtils::EModelCoordSystem) ImportParameters.ModelCoordSys, LocalTransform);
-
-	FQuat Quat;
-	FDatasmithTransformUtils::GetRotation(LocalUETransform, Quat);
-	
-	Actor->SetTranslation(LocalUETransform.GetTranslation() * ImportParameters.ScaleFactor);
-	Actor->SetScale(LocalUETransform.GetScale3D());
-	Actor->SetRotation(Quat);
-}
-
-void FSceneNodeDescription::SetNodeType()
-{
-	if (Line == 0)
-	{
-		Type = UNDEFINED;
-	} 
-	else
-	{
-		const FString& RawDataString = SceneGraphDescription.GetString(Line);
-		switch (RawDataString[0])
-		{
-		case L'C':
-			Type = COMPONENT;
-			break;
-		case L'I':
-			Type = INSTANCE;
-			break;
-		case L'B':
-			Type = BODY;
-			break;
-		case L'E':
-			Type = EXTERNALCOMPONENT;
-			break;
-		default:
-			Type = UNDEFINED;
-			break;
-		}
-	}
-}
-
-uint32 FSceneNodeDescription::GetStaticMeshUuid()
-{
-	if (!BodyUUID)
-	{
-		BodyUUID = GetTypeHash(SceneGraphDescription.GetFileName());
-		BodyUUID = HashCombine(BodyUUID, GetTypeHash(*SceneGraphDescription.GetString(Line)));  // "B "+TEXT(CT_OBJECT_ID)
-	}
-	return BodyUUID;
-}
-
-bool FSceneNodeDescription::GetColor(int32 ColorHId, FColor& OutColor) const
-{
-	return SceneGraphDescription.GetColor(ColorHId, OutColor);
-}
-
-bool FSceneNodeDescription::GetMaterial(int32 MaterialId, CADLibrary::FCADMaterial& OutMaterial) const
-{
-	return SceneGraphDescription.GetMaterial(MaterialId, OutMaterial);
-}
-
-FSceneNodeDescription* FSceneNodeDescription::GetCTNode(int32 NodeId)
-{
-	return SceneGraphDescription.GetCTNode(NodeId);
-}
-
-FSceneNodeDescription::FSceneNodeDescription(FCADSceneGraphDescriptionFile& InSceneGraphDescription, int32 InLine)
-	: SceneGraphDescription(InSceneGraphDescription)
-	, Line(InLine)
-	, StartMeta(-1)
-	, EndMeta(-1)
-	, Type(UNKNOWN)
-	, BodyUUID(0)
-{
-	SetNodeType();
-}
-
-FDatasmithSceneGraphBuilder::FDatasmithSceneGraphBuilder(TMap<FString, FString>& InCADFileToUE4FileMap, TMap< TSharedPtr< IDatasmithMeshElement >, uint32 >& InMeshElementToCTBodyUuidMap, const FString& InCachePath, TSharedRef<IDatasmithScene> InScene, const FDatasmithSceneSource& InSource, const CADLibrary::FImportParameters& InImportParameters)
+FDatasmithSceneGraphBuilder::FDatasmithSceneGraphBuilder(TMap<FString, FString>& InCADFileToUE4FileMap, const FString& InCachePath, TSharedRef<IDatasmithScene> InScene, const FDatasmithSceneSource& InSource, const CADLibrary::FImportParameters& InImportParameters)
 	: CADFileToSceneGraphDescriptionFile(InCADFileToUE4FileMap)
-	, MeshElementToCADBodyUuidMap(InMeshElementToCTBodyUuidMap)
 	, CachePath(InCachePath)
 	, DatasmithScene(InScene)
 	, Source(InSource)
 	, ImportParameters(InImportParameters)
 {
-
 }
 
 bool FDatasmithSceneGraphBuilder::Build()
@@ -426,132 +68,120 @@ bool FDatasmithSceneGraphBuilder::Build()
 
 	const FString& RootFile = FPaths::GetCleanFilename(Source.GetSourceFile());
 
-	FCADSceneGraphDescriptionFile** RawData = CADFileToSceneGraphDescription.Find(RootFile);
-	if (RawData == nullptr || *RawData == nullptr)
+	CurrentMockUp = CADFileToArchiveMockUp.FindRef(RootFile);
+	if (!CurrentMockUp)
 	{
 		return false;
 	}
 
-	FSceneNodeDescription* RootNode = (*RawData)->GetCTNode(1);
-	if (RootNode == nullptr)
+	CadId RootId = 1;
+	int32* Index = CurrentMockUp->CADIdToComponentIndex.Find(RootId);
+	if (!Index)
 	{
 		return false;
 	}
 
 	ActorData Data(TEXT(""));
-	TSharedPtr< IDatasmithActorElement > RootActor = BuildNode(*RootNode, Data);
+	TSharedPtr< IDatasmithActorElement > RootActor = BuildComponent(CurrentMockUp->ComponentSet[*Index], Data);
 	DatasmithScene->AddActor(RootActor);
 
-
-	CADFileToSceneGraphDescription.Empty();
-	ArrayOfSceneGraphDescription.Empty();
-	BodyUuidToMeshElementMap.Empty();
-
 	return true;
 }
 
-
 void FDatasmithSceneGraphBuilder::LoadSceneGraphDescriptionFiles()
 {
-	ArrayOfSceneGraphDescription.Reserve(CADFileToSceneGraphDescriptionFile.Num());
-	CADFileToSceneGraphDescription.Reserve(CADFileToSceneGraphDescriptionFile.Num());
-	for (auto FilePair : CADFileToSceneGraphDescriptionFile)
+	ArchiveMockUps.Reserve(CADFileToSceneGraphDescriptionFile.Num());
+	CADFileToArchiveMockUp.Reserve(CADFileToSceneGraphDescriptionFile.Num());
+
+	for (const auto& FilePair : CADFileToSceneGraphDescriptionFile)
 	{
-		//OutSgFile = FilePair.Value;
-		FString RawDataFile = FPaths::Combine(CachePath, TEXT("scene"), FilePair.Value + TEXT(".sg"));
-		uint32 index = ArrayOfSceneGraphDescription.Emplace(RawDataFile);
-		ArrayOfSceneGraphDescription[index].SetMaterialMaps(MaterialUuidToColor, MaterialUuidToMaterial);
-		CADFileToSceneGraphDescription.Emplace(FilePair.Key, &ArrayOfSceneGraphDescription[index]);
-	}
-}
+		FString MockUpDescriptionFile = FPaths::Combine(CachePath, TEXT("scene"), FilePair.Value + TEXT(".sg"));
 
-TSharedPtr< IDatasmithActorElement >  FDatasmithSceneGraphBuilder::BuildNode(FSceneNodeDescription& Node, ActorData& ParentData)
-{
-	NODE_TYPE Type = Node.GetNodeType();
+		CADLibrary::FArchiveMockUp& MockUpDescription = ArchiveMockUps.Emplace_GetRef();
 
-	switch (Type)
-	{
-	case INSTANCE:
-		return BuildInstance(Node, ParentData);
+		CADFileToArchiveMockUp.Add(FilePair.Key, &MockUpDescription);
+	
+		CADLibrary::DeserializeMockUpFile(*MockUpDescriptionFile, MockUpDescription);
 
-	case COMPONENT:
-		return BuildComponent(Node, ParentData);
+		for(const auto& ColorPair : MockUpDescription.ColorHIdToColor)
+		{
+			ColorNameToColorArchive.Emplace(ColorPair.Value.UEMaterialName, ColorPair.Value);
+		}
 
-	case BODY:
-		return BuildBody(Node, ParentData);
+		for (const auto& MaterialPair : MockUpDescription.MaterialHIdToMaterial)
+		{
+			MaterialNameToMaterialArchive.Emplace(MaterialPair.Value.UEMaterialName, MaterialPair.Value);
+		}
 
-	default:
-		return TSharedPtr< IDatasmithActorElement >();
 	}
 }
 
-TSharedPtr< IDatasmithActorElement >  FDatasmithSceneGraphBuilder::BuildInstance(FSceneNodeDescription& Node, ActorData& ParentData)
+TSharedPtr< IDatasmithActorElement >  FDatasmithSceneGraphBuilder::BuildInstance(CADLibrary::FArchiveInstance& Instance, ActorData& ParentData)
 {
-	TMap<FString, FString> InstanceNodeMetaDataMap;
-	TMap<FString, FString> ReferenceNodeMetaDataMap;
-	Node.GetMetaData(InstanceNodeMetaDataMap);
-
-	FString ActorUUID;
-	FString ActorLabel;
-	GetNodeUUIDAndName(InstanceNodeMetaDataMap, ReferenceNodeMetaDataMap, ParentData.Uuid, ActorUUID, ActorLabel);
-
-	int32 RefId; 
-	NODE_TYPE NodeType;
-	FString ReferenceFile;
-	Node.GetNodeReference(RefId, ReferenceFile, NodeType);
+	CADLibrary::FArchiveComponent* Reference = nullptr;
+	CADLibrary::FArchiveComponent EmptyReference;
 
-	FSceneNodeDescription* ComponentNode = nullptr;
-	switch(NodeType)
+	CADLibrary::FArchiveMockUp* InstanceMockUp = CurrentMockUp;
+	if (Instance.bIsExternalRef)
 	{
-		case COMPONENT:
+		if (!Instance.ExternalRef.IsEmpty())
 		{
-			ComponentNode = Node.GetCTNode(RefId);
-			break;
+			CurrentMockUp = CADFileToArchiveMockUp.FindRef(Instance.ExternalRef);
+			if (CurrentMockUp)
+			{
+				CadId RootId = 1;
+				int32* Index = CurrentMockUp->CADIdToComponentIndex.Find(RootId);
+				if (Index)
+				{
+					Reference = &(CurrentMockUp->ComponentSet[*Index]);
+				}
+			}
 		}
 
-		case EXTERNALCOMPONENT:
+		if(!Reference)
 		{
-			FCADSceneGraphDescriptionFile** RawData = CADFileToSceneGraphDescription.Find(ReferenceFile);
-			if (RawData != nullptr && *RawData != nullptr)
+			CurrentMockUp = InstanceMockUp;
+			int32* Index = CurrentMockUp->CADIdToUnloadedComponentIndex.Find(Instance.ReferenceNodeId);
+			if (Index)
 			{
-				ComponentNode = (*RawData)->GetCTNode(1);
+				Reference = &(CurrentMockUp->UnloadedComponentSet[*Index]);
 			}
-			else
-			{
-				ComponentNode = Node.GetCTNode(RefId);
-			}
-			break;
 		}
-
-		default:
+	}
+	else {
+		int32* Index = CurrentMockUp->CADIdToComponentIndex.Find(Instance.ReferenceNodeId);
+		if (Index)
 		{
-			return TSharedPtr< IDatasmithActorElement >();
+			Reference = &(CurrentMockUp->ComponentSet[*Index]);
 		}
 	}
 
-	if (ComponentNode)
+	if (!Reference) // Should never append
 	{
-		ComponentNode->GetMetaData(ReferenceNodeMetaDataMap);
+		Reference = &EmptyReference;
 	}
 
+	FString ActorUUID;
+	FString ActorLabel;
+	GetNodeUUIDAndName(Instance.MetaData, Reference->MetaData, ParentData.Uuid, ActorUUID, ActorLabel);
+
 	TSharedPtr< IDatasmithActorElement > Actor = CreateActor(*ActorUUID, *ActorLabel);
 	if (!Actor.IsValid())
 	{
 		return TSharedPtr< IDatasmithActorElement >();
 	}
-	AddMetaData(Actor, InstanceNodeMetaDataMap, ReferenceNodeMetaDataMap);
+	AddMetaData(Actor, Instance.MetaData, Reference->MetaData);
 
 	ActorData InstanceData(*ActorUUID, ParentData);
 
-	GetMainMaterial(InstanceNodeMetaDataMap, ComponentNode, InstanceData, bMaterialPropagationIsTopDown);
-	GetMainMaterial(ReferenceNodeMetaDataMap, ComponentNode, InstanceData, bMaterialPropagationIsTopDown);
+	GetMainMaterial(Instance.MetaData, InstanceData, bMaterialPropagationIsTopDown);
+	GetMainMaterial(Reference->MetaData, InstanceData, bMaterialPropagationIsTopDown);
 
-	if (ComponentNode)
-	{
-		AddChildren(Actor, *ComponentNode, InstanceData);
-	}
+	AddChildren(Actor, *Reference, InstanceData);
+
+	AddTransformToActor(Instance, Actor, ImportParameters);
 
-	Node.AddTransformToActor(Actor, ImportParameters);
+	CurrentMockUp = InstanceMockUp;
 	return Actor;
 }
 
@@ -633,16 +263,14 @@ void FDatasmithSceneGraphBuilder::GetNodeUUIDAndName(
 	OutUEUUID = FString::Printf(TEXT("0x%08x"), UEUUID);
 }
 
-TSharedPtr< IDatasmithActorElement > FDatasmithSceneGraphBuilder::BuildComponent(FSceneNodeDescription& Node, ActorData& ParentData)
+TSharedPtr< IDatasmithActorElement > FDatasmithSceneGraphBuilder::BuildComponent(CADLibrary::FArchiveComponent& Component, ActorData& ParentData)
 {
 	TMap<FString, FString> InstanceNodeMetaDataMap;
-	TMap<FString, FString> ReferenceNodeMetaDataMap;
-	Node.GetMetaData(ReferenceNodeMetaDataMap);
 
 	FString ActorUUID = TEXT("");
 	FString ActorLabel = TEXT("");
 	
-	GetNodeUUIDAndName(InstanceNodeMetaDataMap, ReferenceNodeMetaDataMap, ParentData.Uuid, ActorUUID, ActorLabel);
+	GetNodeUUIDAndName(InstanceNodeMetaDataMap, Component.MetaData, ParentData.Uuid, ActorUUID, ActorLabel);
 
 	TSharedPtr< IDatasmithActorElement > Actor = CreateActor(*ActorUUID, *ActorLabel);
 	if (!Actor.IsValid())
@@ -650,27 +278,34 @@ TSharedPtr< IDatasmithActorElement > FDatasmithSceneGraphBuilder::BuildComponent
 		return TSharedPtr< IDatasmithActorElement >();
 	}
 
-	AddMetaData(Actor, InstanceNodeMetaDataMap, ReferenceNodeMetaDataMap);
+	AddMetaData(Actor, InstanceNodeMetaDataMap, Component.MetaData);
 
 	ActorData ComponentData(*ActorUUID, ParentData);
-	GetMainMaterial(ReferenceNodeMetaDataMap, &Node, ComponentData, bMaterialPropagationIsTopDown);
+	GetMainMaterial(Component.MetaData, ComponentData, bMaterialPropagationIsTopDown);
 
-	AddChildren(Actor, Node, ComponentData);
+	AddChildren(Actor, Component, ComponentData);
 
 	return Actor;
 }
 
-TSharedPtr< IDatasmithActorElement > FDatasmithSceneGraphBuilder::BuildBody(FSceneNodeDescription& Node, ActorData& ParentData)
+TSharedPtr< IDatasmithActorElement > FDatasmithSceneGraphBuilder::BuildBody(CADLibrary::FArchiveBody& Body, ActorData& ParentData)
 {
 	TMap<FString, FString> InstanceNodeMetaDataMap;
-	TMap<FString, FString> BodyNodeMetaDataMap;
-	Node.GetMetaData(BodyNodeMetaDataMap);
 	
-	FString BodyUUID = TEXT("");
-	FString BodyLabel = TEXT("");
-	GetNodeUUIDAndName(InstanceNodeMetaDataMap, BodyNodeMetaDataMap, ParentData.Uuid, BodyUUID, BodyLabel);
+	FString BodyUUID;
+	FString BodyLabel;
+	GetNodeUUIDAndName(InstanceNodeMetaDataMap, Body.MetaData, ParentData.Uuid, BodyUUID, BodyLabel);
+
+	// Apply materials on the current part
+	uint32 MaterialUuid = 0;
+	MaterialUuid = ParentData.MaterialUuid ? ParentData.MaterialUuid : ParentData.ColorUuid;
+
+	if (!Body.ColorFaceSet.Num())
+	{
+		Body.ColorFaceSet.Add(MaterialUuid);
+	}
 
-	TSharedPtr< IDatasmithMeshElement > MeshElement = FindOrAddMeshElement(Node, BodyLabel);
+	TSharedPtr< IDatasmithMeshElement > MeshElement = FindOrAddMeshElement(Body, BodyLabel);
 	if (!MeshElement.IsValid())
 	{
 		return TSharedPtr< IDatasmithActorElement >();
@@ -686,10 +321,6 @@ TSharedPtr< IDatasmithActorElement > FDatasmithSceneGraphBuilder::BuildBody(FSce
 	ActorElement->SetLabel(*BodyLabel);
 	ActorElement->SetStaticMeshPathName(MeshElement->GetName());
 
-	// Apply materials on the current part
-	uint32 MaterialUuid = 0;
-	MaterialUuid = ParentData.MaterialUuid ? ParentData.MaterialUuid : ParentData.ColorUuid;
-
 	if (MaterialUuid && ImportParameters.Propagation != CADLibrary::EDisplayDataPropagationMode::BodyOnly)
 	{
 		TSharedPtr< IDatasmithMaterialIDElement > PartMaterialIDElement = FindOrAddMaterial(MaterialUuid);
@@ -705,13 +336,12 @@ TSharedPtr< IDatasmithActorElement > FDatasmithSceneGraphBuilder::BuildBody(FSce
 	return ActorElement;
 }
 
-TSharedPtr< IDatasmithMeshElement > FDatasmithSceneGraphBuilder::FindOrAddMeshElement(FSceneNodeDescription& Node, FString& BodyName)
+TSharedPtr< IDatasmithMeshElement > FDatasmithSceneGraphBuilder::FindOrAddMeshElement(CADLibrary::FArchiveBody& Body, FString& BodyName)
 {
-	const uint32 ShellUuid = Node.GetStaticMeshUuid();
-	FString ShellUuidName = FString::Printf(TEXT("0x%08x"), ShellUuid);
+	FString ShellUuidName = FString::Printf(TEXT("0x%012u"), Body.MeshActorName);
 
 	// Look if geometry has not been already processed, return it if found
-	TSharedPtr< IDatasmithMeshElement >* MeshElementPtr = BodyUuidToMeshElementMap.Find(ShellUuid);
+	TSharedPtr< IDatasmithMeshElement >* MeshElementPtr = BodyUuidToMeshElement.Find(Body.MeshActorName);
 	if (MeshElementPtr != nullptr)
 	{
 		return *MeshElementPtr;
@@ -721,7 +351,7 @@ TSharedPtr< IDatasmithMeshElement > FDatasmithSceneGraphBuilder::FindOrAddMeshEl
 	MeshElement->SetLabel(*BodyName);
 	MeshElement->SetLightmapSourceUV(-1);
 
-	FString BodyFile = FString::Printf(TEXT("UEx%08x"), ShellUuid);
+	FString BodyFile = FString::Printf(TEXT("UEx%08x"), Body.MeshActorName);
 	MeshElement->SetFile(*FPaths::Combine(CachePath, TEXT("body"), BodyFile + TEXT(".ct")));
 
 	// TODO: Set bounding box 
@@ -732,25 +362,23 @@ TSharedPtr< IDatasmithMeshElement > FDatasmithSceneGraphBuilder::FindOrAddMeshEl
 	//	MeshElement->SetDimensions(BoundingBox[3] - BoundingBox[0], BoundingBox[4] - BoundingBox[1], BoundingBox[5] - BoundingBox[2], 0.0f);
 	//}
 
-	TMap<uint32, uint32> MaterialIdToMaterialHashMap;
-	Node.GetMaterialSet(MaterialIdToMaterialHashMap);
-
-	for (auto MaterialId2Hash : MaterialIdToMaterialHashMap)
+	// Currently we assume that face has only colors
+	TSet<uint32>& MaterialSet = Body.ColorFaceSet;
+	for (uint32 MaterialName : MaterialSet)
 	{
 		TSharedPtr< IDatasmithMaterialIDElement > PartMaterialIDElement;
-		PartMaterialIDElement = FindOrAddMaterial(MaterialId2Hash.Value);
+		PartMaterialIDElement = FindOrAddMaterial(MaterialName);
 
 		const TCHAR* MaterialIDElementName = PartMaterialIDElement->GetName();
 
 		TSharedPtr< IDatasmithMaterialIDElement > MaterialIDElement = FDatasmithSceneFactory::CreateMaterialId(MaterialIDElementName);
 
-		MeshElement->SetMaterial(MaterialIDElementName, MaterialId2Hash.Value);
+		MeshElement->SetMaterial(MaterialIDElementName, MaterialName);
 	}
 
 	DatasmithScene->AddMesh(MeshElement);
 
-	BodyUuidToMeshElementMap.Add(ShellUuid, MeshElement);
-	MeshElementToCADBodyUuidMap.Add(MeshElement, ShellUuid);
+	BodyUuidToMeshElement.Add(Body.MeshActorName, MeshElement);
 
 	return MeshElement;
 }
@@ -777,13 +405,13 @@ TSharedPtr<IDatasmithMaterialIDElement> FDatasmithSceneGraphBuilder::FindOrAddMa
 	}
 	else if (MaterialUuid > 0)
 	{
-		if (FColor* Color = MaterialUuidToColor.Find(MaterialUuid))
+		if (CADLibrary::FArchiveColor* Color = ColorNameToColorArchive.Find(MaterialUuid))
 		{
-			MaterialElement = CADLibrary::CreateUEPbrMaterialFromColor(*Color);
+			MaterialElement = CADLibrary::CreateUEPbrMaterialFromColor(Color->Color);
 		}
-		else if (CADLibrary::FCADMaterial* Material = MaterialUuidToMaterial.Find(MaterialUuid))
+		else if (CADLibrary::FArchiveMaterial* Material = MaterialNameToMaterialArchive.Find(MaterialUuid))
 		{
-			MaterialElement = CADLibrary::CreateUEPbrMaterialFromMaterial(*Material, DatasmithScene);
+			MaterialElement = CADLibrary::CreateUEPbrMaterialFromMaterial(Material->Material, DatasmithScene);
 		}
 
 		if (MaterialElement.IsValid())
@@ -796,6 +424,7 @@ TSharedPtr<IDatasmithMaterialIDElement> FDatasmithSceneGraphBuilder::FindOrAddMa
 	if (!MaterialElement.IsValid())
 	{
 		MaterialElement = GetDefaultMaterial();
+		MaterialUuidMap.Add(MaterialUuid, MaterialElement);
 	}
 
 	TSharedPtr< IDatasmithMaterialIDElement > MaterialIDElement = FDatasmithSceneFactory::CreateMaterialId(MaterialElement->GetName());
@@ -803,21 +432,6 @@ TSharedPtr<IDatasmithMaterialIDElement> FDatasmithSceneGraphBuilder::FindOrAddMa
 	return MaterialIDElement;
 }
 
-TSharedPtr<IDatasmithUEPbrMaterialElement> CreateDefaultUEPbrMaterial()
-{
-	// Take the Material diffuse color and connect it to the BaseColor of a UEPbrMaterial
-	TSharedRef<IDatasmithUEPbrMaterialElement> MaterialElement = FDatasmithSceneFactory::CreateUEPbrMaterial(TEXT("0"));
-	MaterialElement->SetLabel(TEXT("DefaultCADImportMaterial"));
-
-	FLinearColor LinearColor = FLinearColor::FromPow22Color(FColor(200, 200, 200, 255));
-	IDatasmithMaterialExpressionColor* ColorExpression = MaterialElement->AddMaterialExpression<IDatasmithMaterialExpressionColor>();
-	ColorExpression->SetName(TEXT("Diffuse Color"));
-	ColorExpression->GetColor() = LinearColor;
-	MaterialElement->GetBaseColor().SetExpression(ColorExpression);
-
-	return MaterialElement;
-}
-
 void FDatasmithSceneGraphBuilder::AddMetaData(TSharedPtr< IDatasmithActorElement > ActorElement, TMap<FString, FString>& InstanceNodeAttributeSetMap, TMap<FString, FString>& ReferenceNodeAttributeSetMap)
 {
 	// Initialize list of attributes not to pass as meta-data
@@ -893,7 +507,7 @@ void FDatasmithSceneGraphBuilder::AddMetaData(TSharedPtr< IDatasmithActorElement
 		MetaDataRefElement->AddProperty(KeyValueProperty);
 	}
 
-	for (auto& Attribute : InstanceNodeAttributeSetMap)
+	for (const auto& Attribute : InstanceNodeAttributeSetMap)
 	{
 		if (UnwantedAttributes.Contains(*Attribute.Key))
 		{
@@ -914,36 +528,43 @@ void FDatasmithSceneGraphBuilder::AddMetaData(TSharedPtr< IDatasmithActorElement
 
 }
 
-void FDatasmithSceneGraphBuilder::LinkActor(TSharedPtr< IDatasmithActorElement > ParentActor, TSharedPtr< IDatasmithActorElement > Actor)
+bool FDatasmithSceneGraphBuilder::DoesActorHaveChildrenOrIsAStaticMesh(const TSharedPtr< IDatasmithActorElement >& ActorElement)
 {
-	if (Actor.IsValid())
+	if (ActorElement != nullptr)
 	{
-		if (ParentActor.IsValid())
+		if (ActorElement->GetChildrenCount() > 0)
 		{
-			ParentActor->AddChild(Actor);
+			return true;
 		}
-		else
+		else if (ActorElement->IsA(EDatasmithElementType::StaticMeshActor))
 		{
-			DatasmithScene->AddActor(Actor);
+			const TSharedPtr< IDatasmithMeshActorElement >& MeshActorElement = StaticCastSharedPtr< IDatasmithMeshActorElement >(ActorElement);
+			return FCString::Strlen(MeshActorElement->GetStaticMeshPathName()) > 0;
 		}
 	}
+	return false;
 }
 
-void FDatasmithSceneGraphBuilder::AddChildren(TSharedPtr< IDatasmithActorElement > Actor, FSceneNodeDescription& ComponentNode, ActorData& ParentData)
+
+void FDatasmithSceneGraphBuilder::AddChildren(TSharedPtr< IDatasmithActorElement > Actor, CADLibrary::FArchiveComponent& Component, ActorData& ParentData)
 {
-	TArray<int32> ChildrenIds;
-	ComponentNode.GetChildren(ChildrenIds);
-	for (const auto& ChildId : ChildrenIds)
+	for (const int32 ChildId : Component.Children)
 	{
-		FSceneNodeDescription* ChildNode = ComponentNode.GetCTNode(ChildId);
-		if (ChildNode == nullptr)
+		if (int32* ChildNodeIndex = CurrentMockUp->CADIdToInstanceIndex.Find(ChildId))
 		{
-			continue;
+			TSharedPtr< IDatasmithActorElement > ChildActor = BuildInstance(CurrentMockUp->InstanceSet[*ChildNodeIndex], ParentData);
+			if (ChildActor.IsValid() && DoesActorHaveChildrenOrIsAStaticMesh(ChildActor))
+			{
+				Actor->AddChild(ChildActor);
+			}
 		}
-		TSharedPtr< IDatasmithActorElement > ChildActor = BuildNode(*ChildNode, ParentData);
-		if (ChildActor.IsValid())
+		if (int32* ChildNodeIndex = CurrentMockUp->CADIdToBodyIndex.Find(ChildId))
 		{
-			Actor->AddChild(ChildActor);
+			TSharedPtr< IDatasmithActorElement > ChildActor = BuildBody(CurrentMockUp->BodySet[*ChildNodeIndex], ParentData);
+			if (ChildActor.IsValid() && DoesActorHaveChildrenOrIsAStaticMesh(ChildActor))
+			{
+				Actor->AddChild(ChildActor);
+			}
 		}
 	}
 }
diff --git a/Engine/Plugins/Enterprise/DatasmithCADImporter/Source/DatasmithCADTranslator/Private/DatasmithSceneGraphBuilder.h b/Engine/Plugins/Enterprise/DatasmithCADImporter/Source/DatasmithCADTranslator/Private/DatasmithSceneGraphBuilder.h
--- a/Engine/Plugins/Enterprise/DatasmithCADImporter/Source/DatasmithCADTranslator/Private/DatasmithSceneGraphBuilder.h
+++ b/Engine/Plugins/Enterprise/DatasmithCADImporter/Source/DatasmithCADTranslator/Private/DatasmithSceneGraphBuilder.h
@@ -33,41 +33,16 @@ enum NODE_TYPE
 	UNKNOWN
 };
 
-
-class FSceneNodeDescription
+namespace CADLibrary		
 {
-public:
-	FSceneNodeDescription(FCADSceneGraphDescriptionFile& InRawData, int32 InLine);
-
-	NODE_TYPE GetNodeType() { return Type; }
-
-	void GetMetaData(TMap<FString, FString>& MetaDataMap);
-	void GetChildren(TArray<int32>& Children);
-	
-	void GetNodeReference(int32& OutRefId, FString& OutExternalFile, NODE_TYPE& OutType);
-	void AddTransformToActor(TSharedPtr< IDatasmithActorElement > Actor, const CADLibrary::FImportParameters& InImportParameters);
-
-	FSceneNodeDescription* GetCTNode(int32 NodeId);
-
-	uint32 GetStaticMeshUuid();
-	void GetMaterialSet(TMap<uint32, uint32>& MaterialIdToMaterialHashMap);
-
-	bool GetColor(int32 ColorHId, FColor& Color) const;
-	bool GetMaterial(int32 MaterialId, CADLibrary::FCADMaterial& Material) const;
-
-protected:
-	void SetNodeType();
-	FString SceneGraphFile;
-	FCADSceneGraphDescriptionFile& SceneGraphDescription;
-	int32 Line;
-	int32 StartMeta;
-	int32 EndMeta;
-	NODE_TYPE Type;
-
-	uint32 BodyUUID;
-	
-};
-
+class FArchiveMockUp;
+class ICADArchiveObject;
+class FArchiveInstance;
+class FArchiveComponent;
+class FArchiveBody;
+class FArchiveMaterial;
+class FArchiveColor;
+}
 
 class ActorData  //#ueent_CAD
 {
@@ -99,95 +74,57 @@ public:
 
 
 
-class FCADSceneGraphDescriptionFile
-{
-public:
-	FCADSceneGraphDescriptionFile(const FString& InFileName);
-
-	void ReadFile();
-
-	FSceneNodeDescription* GetCTNode(int32 nodeId)
-	{
-		return SceneNodeIdToSceneNodeDescriptionMap.Find(nodeId);
-	}
-
-	const FString& GetString(int32 line) const
-	{
-		return SceneGraphDescriptionData[line];
-	}
-
-	const FString& GetFileName() const
-	{
-		return FileName;
-	}
-
-	bool GetColor(uint32 ColorHId, FColor& Color) const;
-
-	bool GetMaterial(int32 MaterialId, CADLibrary::FCADMaterial& OutMaterial) const;
-
-	void SetMaterialMaps(TMap< uint32, FColor>& MaterialUuidToColor, TMap< uint32, CADLibrary::FCADMaterial>& MaterialUuidToMaterial);
-
-private:
-	void GetMaterialDescription(int32 LineNumber, CADLibrary::FCADMaterial& OutMaterial) const;
-
-	FString FileName;
-	TArray<FString> SceneGraphDescriptionData;
-	TMap<int32, FSceneNodeDescription> SceneNodeIdToSceneNodeDescriptionMap;
-	TMap<int32, FColor> ColorIdToColor;
-	TMap<int32, CADLibrary::FCADMaterial> MaterialIdToMaterial;
-};
-
-
 
 class FDatasmithSceneGraphBuilder
 {
 public:
-	FDatasmithSceneGraphBuilder(TMap<FString, FString>& InCADFileToUE4FileMap, TMap< TSharedPtr< IDatasmithMeshElement >, uint32 >& InMeshElementToCTBodyUuidMap, const FString& InCachePath, 
-		TSharedRef<IDatasmithScene> InScene, const FDatasmithSceneSource& InSource, const CADLibrary::FImportParameters& InImportParameters);
+	FDatasmithSceneGraphBuilder(
+		TMap<FString, FString>& InCADFileToUE4FileMap, 
+		const FString& InCachePath, 
+		TSharedRef<IDatasmithScene> InScene, 
+		const FDatasmithSceneSource& InSource, 
+		const CADLibrary::FImportParameters& InImportParameters);
 
 	bool Build();
 
-protected:
+private:
 	void LoadSceneGraphDescriptionFiles();
 
-	TSharedPtr< IDatasmithActorElement > BuildNode(FSceneNodeDescription& Node, ActorData& ParentData);
-	TSharedPtr< IDatasmithActorElement > BuildInstance(FSceneNodeDescription& Node, ActorData& ParentData);
-	TSharedPtr< IDatasmithActorElement > BuildComponent(FSceneNodeDescription& Node, ActorData& ParentData);
-	TSharedPtr< IDatasmithActorElement > BuildBody(FSceneNodeDescription& NodeId, ActorData& ParentData);
+	TSharedPtr< IDatasmithActorElement > BuildInstance(CADLibrary::FArchiveInstance& Instance, ActorData& ParentData);
+	TSharedPtr< IDatasmithActorElement > BuildComponent(CADLibrary::FArchiveComponent& Component, ActorData& ParentData);
+	TSharedPtr< IDatasmithActorElement > BuildBody(CADLibrary::FArchiveBody& Body, ActorData& ParentData);
 
 	void AddMetaData(TSharedPtr< IDatasmithActorElement > ActorElement, TMap<FString, FString>& InstanceNodeAttributeSetMap, TMap<FString, FString>& ReferenceNodeAttributeSetMap);
-	void AddChildren(TSharedPtr< IDatasmithActorElement > Actor, FSceneNodeDescription& Node, ActorData& ParentData);
-	void LinkActor(TSharedPtr< IDatasmithActorElement > ParentActor, TSharedPtr< IDatasmithActorElement > Actor);
+	void AddChildren(TSharedPtr< IDatasmithActorElement > Actor, CADLibrary::FArchiveComponent& Component, ActorData& ParentData);
+	bool DoesActorHaveChildrenOrIsAStaticMesh(const TSharedPtr< IDatasmithActorElement >& ActorElement);
 
 	TSharedPtr< IDatasmithUEPbrMaterialElement > GetDefaultMaterial();
 	TSharedPtr<IDatasmithMaterialIDElement> FindOrAddMaterial(uint32 MaterialUuid);
 
 	TSharedPtr< IDatasmithActorElement > CreateActor(const TCHAR* ActorUUID, const TCHAR* ActorLabel);
-	TSharedPtr< IDatasmithMeshElement > FindOrAddMeshElement(FSceneNodeDescription& Node, FString& InLabel);
+	TSharedPtr< IDatasmithMeshElement > FindOrAddMeshElement(CADLibrary::FArchiveBody& Body, FString& InLabel);
 
 	void GetNodeUUIDAndName(TMap<FString, FString>& InInstanceNodeMetaDataMap, TMap<FString, FString>& InReferenceNodeMetaDataMap, const TCHAR* InParentUEUUID, FString& OutUEUUID, FString& OutName);
 
-
 protected:
 	TMap<FString, FString>& CADFileToSceneGraphDescriptionFile;
-	TMap< TSharedPtr< IDatasmithMeshElement >, uint32 >& MeshElementToCADBodyUuidMap;
 	const FString& CachePath;
 	TSharedRef<IDatasmithScene> DatasmithScene;
 	const FDatasmithSceneSource& Source;
 	const CADLibrary::FImportParameters& ImportParameters;
 
-	TMap<FString, FCADSceneGraphDescriptionFile*> CADFileToSceneGraphDescription;
-
-	TArray<FCADSceneGraphDescriptionFile> ArrayOfSceneGraphDescription; //#ueent_CAD
+	TArray<CADLibrary::FArchiveMockUp> ArchiveMockUps;
+	TMap<FString, CADLibrary::FArchiveMockUp*> CADFileToArchiveMockUp;
 
-	TMap< uint32, TSharedPtr< IDatasmithMeshElement > > BodyUuidToMeshElementMap;
+	TMap< CADUUID, TSharedPtr< IDatasmithMeshElement > > BodyUuidToMeshElement;
 
-	TMap< uint32, TSharedPtr< IDatasmithUEPbrMaterialElement > > MaterialUuidMap;
+	TMap< CADUUID, TSharedPtr< IDatasmithUEPbrMaterialElement > > MaterialUuidMap;
 	TSharedPtr<IDatasmithUEPbrMaterialElement > DefaultMaterial;
 
-	TMap< uint32, FColor> MaterialUuidToColor;
-	TMap< uint32, CADLibrary::FCADMaterial> MaterialUuidToMaterial;
+	TMap<CADUUID, CADLibrary::FArchiveColor> ColorNameToColorArchive; 
+	TMap<CADUUID, CADLibrary::FArchiveMaterial> MaterialNameToMaterialArchive; 
 
+	CADLibrary::FArchiveMockUp* CurrentMockUp;
 
 	bool bPreferMaterial;
 	bool bMaterialPropagationIsTopDown;
diff --git a/Engine/Plugins/Enterprise/DatasmithCADImporter/Source/DatasmithDispatcher/Private/DatasmithDispatcher.cpp b/Engine/Plugins/Enterprise/DatasmithCADImporter/Source/DatasmithDispatcher/Private/DatasmithDispatcher.cpp
--- a/Engine/Plugins/Enterprise/DatasmithCADImporter/Source/DatasmithDispatcher/Private/DatasmithDispatcher.cpp
+++ b/Engine/Plugins/Enterprise/DatasmithCADImporter/Source/DatasmithDispatcher/Private/DatasmithDispatcher.cpp
@@ -207,7 +207,7 @@ void FDatasmithDispatcher::ProcessLocal()
 	{
 		FString FullPath = Task->FileName;
 
-		CADLibrary::FCoreTechFileParser FileParser(FullPath, ProcessCacheFolder, ImportParameters, true, true);
+		CADLibrary::FCoreTechFileParser FileParser(FullPath, ProcessCacheFolder, ImportParameters);
 		ETaskState ProcessResult = FileParser.ProcessFile();
 
 		ETaskState TaskState = ProcessResult;
@@ -224,7 +224,7 @@ void FDatasmithDispatcher::ProcessLocal()
 					AddTask(FPaths::Combine(CurrentPath, ExternalFile));
 				}
 			}
-			LinkCTFileToUnrealCacheFile(FileParser.GetFileName(), FileParser.GetSceneGraphFile(), FileParser.GetMeshFile());
+			LinkCTFileToUnrealCacheFile(FileParser.GetCADFileName(), FileParser.GetSceneGraphFile(), FileParser.GetMeshFileName());
 		}
 	}
 #endif // CAD_INTERFACE
diff --git a/Engine/Source/Programs/Enterprise/Datasmith/DatasmithCADWorker/Private/DatasmithCADWorkerImpl.cpp b/Engine/Source/Programs/Enterprise/Datasmith/DatasmithCADWorker/Private/DatasmithCADWorkerImpl.cpp
--- a/Engine/Source/Programs/Enterprise/Datasmith/DatasmithCADWorker/Private/DatasmithCADWorkerImpl.cpp
+++ b/Engine/Source/Programs/Enterprise/Datasmith/DatasmithCADWorker/Private/DatasmithCADWorkerImpl.cpp
@@ -125,7 +125,7 @@ void FDatasmithCADWorkerImpl::ProcessCommand(const FRunTaskCommand& RunTaskComma
 
 	FCompletedTaskCommand CompletedTask;
 #ifdef CAD_INTERFACE
-	CADLibrary::FCoreTechFileParser FileParser(FileToProcess, CachePath, ImportParameters, true, true);
+	CADLibrary::FCoreTechFileParser FileParser(FileToProcess, CachePath, ImportParameters);
 	CADLibrary::FCoreTechFileParser::EProcessResult ProcessResult = FileParser.ProcessFile();
 
 	CompletedTask.ProcessResult = ProcessResult;
@@ -134,7 +134,7 @@ void FDatasmithCADWorkerImpl::ProcessCommand(const FRunTaskCommand& RunTaskComma
 	{
 		CompletedTask.ExternalReferences = FileParser.GetExternalRefSet().Array();
 		CompletedTask.SceneGraphFileName = FileParser.GetSceneGraphFile();
-		CompletedTask.GeomFileName = FileParser.GetMeshFile();
+		CompletedTask.GeomFileName = FileParser.GetMeshFileName();
 	}
 #endif // CAD_INTERFACE
 	CommandIO.SendCommand(CompletedTask, Config::SendCommandTimeout_s);
