Fix Incremental UE4Edito.99237247.log

diff --git a/Engine/Plugins/Enterprise/DatasmithCADImporter/Source/DatasmithCADTranslator/Private/DatasmithSceneGraphBuilder.h b/Engine/Plugins/Enterprise/DatasmithCADImporter/Source/DatasmithCADTranslator/Private/DatasmithSceneGraphBuilder.h
--- a/Engine/Plugins/Enterprise/DatasmithCADImporter/Source/DatasmithCADTranslator/Private/DatasmithSceneGraphBuilder.h
+++ b/Engine/Plugins/Enterprise/DatasmithCADImporter/Source/DatasmithCADTranslator/Private/DatasmithSceneGraphBuilder.h
@@ -5,6 +5,7 @@
 #include "CoreMinimal.h"
 
 #include "CADData.h"
+#include "CADSceneGraph.h"
 #include "CoreTechHelper.h"
 #include "CoreTechTypes.h"
 #include "CTSession.h"
