Add Creo extension in the list of ZUp_RightHanded format
Improve FDatasmithSceneSource::SetSourceFile in case of creo extension (asm.1)

diff --git a/Engine/Plugins/Enterprise/DatasmithCADImporter/Source/DatasmithCADTranslator/Private/DatasmithCADTranslator.cpp b/Engine/Plugins/Enterprise/DatasmithCADImporter/Source/DatasmithCADTranslator/Private/DatasmithCADTranslator.cpp
--- a/Engine/Plugins/Enterprise/DatasmithCADImporter/Source/DatasmithCADTranslator/Private/DatasmithCADTranslator.cpp
+++ b/Engine/Plugins/Enterprise/DatasmithCADImporter/Source/DatasmithCADTranslator/Private/DatasmithCADTranslator.cpp
@@ -77,7 +77,10 @@ bool FDatasmithCADTranslator::LoadScene(TSharedRef<IDatasmithScene> DatasmithSce
 	}
 
 	ImportParameters.ModelCoordSys = CADLibrary::EModelCoordSystem::ZUp_RightHanded;
-	if (FileExtension == TEXT("sldprt") || FileExtension == TEXT("sldasm") || FileExtension == TEXT("iam") || FileExtension == TEXT("ipt"))
+	if (FileExtension == TEXT("sldprt") || FileExtension == TEXT("sldasm") || // Solidworks
+		FileExtension == TEXT("iam") || FileExtension == TEXT("ipt") || // Inventor
+		FileExtension.StartsWith(TEXT("asm")) || FileExtension.StartsWith(TEXT("creo")) || FileExtension.StartsWith(TEXT("prt")) // Creo
+		)
 	{
 		ImportParameters.ModelCoordSys = CADLibrary::EModelCoordSystem::YUp_RightHanded;
 		ImportParameters.DisplayPreference = CADLibrary::EDisplayPreference::ColorOnly;
diff --git a/Engine/Plugins/Enterprise/DatasmithImporter/Source/DatasmithImporter/Private/DatasmithSceneSource.cpp b/Engine/Plugins/Enterprise/DatasmithImporter/Source/DatasmithImporter/Private/DatasmithSceneSource.cpp
--- a/Engine/Plugins/Enterprise/DatasmithImporter/Source/DatasmithImporter/Private/DatasmithSceneSource.cpp
+++ b/Engine/Plugins/Enterprise/DatasmithImporter/Source/DatasmithImporter/Private/DatasmithSceneSource.cpp
@@ -14,6 +14,11 @@ void FDatasmithSceneSource::SetSourceFile(const FString& InFilePath)
 	FilePath = InFilePath;
 	SceneDeducedName = FPaths::GetBaseFilename(FilePath);
 	FileExtension = FPaths::GetExtension(FilePath);
+	if (FCString::IsNumeric(*FileExtension))
+	{
+		FileExtension = FPaths::GetExtension(SceneDeducedName) + TEXT(".") + FileExtension;
+		SceneDeducedName = FPaths::GetBaseFilename(SceneDeducedName);
+	}
 }
 
 void FDatasmithSceneSource::SetSceneName(const FString& InSceneName)
