Tessellation options are ignored when importing CAD file through Datasmith

diff --git a/Engine/Plugins/Enterprise/DatasmithCADImporter/Source/DatasmithDispatcher/Private/DatasmithDispatcher.cpp b/Engine/Plugins/Enterprise/DatasmithCADImporter/Source/DatasmithDispatcher/Private/DatasmithDispatcher.cpp
--- a/Engine/Plugins/Enterprise/DatasmithCADImporter/Source/DatasmithDispatcher/Private/DatasmithDispatcher.cpp
+++ b/Engine/Plugins/Enterprise/DatasmithCADImporter/Source/DatasmithDispatcher/Private/DatasmithDispatcher.cpp
@@ -118,7 +118,7 @@ void FDatasmithDispatcher::Process(bool bWithProcessor)
 					if (WorkerId < NumberOfWorkers + Config::MaxRestartAllowed)
 					{
 						Handler.~FDatasmithWorkerHandler();
-						new (&Handler) FDatasmithWorkerHandler(*this, ProcessCacheFolder, WorkerId);
+						new (&Handler) FDatasmithWorkerHandler(*this, ImportParameters, ProcessCacheFolder, WorkerId);
 						UE_LOG(LogDatasmithDispatcher, Warning, TEXT("Restarting worker (new worker: %d)"), WorkerId);
 					}
 					else if (bLogRestartError)
@@ -182,7 +182,7 @@ void FDatasmithDispatcher::SpawnHandlers()
 	WorkerHandlers.Reserve(NumberOfWorkers);
 	for (int32 Index = 0; Index < NumberOfWorkers; Index++)
 	{
-		WorkerHandlers.Emplace(*this, ProcessCacheFolder, GetNextWorkerId());
+		WorkerHandlers.Emplace(*this, ImportParameters, ProcessCacheFolder, GetNextWorkerId());
 	}
 }
 
diff --git a/Engine/Plugins/Enterprise/DatasmithCADImporter/Source/DatasmithDispatcher/Private/DatasmithWorkerHandler.cpp b/Engine/Plugins/Enterprise/DatasmithCADImporter/Source/DatasmithDispatcher/Private/DatasmithWorkerHandler.cpp
--- a/Engine/Plugins/Enterprise/DatasmithCADImporter/Source/DatasmithDispatcher/Private/DatasmithWorkerHandler.cpp
+++ b/Engine/Plugins/Enterprise/DatasmithCADImporter/Source/DatasmithDispatcher/Private/DatasmithWorkerHandler.cpp
@@ -38,7 +38,7 @@ static FString GetWorkerExecutablePath()
 	return ProcessorPath;
 }
 
-FDatasmithWorkerHandler::FDatasmithWorkerHandler(FDatasmithDispatcher& InDispatcher, const FString& InCachePath, uint32 Id)
+FDatasmithWorkerHandler::FDatasmithWorkerHandler(FDatasmithDispatcher& InDispatcher, const CADLibrary::FImportParameters& InImportParameters, FString& InCachePath, uint32 Id)
 	: Dispatcher(InDispatcher)
 	, WorkerState(EWorkerState::Uninitialized)
 	, ErrorState(EWorkerErrorState::Ok)
@@ -47,6 +47,7 @@ FDatasmithWorkerHandler::FDatasmithWorkerHandler(FDatasmithDispatcher& InDispatc
 {
 	ThreadName = FString(TEXT("DatasmithWorkerHandler_")) + FString::FromInt(Id);
 	IOThread = FThread(*ThreadName, [this]() { Run(); } );
+	ImportParametersCommand.ImportParameters = InImportParameters;
 }
 
 FDatasmithWorkerHandler::~FDatasmithWorkerHandler()
@@ -143,6 +144,11 @@ void FDatasmithWorkerHandler::RunInternal()
 					break;
 				}
 
+				if (CommandIO.SendCommand(ImportParametersCommand, Config::SendCommandTimeout_s))
+				{
+					UE_LOG(LogDatasmithDispatcher, Verbose, TEXT("Import Parameters sent"));
+				}
+
 				WorkerState = EWorkerState::Idle;
 				break;
 			}
diff --git a/Engine/Plugins/Enterprise/DatasmithCADImporter/Source/DatasmithDispatcher/Public/DatasmithWorkerHandler.h b/Engine/Plugins/Enterprise/DatasmithCADImporter/Source/DatasmithDispatcher/Public/DatasmithWorkerHandler.h
--- a/Engine/Plugins/Enterprise/DatasmithCADImporter/Source/DatasmithDispatcher/Public/DatasmithWorkerHandler.h
+++ b/Engine/Plugins/Enterprise/DatasmithCADImporter/Source/DatasmithDispatcher/Public/DatasmithWorkerHandler.h
@@ -4,6 +4,7 @@
 
 #include "CoreMinimal.h"
 
+#include "CADOptions.h"
 #include "DatasmithCommands.h"
 #include "DatasmithDispatcherNetworking.h"
 #include "DatasmithDispatcherTask.h"
@@ -41,7 +42,7 @@ class FDatasmithWorkerHandler
 	};
 
 public:
-	FDatasmithWorkerHandler(FDatasmithDispatcher& InDispatcher, const FString& InCachePath, uint32 Id);
+	FDatasmithWorkerHandler(FDatasmithDispatcher& InDispatcher, const CADLibrary::FImportParameters& InImportParameters, FString& InCachePath, uint32 Id);
 	~FDatasmithWorkerHandler();
 
 	void Run();
@@ -77,6 +78,7 @@ private:
 
 	// self
 	FString CachePath;
+	FImportParametersCommand ImportParametersCommand;
 	TOptional<FTask> CurrentTask;
 	bool bShouldTerminate;
 
