Fix reimport

diff --git a/Engine/Plugins/Enterprise/DatasmithCADImporter/Source/CADTools/Public/CADOptions.h b/Engine/Plugins/Enterprise/DatasmithCADImporter/Source/CADTools/Public/CADOptions.h
--- a/Engine/Plugins/Enterprise/DatasmithCADImporter/Source/CADTools/Public/CADOptions.h
+++ b/Engine/Plugins/Enterprise/DatasmithCADImporter/Source/CADTools/Public/CADOptions.h
@@ -48,6 +48,21 @@ namespace CADLibrary
 		EDisplayDataPropagationMode Propagation = EDisplayDataPropagationMode::TopDown;
 		EDisplayPreference DisplayPreference = EDisplayPreference::MaterialPrefered;
 		bool bScaleUVMap = true;
+
+		uint32 GetHash() const
+		{
+			uint32 Hash = 0; 
+			for (double Param : {MetricUnit, ScaleFactor, ChordTolerance, MaxEdgeLength, MaxNormalAngle})
+			{
+				Hash = HashCombine(Hash, GetTypeHash(Param));
+			}
+			for (uint32 Param : {uint32(StitchingTechnique), uint32(ModelCoordSys), uint32(Propagation), uint32(DisplayPreference)})
+			{
+				Hash = HashCombine(Hash, GetTypeHash(Param));
+			}
+			Hash = HashCombine(Hash, GetTypeHash(bScaleUVMap));
+			return Hash;
+		}
 	};
 
 	struct FMeshParameters
diff --git a/Engine/Plugins/Enterprise/DatasmithCADImporter/Source/DatasmithCADTranslator/Private/DatasmithSceneGraphBuilder.cpp b/Engine/Plugins/Enterprise/DatasmithCADImporter/Source/DatasmithCADTranslator/Private/DatasmithSceneGraphBuilder.cpp
--- a/Engine/Plugins/Enterprise/DatasmithCADImporter/Source/DatasmithCADTranslator/Private/DatasmithSceneGraphBuilder.cpp
+++ b/Engine/Plugins/Enterprise/DatasmithCADImporter/Source/DatasmithCADTranslator/Private/DatasmithSceneGraphBuilder.cpp
@@ -59,6 +59,7 @@ FDatasmithSceneGraphBuilder::FDatasmithSceneGraphBuilder(TMap<FString, FString>&
 	, DatasmithScene(InScene)
 	, Source(InSource)
 	, ImportParameters(InImportParameters)
+	, ImportParametersHash(ImportParameters.GetHash())
 {
 }
 
@@ -351,6 +352,18 @@ TSharedPtr< IDatasmithMeshElement > FDatasmithSceneGraphBuilder::FindOrAddMeshEl
 	MeshElement->SetLabel(*BodyName);
 	MeshElement->SetLightmapSourceUV(-1);
 
+	// Set MeshElement FileHash used for re-import task 
+	FMD5 MD5; // unique Value that define the mesh
+	MD5.Update(reinterpret_cast<const uint8*>(&ImportParametersHash), sizeof ImportParametersHash);
+	// the scene graph archive name that is define by the name and the stat of the file (creation date, size)
+	MD5.Update(reinterpret_cast<const uint8*>(CurrentMockUp->SceneGraphArchive.GetCharArray().GetData()), CurrentMockUp->SceneGraphArchive.GetCharArray().Num());
+	// MeshActorName
+	MD5.Update(reinterpret_cast<const uint8*>(&Body.MeshActorName), sizeof Body.MeshActorName);
+
+	FMD5Hash Hash;
+	Hash.Set(MD5);
+	MeshElement->SetFileHash(Hash);
+
 	FString BodyFile = FString::Printf(TEXT("UEx%08x"), Body.MeshActorName);
 	MeshElement->SetFile(*FPaths::Combine(CachePath, TEXT("body"), BodyFile + TEXT(".ct")));
 
diff --git a/Engine/Plugins/Enterprise/DatasmithCADImporter/Source/DatasmithCADTranslator/Private/DatasmithSceneGraphBuilder.h b/Engine/Plugins/Enterprise/DatasmithCADImporter/Source/DatasmithCADTranslator/Private/DatasmithSceneGraphBuilder.h
--- a/Engine/Plugins/Enterprise/DatasmithCADImporter/Source/DatasmithCADTranslator/Private/DatasmithSceneGraphBuilder.h
+++ b/Engine/Plugins/Enterprise/DatasmithCADImporter/Source/DatasmithCADTranslator/Private/DatasmithSceneGraphBuilder.h
@@ -113,6 +113,7 @@ protected:
 	TSharedRef<IDatasmithScene> DatasmithScene;
 	const FDatasmithSceneSource& Source;
 	const CADLibrary::FImportParameters& ImportParameters;
+	const uint32 ImportParametersHash;
 
 	TArray<CADLibrary::FArchiveMockUp> ArchiveMockUps;
 	TMap<FString, CADLibrary::FArchiveMockUp*> CADFileToArchiveMockUp;
diff --git a/Engine/Plugins/Enterprise/DatasmithCADImporter/Source/DatasmithWireTranslator/Private/DatasmithWireTranslator.cpp b/Engine/Plugins/Enterprise/DatasmithCADImporter/Source/DatasmithWireTranslator/Private/DatasmithWireTranslator.cpp
--- a/Engine/Plugins/Enterprise/DatasmithCADImporter/Source/DatasmithWireTranslator/Private/DatasmithWireTranslator.cpp
+++ b/Engine/Plugins/Enterprise/DatasmithCADImporter/Source/DatasmithWireTranslator/Private/DatasmithWireTranslator.cpp
@@ -112,6 +112,20 @@ public:
 	}
 };
 
+uint32 GetSceneFileHash(const FString& FullPath, const FString& FileName)
+{
+	FFileStatData FileStatData = IFileManager::Get().GetStatData(*FullPath);
+
+	int64 FileSize = FileStatData.FileSize;
+	FDateTime ModificationTime = FileStatData.ModificationTime;
+
+	uint32 FileHash = GetTypeHash(FileName);
+	FileHash = HashCombine(FileHash, GetTypeHash(FileSize));
+	FileHash = HashCombine(FileHash, GetTypeHash(ModificationTime));
+
+	return FileHash;
+}
+
 class FWireTranslatorImpl
 {
 public:
@@ -120,7 +134,7 @@ public:
 		, SceneName(FPaths::GetBaseFilename(InSceneFullName))
 		, CurrentPath(FPaths::GetPath(InSceneFullName))
 		, SceneFullPath(InSceneFullName)
-		, TessellationOptionsHash(0)
+		, SceneFileHash(0)
 		, AlRootNode(nullptr)
 		, FileVersion(0)
 		, ArchiveWireVersion(0)
@@ -205,7 +219,9 @@ private:
 	FString SceneFullPath;
 
 	FDatasmithTessellationOptions TessellationOptions;
-	uint32 TessellationOptionsHash;
+	// Hash value of the scene file used to check if the file has been modified for re-import
+	uint32 SceneFileHash;
+
 	AlDagNode* AlRootNode;
 
 	/** Table of correspondence between mesh identifier and associated Datasmith mesh element */
@@ -238,7 +254,7 @@ private:
 void FWireTranslatorImpl::SetTessellationOptions(const FDatasmithTessellationOptions& Options)
 {
 	TessellationOptions = Options;
-	TessellationOptionsHash = Options.GetHash();
+	SceneFileHash = HashCombine(Options.GetHash(), GetSceneFileHash(SceneFullPath, SceneName));
 }
 #endif
 
@@ -1400,7 +1416,7 @@ TSharedPtr< IDatasmithMeshElement > FWireTranslatorImpl::FindOrAddMeshElement(Al
 	MeshElement->SetLabel(*MeshInfo.Label);
 	MeshElement->SetLightmapSourceUV(-1);
 
-// TODO
+	// TODO
 	// Get bounding box saved by GPure
 	double BoundingBox[8][4];
 	ShellNode.boundingBox(BoundingBox);
@@ -1411,6 +1427,14 @@ TSharedPtr< IDatasmithMeshElement > FWireTranslatorImpl::FindOrAddMeshElement(Al
 	//	MeshElement->SetDimensions(BoundingBox[3] - BoundingBox[0], BoundingBox[4] - BoundingBox[1], BoundingBox[5] - BoundingBox[2], 0.0f);
 	//}
 
+	// Set MeshElement FileHash used for re-import task 
+	FMD5 MD5; // unique Value that define the mesh
+	MD5.Update(reinterpret_cast<const uint8*>(&SceneFileHash), sizeof SceneFileHash);
+	// MeshActor Name
+	MD5.Update(reinterpret_cast<const uint8*>(&ShellUUID), sizeof ShellUUID);
+	FMD5Hash Hash;
+	Hash.Set(MD5);
+	MeshElement->SetFileHash(Hash);
 
 	if (ShaderName)
 	{
