Don't build "Constructive Plane" body

diff --git a/Engine/Plugins/Enterprise/DatasmithCADImporter/Source/CADInterfaces/Private/CoreTechFileParser.cpp b/Engine/Plugins/Enterprise/DatasmithCADImporter/Source/CADInterfaces/Private/CoreTechFileParser.cpp
--- a/Engine/Plugins/Enterprise/DatasmithCADImporter/Source/CADInterfaces/Private/CoreTechFileParser.cpp
+++ b/Engine/Plugins/Enterprise/DatasmithCADImporter/Source/CADInterfaces/Private/CoreTechFileParser.cpp
@@ -794,8 +794,10 @@ bool FCoreTechFileParser::ReadComponent(CT_OBJECT_ID ComponentId, uint32 Default
 	CT_OBJECT_ID ChildId;
 	while ((ChildId = Children.IteratorIter()) != 0)
 	{
-		MockUpDescription.ComponentSet[Index].Children.Add(ChildId);
-		ReadNode(ChildId, DefaultMaterialHash);
+		if (ReadNode(ChildId, DefaultMaterialHash))
+		{
+			MockUpDescription.ComponentSet[Index].Children.Add(ChildId);
+		}
 	}
 
 	return true;
@@ -914,6 +916,19 @@ void FCoreTechFileParser::GetBodyTessellation(CT_OBJECT_ID BodyId, FBodyMesh& Ou
 
 bool FCoreTechFileParser::ReadBody(CT_OBJECT_ID BodyId, uint32 DefaultMaterialHash)
 {
+	// Is this body a constructive geometry ?
+	CT_LIST_IO FaceList;
+	CT_BODY_IO::AskFaces(BodyId, FaceList);
+	if (1 == FaceList.Count())
+	{
+		FaceList.IteratorInitialize();
+		FString Value;
+		GetStringMetaDataValue(FaceList.IteratorIter(), TEXT("Constructive Plane"), Value);
+		if (Value == TEXT("true"))
+		{
+			return false;
+		}
+	}
 
 	int32 Index = MockUpDescription.BodySet.Emplace(BodyId);
 	MockUpDescription.CADIdToBodyIndex.Add(BodyId, Index);
@@ -987,6 +1002,27 @@ void FCoreTechFileParser::GetAttributeValue(CT_ATTRIB_TYPE AttributType, int Ith
 	}
 }
 
+void FCoreTechFileParser::GetStringMetaDataValue(CT_OBJECT_ID NodeId, const TCHAR* InMetaDataName, FString& OutMetaDataValue)
+{
+	CT_STR FieldName;
+	CT_UINT32 IthAttrib = 0;
+	while (CT_OBJECT_IO::SearchAttribute(NodeId, CT_ATTRIB_STRING_METADATA, IthAttrib++) == IO_OK)
+	{
+		if (CT_CURRENT_ATTRIB_IO::AskStrField(ITH_STRING_METADATA_NAME, FieldName) != IO_OK) break;
+		if (!FCString::Strcmp(InMetaDataName, FieldName.toUnicode()))
+		{
+			CT_STR FieldStrValue;
+			if (CT_CURRENT_ATTRIB_IO::AskStrField(ITH_STRING_METADATA_VALUE, FieldStrValue) != IO_OK) break;
+			if (FieldStrValue.IsEmpty())
+			{
+				return;
+			}
+			OutMetaDataValue = FieldStrValue.toUnicode();
+			return;
+		}
+	}
+}
+
 void FCoreTechFileParser::ReadNodeMetaData(CT_OBJECT_ID NodeId, TMap<FString, FString>& OutMetaData)
 {
 	const FString ConfigName = TEXT("Configuration Name");
@@ -1001,10 +1037,6 @@ void FCoreTechFileParser::ReadNodeMetaData(CT_OBJECT_ID NodeId, TMap<FString, FS
 		}
 	}
 
-	CT_UINT32 NbAttributes;
-	CT_OBJECT_IO::AskNbAttributes(NodeId, CT_ATTRIB_ALL, NbAttributes);
-	CT_UINT32 ith_attrib = 0;
-
 	CT_SHOW_ATTRIBUTE IsShow = CT_UNKNOWN;
 	if (CT_OBJECT_IO::AskShowAttribute(NodeId, IsShow) == IO_OK)
 	{
@@ -1022,7 +1054,8 @@ void FCoreTechFileParser::ReadNodeMetaData(CT_OBJECT_ID NodeId, TMap<FString, FS
 		}
 	}
 
-	while (CT_OBJECT_IO::SearchAttribute(NodeId, CT_ATTRIB_ALL, ith_attrib++) == IO_OK)
+	CT_UINT32 IthAttrib = 0;
+	while (CT_OBJECT_IO::SearchAttribute(NodeId, CT_ATTRIB_ALL, IthAttrib++) == IO_OK)
 	{
 		// Get the current attribute type
 		CT_ATTRIB_TYPE       AttributeType;
diff --git a/Engine/Plugins/Enterprise/DatasmithCADImporter/Source/CADInterfaces/Public/CoreTechFileParser.h b/Engine/Plugins/Enterprise/DatasmithCADImporter/Source/CADInterfaces/Public/CoreTechFileParser.h
--- a/Engine/Plugins/Enterprise/DatasmithCADImporter/Source/CADInterfaces/Public/CoreTechFileParser.h
+++ b/Engine/Plugins/Enterprise/DatasmithCADImporter/Source/CADInterfaces/Public/CoreTechFileParser.h
@@ -100,6 +100,7 @@ private:
 	void SetFaceMainMaterial(FObjectDisplayDataId& InFaceMaterial, FObjectDisplayDataId& InBodyMaterial, FBodyMesh& BodyMesh, int32 FaceIndex);
 
 	void ReadNodeMetaData(CT_OBJECT_ID NodeId, TMap<FString, FString>& OutMetaData);
+	void GetStringMetaDataValue(CT_OBJECT_ID NodeId, const TCHAR* InMetaDataName, FString& OutMetaDataValue);
 
 	CT_FLAGS SetCoreTechImportOption(const FString& MainFileExt);
 	void GetAttributeValue(CT_ATTRIB_TYPE attrib_type, int ith_field, FString& value);
