Fix import DWG, fix dwg orientation

diff --git a/Engine/Plugins/Enterprise/DatasmithCADImporter/Source/CADInterfaces/Private/CoreTechTypes.cpp b/Engine/Plugins/Enterprise/DatasmithCADImporter/Source/CADInterfaces/Private/CoreTechTypes.cpp
--- a/Engine/Plugins/Enterprise/DatasmithCADImporter/Source/CADInterfaces/Private/CoreTechTypes.cpp
+++ b/Engine/Plugins/Enterprise/DatasmithCADImporter/Source/CADInterfaces/Private/CoreTechTypes.cpp
@@ -1,6 +1,8 @@
 // Copyright 1998-2019 Epic Games, Inc. All Rights Reserved.
 #include "CoreTechTypes.h"
 
+#include "Misc/Paths.h"
+
 #ifdef CAD_INTERFACE
 
 const char* CoreTechLicenseKey =
@@ -15,7 +17,8 @@ namespace CADLibrary
 	CT_IO_ERROR CTKIO_InitializeKernel(double Unit)
 	{
 		CT_STR appName = CoreTechLicenseKey;
-		return CT_KERNEL_IO::InitializeKernel(appName, Unit, 0.00001 / Unit);
+		FString KernelIOPath = FPaths::Combine(FPaths::EnginePluginsDir(), TEXT("Enterprise/DatasmithCADImporter/Binaries/Win64"));
+		return CT_KERNEL_IO::InitializeKernel(appName, Unit, 0.00001 / Unit, *KernelIOPath);
 	}
 
 	CT_IO_ERROR CTKIO_ShutdownKernel()
