Improve KernelIO initialization

diff --git a/Engine/Plugins/Enterprise/DatasmithCADImporter/Source/CADInterfaces/Private/CoreTechTypes.cpp b/Engine/Plugins/Enterprise/DatasmithCADImporter/Source/CADInterfaces/Private/CoreTechTypes.cpp
--- a/Engine/Plugins/Enterprise/DatasmithCADImporter/Source/CADInterfaces/Private/CoreTechTypes.cpp
+++ b/Engine/Plugins/Enterprise/DatasmithCADImporter/Source/CADInterfaces/Private/CoreTechTypes.cpp
@@ -17,7 +17,7 @@ namespace CADLibrary
 	CT_IO_ERROR CTKIO_InitializeKernel(double Unit)
 	{
 		CT_STR appName = CoreTechLicenseKey;
-		FString KernelIOPath = FPaths::Combine(FPaths::EnginePluginsDir(), TEXT("Enterprise/DatasmithCADImporter/Binaries/Win64"));
+		FString KernelIOPath = FPaths::Combine(FPaths::EnginePluginsDir(), KERNEL_IO_PLUGINSPATH);
 		return CT_KERNEL_IO::InitializeKernel(appName, Unit, 0.00001 / Unit, *KernelIOPath);
 	}
 
