Materials from JT file (not only) are not assigned
Default material was not set to the tesselation when face and body has no material

diff --git a/Engine/Plugins/Enterprise/DatasmithCADImporter/Source/CADInterfaces/Private/CoreTechFileParser.cpp b/Engine/Plugins/Enterprise/DatasmithCADImporter/Source/CADInterfaces/Private/CoreTechFileParser.cpp
--- a/Engine/Plugins/Enterprise/DatasmithCADImporter/Source/CADInterfaces/Private/CoreTechFileParser.cpp
+++ b/Engine/Plugins/Enterprise/DatasmithCADImporter/Source/CADInterfaces/Private/CoreTechFileParser.cpp
@@ -410,10 +410,10 @@ void FCoreTechFileParser::SetFaceMainMaterial(FObjectDisplayDataId& InFaceMateri
 		FaceTessellations.ColorName = Color.UEMaterialName;
 		BodyMesh.ColorSet.Add(Color.UEMaterialName);
 	}
-	else if(InFaceMaterial.DefaultMaterialName)
+	else if(InBodyMaterial.DefaultMaterialName)
 	{
-		FaceTessellations.ColorName = InFaceMaterial.DefaultMaterialName;
-		BodyMesh.ColorSet.Add(InFaceMaterial.DefaultMaterialName);
+		FaceTessellations.ColorName = InBodyMaterial.DefaultMaterialName;
+		BodyMesh.ColorSet.Add(InBodyMaterial.DefaultMaterialName);
 	}
 }
 
