Fix Import DWG worker
KernelIO has to be initialized with its library Path. This path is known relative to the engine folder.
So this path has to be send to the worker.

diff --git a/Engine/Plugins/Enterprise/DatasmithCADImporter/Source/CADInterfaces/Private/CoreTechFileParser.cpp b/Engine/Plugins/Enterprise/DatasmithCADImporter/Source/CADInterfaces/Private/CoreTechFileParser.cpp
--- a/Engine/Plugins/Enterprise/DatasmithCADImporter/Source/CADInterfaces/Private/CoreTechFileParser.cpp
+++ b/Engine/Plugins/Enterprise/DatasmithCADImporter/Source/CADInterfaces/Private/CoreTechFileParser.cpp
@@ -490,13 +490,13 @@ void FCoreTechFileParser::ReadMaterials()
 	}
 }
 
-FCoreTechFileParser::FCoreTechFileParser(const FString& InCADFullPath, const FString& InCachePath, const FImportParameters& ImportParams)
+FCoreTechFileParser::FCoreTechFileParser(const FString& InCADFullPath, const FString& InCachePath, const FImportParameters& ImportParams, const TCHAR* KernelIOPath)
 	: CachePath(InCachePath)
 	, FullPath(InCADFullPath)
 	, bNeedSaveCTFile(false)
 	, ImportParameters(ImportParams)
 {
-	CTKIO_InitializeKernel(ImportParameters.MetricUnit);
+	CTKIO_InitializeKernel(ImportParameters.MetricUnit, KernelIOPath);
 }
 
 FCoreTechFileParser::EProcessResult FCoreTechFileParser::ProcessFile()
@@ -531,6 +531,7 @@ FCoreTechFileParser::EProcessResult FCoreTechFileParser::ProcessFile()
 	MeshArchiveFilePath = FPaths::Combine(CachePath, TEXT("mesh"), MeshArchiveFile + TEXT(".gm"));
 
 	bool bNeedToProceed = true;
+#ifndef IGNORE_CACHE
 	if (IFileManager::Get().FileExists(*SceneGraphArchiveFilePath))
 	{
 		if (!IFileManager::Get().FileExists(*CTFilePath)) // the file is scene graph only because no CT file
@@ -553,7 +554,7 @@ FCoreTechFileParser::EProcessResult FCoreTechFileParser::ProcessFile()
 		LoadSceneGraphArchive(SceneGraphArchiveFilePath);
 		return EProcessResult::ProcessOk;
 	}
-
+#endif
 	// Process the file
 	return ReadFileWithKernelIO();
 }
diff --git a/Engine/Plugins/Enterprise/DatasmithCADImporter/Source/CADInterfaces/Private/CoreTechTypes.cpp b/Engine/Plugins/Enterprise/DatasmithCADImporter/Source/CADInterfaces/Private/CoreTechTypes.cpp
--- a/Engine/Plugins/Enterprise/DatasmithCADImporter/Source/CADInterfaces/Private/CoreTechTypes.cpp
+++ b/Engine/Plugins/Enterprise/DatasmithCADImporter/Source/CADInterfaces/Private/CoreTechTypes.cpp
@@ -14,11 +14,10 @@ namespace CADLibrary
 	// Note: CTKIO_* functions are not functionally useful.
 	// This wrapping allows a correct profiling of the CT API.
 
-	CT_IO_ERROR CTKIO_InitializeKernel(double Unit)
+	CT_IO_ERROR CTKIO_InitializeKernel(double Unit, const TCHAR* KernelIOPath)
 	{
 		CT_STR appName = CoreTechLicenseKey;
-		FString KernelIOPath = FPaths::Combine(FPaths::EnginePluginsDir(), KERNEL_IO_PLUGINSPATH);
-		return CT_KERNEL_IO::InitializeKernel(appName, Unit, 0.00001 / Unit, *KernelIOPath);
+		return CT_KERNEL_IO::InitializeKernel(appName, Unit, 0.00001 / Unit, KernelIOPath);
 	}
 
 	CT_IO_ERROR CTKIO_ShutdownKernel()
diff --git a/Engine/Plugins/Enterprise/DatasmithCADImporter/Source/CADInterfaces/Public/CoreTechFileParser.h b/Engine/Plugins/Enterprise/DatasmithCADImporter/Source/CADInterfaces/Public/CoreTechFileParser.h
--- a/Engine/Plugins/Enterprise/DatasmithCADImporter/Source/CADInterfaces/Public/CoreTechFileParser.h
+++ b/Engine/Plugins/Enterprise/DatasmithCADImporter/Source/CADInterfaces/Public/CoreTechFileParser.h
@@ -55,7 +55,13 @@ class CADINTERFACES_API FCoreTechFileParser
 public:
 	using EProcessResult = ECoretechParsingResult;
 
-	FCoreTechFileParser(const FString& InCTFullPath, const FString& InCachePath, const FImportParameters& ImportParams);
+	/**
+	 * @param InCTFullPath Full path of the CAD file to parse
+	 * @param InCachePath Full path of the cache in which the data will be saved
+	 * @param ImportParams Parameters that setting import data like mesh SAG...
+	 * @param KernelIOPath Full Path of KernelIO libraries (oda_translator.exe, ...). Mandatory to import DWG, or DGN files
+	 */
+	FCoreTechFileParser(const FString& InCTFullPath, const FString& InCachePath, const FImportParameters& ImportParams, const TCHAR* KernelIOPath = TEXT(""));
 
 	double GetMetricUnit() const { return 0.01; }
 
diff --git a/Engine/Plugins/Enterprise/DatasmithCADImporter/Source/CADInterfaces/Public/CoreTechTypes.h b/Engine/Plugins/Enterprise/DatasmithCADImporter/Source/CADInterfaces/Public/CoreTechTypes.h
--- a/Engine/Plugins/Enterprise/DatasmithCADImporter/Source/CADInterfaces/Public/CoreTechTypes.h
+++ b/Engine/Plugins/Enterprise/DatasmithCADImporter/Source/CADInterfaces/Public/CoreTechTypes.h
@@ -53,7 +53,7 @@ namespace CADLibrary
 	 * This methode set Tolerance to 0.00001 m whether 0.01 mm
 	 * @param MetricUnit: Length unit express in meter
 	 */
-	CADINTERFACES_API CT_IO_ERROR CTKIO_InitializeKernel(double MetricUnit);
+	CADINTERFACES_API CT_IO_ERROR CTKIO_InitializeKernel(double MetricUnit, const TCHAR* = TEXT(""));
 	CADINTERFACES_API CT_IO_ERROR CTKIO_ShutdownKernel();
 
 	CADINTERFACES_API CT_IO_ERROR CTKIO_UnloadModel();
diff --git a/Engine/Plugins/Enterprise/DatasmithCADImporter/Source/DatasmithDispatcher/Private/DatasmithDispatcher.cpp b/Engine/Plugins/Enterprise/DatasmithCADImporter/Source/DatasmithDispatcher/Private/DatasmithDispatcher.cpp
--- a/Engine/Plugins/Enterprise/DatasmithCADImporter/Source/DatasmithDispatcher/Private/DatasmithDispatcher.cpp
+++ b/Engine/Plugins/Enterprise/DatasmithCADImporter/Source/DatasmithDispatcher/Private/DatasmithDispatcher.cpp
@@ -202,12 +202,15 @@ void FDatasmithDispatcher::CloseHandlers()
 
 void FDatasmithDispatcher::ProcessLocal()
 {
+	FString KernelIOPath = FPaths::Combine(FPaths::EnginePluginsDir(), TEXT(KERNEL_IO_PLUGINSPATH));
+	KernelIOPath = FPaths::ConvertRelativePathToFull(KernelIOPath);
+
 #ifdef CAD_INTERFACE
 	while (TOptional<FTask> Task = GetNextTask())
 	{
 		FString FullPath = Task->FileName;
 
-		CADLibrary::FCoreTechFileParser FileParser(FullPath, ProcessCacheFolder, ImportParameters);
+		CADLibrary::FCoreTechFileParser FileParser(FullPath, ProcessCacheFolder, ImportParameters, *KernelIOPath);
 		ETaskState ProcessResult = FileParser.ProcessFile();
 
 		ETaskState TaskState = ProcessResult;
diff --git a/Engine/Plugins/Enterprise/DatasmithCADImporter/Source/DatasmithDispatcher/Private/DatasmithWorkerHandler.cpp b/Engine/Plugins/Enterprise/DatasmithCADImporter/Source/DatasmithDispatcher/Private/DatasmithWorkerHandler.cpp
--- a/Engine/Plugins/Enterprise/DatasmithCADImporter/Source/DatasmithDispatcher/Private/DatasmithWorkerHandler.cpp
+++ b/Engine/Plugins/Enterprise/DatasmithCADImporter/Source/DatasmithDispatcher/Private/DatasmithWorkerHandler.cpp
@@ -57,6 +57,9 @@ FDatasmithWorkerHandler::~FDatasmithWorkerHandler()
 
 void FDatasmithWorkerHandler::StartWorkerProcess()
 {
+	FString KernelIOPath = FPaths::Combine(FPaths::EnginePluginsDir(), TEXT(KERNEL_IO_PLUGINSPATH));
+	KernelIOPath = FPaths::ConvertRelativePathToFull(KernelIOPath);
+
 	ensure(ErrorState == EWorkerErrorState::Ok);
 	FString ProcessorPath = GetWorkerExecutablePath();
 	if (FPaths::FileExists(ProcessorPath))
@@ -72,6 +75,7 @@ void FDatasmithWorkerHandler::StartWorkerProcess()
 		CommandToProcess += TEXT(" -ServerPID ") + FString::FromInt(FPlatformProcess::GetCurrentProcessId());
 		CommandToProcess += TEXT(" -ServerPort ") + FString::FromInt(ListenPort);
 		CommandToProcess += TEXT(" -CacheDir \"") + CachePath + TEXT('"');
+		CommandToProcess += TEXT(" -KernelIOPath \"") + KernelIOPath + TEXT('"');
 		UE_LOG(LogDatasmithDispatcher, Verbose, TEXT("CommandToProcess: %s"), *CommandToProcess);
 
 		uint32 WorkerProcessId = 0;
diff --git a/Engine/Source/Programs/Enterprise/Datasmith/DatasmithCADWorker/Private/DatasmithCADWorker.cpp b/Engine/Source/Programs/Enterprise/Datasmith/DatasmithCADWorker/Private/DatasmithCADWorker.cpp
--- a/Engine/Source/Programs/Enterprise/Datasmith/DatasmithCADWorker/Private/DatasmithCADWorker.cpp
+++ b/Engine/Source/Programs/Enterprise/Datasmith/DatasmithCADWorker/Private/DatasmithCADWorker.cpp
@@ -37,18 +37,13 @@ bool HasParameter(int32 Argc, TCHAR* Argv[], const FString& InParam)
 
 int32 Main(int32 Argc, TCHAR * Argv[])
 {
-	FString ServerPID, ServerPort, CacheDirectory;
+	FString ServerPID, ServerPort, CacheDirectory, KernelIOPath;
 	GetParameter(Argc, Argv, "-ServerPID", ServerPID);
 	GetParameter(Argc, Argv, "-ServerPort", ServerPort);
 	GetParameter(Argc, Argv, "-CacheDir", CacheDirectory);
+	GetParameter(Argc, Argv, "-KernelIOPath", KernelIOPath);
 
-	FString BinaryDir = FPaths::Combine(FPaths::EngineDir(), TEXT("Plugins"), TEXT("Enterprise"), TEXT("DatasmithCADImporter"), TEXT("Binaries"), TEXT("Win64"));
-	FString KernelIOLibDir = FPaths::Combine(BinaryDir, TEXT("kernel_io.dll"));
-	FPlatformProcess::PushDllDirectory(*BinaryDir);
-	void* KernelIODllHandle = FPlatformProcess::GetDllHandle(*KernelIOLibDir);
-	FPlatformProcess::PopDllDirectory(*BinaryDir);
-
-	FDatasmithCADWorkerImpl Worker(FCString::Atoi(*ServerPID), FCString::Atoi(*ServerPort), CacheDirectory);
+	FDatasmithCADWorkerImpl Worker(FCString::Atoi(*ServerPID), FCString::Atoi(*ServerPort), KernelIOPath, CacheDirectory);
 	Worker.Run();
 
 	return 0;
diff --git a/Engine/Source/Programs/Enterprise/Datasmith/DatasmithCADWorker/Private/DatasmithCADWorkerImpl.cpp b/Engine/Source/Programs/Enterprise/Datasmith/DatasmithCADWorker/Private/DatasmithCADWorkerImpl.cpp
--- a/Engine/Source/Programs/Enterprise/Datasmith/DatasmithCADWorker/Private/DatasmithCADWorkerImpl.cpp
+++ b/Engine/Source/Programs/Enterprise/Datasmith/DatasmithCADWorker/Private/DatasmithCADWorkerImpl.cpp
@@ -16,9 +16,10 @@ DEFINE_LOG_CATEGORY(LogDatasmithCADWorker);
 using namespace DatasmithDispatcher;
 
 
-FDatasmithCADWorkerImpl::FDatasmithCADWorkerImpl(int32 InServerPID, int32 InServerPort, const FString& InCachePath)
+FDatasmithCADWorkerImpl::FDatasmithCADWorkerImpl(int32 InServerPID, int32 InServerPort, const FString& InKernelIOPath, const FString& InCachePath)
 	: ServerPID(InServerPID)
 	, ServerPort(InServerPort)
+	, KernelIOPath(InKernelIOPath)
 	, CachePath(InCachePath)
 	, PingStartCycle(0)
 {
@@ -125,7 +126,7 @@ void FDatasmithCADWorkerImpl::ProcessCommand(const FRunTaskCommand& RunTaskComma
 
 	FCompletedTaskCommand CompletedTask;
 #ifdef CAD_INTERFACE
-	CADLibrary::FCoreTechFileParser FileParser(FileToProcess, CachePath, ImportParameters);
+	CADLibrary::FCoreTechFileParser FileParser(FileToProcess, CachePath, ImportParameters, *KernelIOPath);
 	CADLibrary::FCoreTechFileParser::EProcessResult ProcessResult = FileParser.ProcessFile();
 
 	CompletedTask.ProcessResult = ProcessResult;
diff --git a/Engine/Source/Programs/Enterprise/Datasmith/DatasmithCADWorker/Private/DatasmithCADWorkerImpl.h b/Engine/Source/Programs/Enterprise/Datasmith/DatasmithCADWorker/Private/DatasmithCADWorkerImpl.h
--- a/Engine/Source/Programs/Enterprise/Datasmith/DatasmithCADWorker/Private/DatasmithCADWorkerImpl.h
+++ b/Engine/Source/Programs/Enterprise/Datasmith/DatasmithCADWorker/Private/DatasmithCADWorkerImpl.h
@@ -19,7 +19,7 @@ struct FImportParameters;
 class FDatasmithCADWorkerImpl
 {
 public:
-	FDatasmithCADWorkerImpl(int32 InServerPID, int32 InServerPort, const FString& InCachePath);
+	FDatasmithCADWorkerImpl(int32 InServerPID, int32 InServerPort, const FString& KernelIOPath, const FString& InCachePath);
 	bool Run();
 
 private:
@@ -36,6 +36,7 @@ private:
 
 	int32 ServerPID;
 	int32 ServerPort;
+	FString KernelIOPath;
 	FString CachePath;
 	CADLibrary::FImportParameters ImportParameters;
 	uint64 PingStartCycle;
