Fix github build

diff --git a/Engine/Plugins/Enterprise/DatasmithCADImporter/Source/DatasmithDispatcher/Private/DatasmithWorkerHandler.cpp b/Engine/Plugins/Enterprise/DatasmithCADImporter/Source/DatasmithDispatcher/Private/DatasmithWorkerHandler.cpp
--- a/Engine/Plugins/Enterprise/DatasmithCADImporter/Source/DatasmithDispatcher/Private/DatasmithWorkerHandler.cpp
+++ b/Engine/Plugins/Enterprise/DatasmithCADImporter/Source/DatasmithDispatcher/Private/DatasmithWorkerHandler.cpp
@@ -57,8 +57,11 @@ FDatasmithWorkerHandler::~FDatasmithWorkerHandler()
 
 void FDatasmithWorkerHandler::StartWorkerProcess()
 {
-	FString KernelIOPath = FPaths::Combine(FPaths::EnginePluginsDir(), TEXT(KERNEL_IO_PLUGINSPATH));
+	FString KernelIOPath;
+#ifdef CAD_INTERFACE
+	KernelIOPath = FPaths::Combine(FPaths::EnginePluginsDir(), TEXT(KERNEL_IO_PLUGINSPATH));
 	KernelIOPath = FPaths::ConvertRelativePathToFull(KernelIOPath);
+#endif
 
 	ensure(ErrorState == EWorkerErrorState::Ok);
 	FString ProcessorPath = GetWorkerExecutablePath();
