The crash is caused by a label bigger than 512 chars.
But behind this crash there is also an other problem: package path must be smaller than ~256 chars
To fix it we take care that label is smaller than 50 chars.
In the case of jira UE-84629, the name is a path (usual in Catia file), so the "CleanFilename" is get, otherwise, the last 20 chars are used.

diff --git a/Engine/Plugins/Enterprise/DatasmithCADImporter/Source/DatasmithCADTranslator/Private/DatasmithSceneGraphBuilder.cpp b/Engine/Plugins/Enterprise/DatasmithCADImporter/Source/DatasmithCADTranslator/Private/DatasmithSceneGraphBuilder.cpp
--- a/Engine/Plugins/Enterprise/DatasmithCADImporter/Source/DatasmithCADTranslator/Private/DatasmithSceneGraphBuilder.cpp
+++ b/Engine/Plugins/Enterprise/DatasmithCADImporter/Source/DatasmithCADTranslator/Private/DatasmithSceneGraphBuilder.cpp
@@ -51,6 +51,28 @@ namespace
 		Actor->SetRotation(LocalUETransform.GetRotation());
 	}
 
+	// Method to reduce the size of huge label. The length of the package path, based on label, cannot be bigger than ~256
+	void CleanName(FString& Label)
+	{
+		const int32 MaxLabelSize = 50; // If the label is smaller than this value, the label is not modified. This size of package name is "acceptable"
+		const int32 ReasonableLabelSize = 20; // If the label has to be cut, a label that is not too long is preferred. 
+		const int32 MinLabelSize = 5; // If the label is smaller than this value, the label is too much reduce. Therefore a ReasonableLabelSize is preferred 
+
+		if (Label.Len() < MaxLabelSize)
+		{
+			return;
+		}
+
+		FString NewLabel;
+		NewLabel = FPaths::GetCleanFilename(Label);
+		if ((NewLabel.Len() < MaxLabelSize) && (NewLabel.Len() > MinLabelSize))
+		{
+			Label = NewLabel;
+			return;
+		}
+
+		Label = Label.Right(ReasonableLabelSize);
+	}
 }
 
 FDatasmithSceneGraphBuilder::FDatasmithSceneGraphBuilder(TMap<FString, FString>& InCADFileToUE4FileMap, const FString& InCachePath, TSharedRef<IDatasmithScene> InScene, const FDatasmithSceneSource& InSource, const CADLibrary::FImportParameters& InImportParameters)
@@ -231,6 +253,7 @@ void FDatasmithSceneGraphBuilder::GetNodeUUIDAndName(
 
 	//ReferenceInstanceName = ReferenceName + TEXT("(") + (IOriginalName ? *IOriginalName : ReferenceName) + TEXT(")");
 	OutName = IOriginalName ? *IOriginalName : IName ? *IName : ReferenceName;
+	CleanName(OutName);
 
 	uint32 UEUUID = 0;
 	UEUUID = GetTypeHash(InParentUEUUID);
