Fix import of project with external reference in subfolder
Before, only the name of external references was saved in the task list, the full path was rebuild by adding the folder path of the root file
Now, the full path of external references is saved in the task list

diff --git a/Engine/Plugins/Enterprise/DatasmithCADImporter/Source/CADInterfaces/Private/CADSceneGraph.cpp b/Engine/Plugins/Enterprise/DatasmithCADImporter/Source/CADInterfaces/Private/CADSceneGraph.cpp
--- a/Engine/Plugins/Enterprise/DatasmithCADImporter/Source/CADInterfaces/Private/CADSceneGraph.cpp
+++ b/Engine/Plugins/Enterprise/DatasmithCADImporter/Source/CADInterfaces/Private/CADSceneGraph.cpp
@@ -74,7 +74,7 @@ FArchive& operator<<(FArchive& Ar, FArchiveMockUp& MockUp)
 	Ar << MockUp.ColorHIdToColor;
 	Ar << MockUp.MaterialHIdToMaterial;
 
-	Ar << MockUp.InstanceSet;
+	Ar << MockUp.Instances;
 	Ar << MockUp.ComponentSet;
 	Ar << MockUp.UnloadedComponentSet;
 	Ar << MockUp.BodySet;
diff --git a/Engine/Plugins/Enterprise/DatasmithCADImporter/Source/CADInterfaces/Private/CoreTechFileParser.cpp b/Engine/Plugins/Enterprise/DatasmithCADImporter/Source/CADInterfaces/Private/CoreTechFileParser.cpp
--- a/Engine/Plugins/Enterprise/DatasmithCADImporter/Source/CADInterfaces/Private/CoreTechFileParser.cpp
+++ b/Engine/Plugins/Enterprise/DatasmithCADImporter/Source/CADInterfaces/Private/CoreTechFileParser.cpp
@@ -628,7 +628,7 @@ FCoreTechFileParser::EProcessResult FCoreTechFileParser::ReadFileWithKernelIO()
 	MockUpDescription.BodySet.Reserve(NbElements[CT_BODY_INDEX]);
 	MockUpDescription.ComponentSet.Reserve(NbElements[CT_ASSEMBLY_INDEX] + NbElements[CT_PART_INDEX] + NbElements[CT_COMPONENT_INDEX]);
 	MockUpDescription.UnloadedComponentSet.Reserve(NbElements[CT_UNLOADED_COMPONENT_INDEX] + NbElements[CT_UNLOADED_ASSEMBLY_INDEX] + NbElements[CT_UNLOADED_PART_INDEX]);
-	MockUpDescription.InstanceSet.Reserve(NbElements[CT_INSTANCE_INDEX]);
+	MockUpDescription.Instances.Reserve(NbElements[CT_INSTANCE_INDEX]);
 
 	MockUpDescription.CADIdToBodyIndex.Reserve(NbElements[CT_BODY_INDEX]);
 	MockUpDescription.CADIdToComponentIndex.Reserve(NbElements[CT_ASSEMBLY_INDEX] + NbElements[CT_PART_INDEX] + NbElements[CT_COMPONENT_INDEX]);
@@ -808,11 +808,11 @@ bool FCoreTechFileParser::ReadInstance(CT_OBJECT_ID InstanceNodeId, uint32 Defau
 {
 	NodeConfiguration.Empty();
 
-	int32 Index = MockUpDescription.InstanceSet.Emplace(InstanceNodeId);
+	int32 Index = MockUpDescription.Instances.Emplace(InstanceNodeId);
 	MockUpDescription.CADIdToInstanceIndex.Add(InstanceNodeId, Index);
-	ReadNodeMetaData(InstanceNodeId, MockUpDescription.InstanceSet[Index].MetaData);
+	ReadNodeMetaData(InstanceNodeId, MockUpDescription.Instances[Index].MetaData);
 
-	if (uint32 MaterialHash = GetObjectMaterial(MockUpDescription.InstanceSet[Index]))
+	if (uint32 MaterialHash = GetObjectMaterial(MockUpDescription.Instances[Index]))
 	{
 		DefaultMaterialHash = MaterialHash;
 	}
@@ -821,7 +821,7 @@ bool FCoreTechFileParser::ReadInstance(CT_OBJECT_ID InstanceNodeId, uint32 Defau
 	double Matrix[16];
 	if (CT_INSTANCE_IO::AskTransformation(InstanceNodeId, Matrix) == IO_OK)
 	{
-		float* MatrixFloats = (float*)MockUpDescription.InstanceSet[Index].TransformMatrix.M;
+		float* MatrixFloats = (float*)MockUpDescription.Instances[Index].TransformMatrix.M;
 		for (int32 index = 0; index < 16; index++)
 		{
 			MatrixFloats[index] = (float) Matrix[index];
@@ -833,27 +833,29 @@ bool FCoreTechFileParser::ReadInstance(CT_OBJECT_ID InstanceNodeId, uint32 Defau
 	CT_IO_ERROR CTReturn = CT_INSTANCE_IO::AskChild(InstanceNodeId, ReferenceNodeId);
 	if (CTReturn != CT_IO_ERROR::IO_OK)
 		return false;
-	MockUpDescription.InstanceSet[Index].ReferenceNodeId = ReferenceNodeId;
+	MockUpDescription.Instances[Index].ReferenceNodeId = ReferenceNodeId;
 
 	CT_OBJECT_TYPE type;
 	CT_OBJECT_IO::AskType(ReferenceNodeId, type);
 	if (type == CT_UNLOADED_PART_TYPE || type == CT_UNLOADED_COMPONENT_TYPE || type == CT_UNLOADED_ASSEMBLY_TYPE)
 	{
-		MockUpDescription.InstanceSet[Index].bIsExternalRef = true;
+		MockUpDescription.Instances[Index].bIsExternalRef = true;
 
 		CT_STR ComponentFile, FileType;
 		CT_COMPONENT_IO::AskExternalDefinition(ReferenceNodeId, ComponentFile, FileType);
-		MockUpDescription.InstanceSet[Index].ExternalRef = FPaths::GetCleanFilename(ComponentFile.toUnicode());
+		FString ExternalRefFullPath = ComponentFile.toUnicode();
+
+		MockUpDescription.Instances[Index].ExternalRef = FPaths::GetCleanFilename(ExternalRefFullPath);
 		if(!NodeConfiguration.IsEmpty())
 		{
-			MockUpDescription.InstanceSet[Index].ExternalRef += TEXT("|") + NodeConfiguration;
+			MockUpDescription.Instances[Index].ExternalRef += TEXT("|") + NodeConfiguration;
 		}
 
-		MockUpDescription.ExternalRefSet.Add(MockUpDescription.InstanceSet[Index].ExternalRef);
+		MockUpDescription.ExternalRefSet.Add(ExternalRefFullPath);
 	}
 	else
 	{
-		MockUpDescription.InstanceSet[Index].bIsExternalRef = false;
+		MockUpDescription.Instances[Index].bIsExternalRef = false;
 	}
 
 	return ReadNode(ReferenceNodeId, DefaultMaterialHash);
diff --git a/Engine/Plugins/Enterprise/DatasmithCADImporter/Source/CADInterfaces/Public/CADSceneGraph.h b/Engine/Plugins/Enterprise/DatasmithCADImporter/Source/CADInterfaces/Public/CADSceneGraph.h
--- a/Engine/Plugins/Enterprise/DatasmithCADImporter/Source/CADInterfaces/Public/CADSceneGraph.h
+++ b/Engine/Plugins/Enterprise/DatasmithCADImporter/Source/CADInterfaces/Public/CADSceneGraph.h
@@ -138,7 +138,7 @@ public:
 	TArray<FArchiveBody> BodySet;
 	TArray<FArchiveComponent> ComponentSet;
 	TArray<FArchiveUnloadedComponent> UnloadedComponentSet;
-	TArray<FArchiveInstance> InstanceSet;
+	TArray<FArchiveInstance> Instances;
 
 	TMap<CadId, int32> CADIdToBodyIndex;
 	TMap<CadId, int32> CADIdToComponentIndex;
diff --git a/Engine/Plugins/Enterprise/DatasmithCADImporter/Source/DatasmithCADTranslator/Private/DatasmithSceneGraphBuilder.cpp b/Engine/Plugins/Enterprise/DatasmithCADImporter/Source/DatasmithCADTranslator/Private/DatasmithSceneGraphBuilder.cpp
--- a/Engine/Plugins/Enterprise/DatasmithCADImporter/Source/DatasmithCADTranslator/Private/DatasmithSceneGraphBuilder.cpp
+++ b/Engine/Plugins/Enterprise/DatasmithCADImporter/Source/DatasmithCADTranslator/Private/DatasmithSceneGraphBuilder.cpp
@@ -588,7 +588,7 @@ void FDatasmithSceneGraphBuilder::AddChildren(TSharedPtr< IDatasmithActorElement
 	{
 		if (int32* ChildNodeIndex = CurrentMockUp->CADIdToInstanceIndex.Find(ChildId))
 		{
-			TSharedPtr< IDatasmithActorElement > ChildActor = BuildInstance(CurrentMockUp->InstanceSet[*ChildNodeIndex], ParentData);
+			TSharedPtr< IDatasmithActorElement > ChildActor = BuildInstance(CurrentMockUp->Instances[*ChildNodeIndex], ParentData);
 			if (ChildActor.IsValid() && DoesActorHaveChildrenOrIsAStaticMesh(ChildActor))
 			{
 				Actor->AddChild(ChildActor);
diff --git a/Engine/Plugins/Enterprise/DatasmithCADImporter/Source/DatasmithDispatcher/Private/DatasmithDispatcher.cpp b/Engine/Plugins/Enterprise/DatasmithCADImporter/Source/DatasmithDispatcher/Private/DatasmithDispatcher.cpp
--- a/Engine/Plugins/Enterprise/DatasmithCADImporter/Source/DatasmithDispatcher/Private/DatasmithDispatcher.cpp
+++ b/Engine/Plugins/Enterprise/DatasmithCADImporter/Source/DatasmithDispatcher/Private/DatasmithDispatcher.cpp
@@ -218,13 +218,12 @@ void FDatasmithDispatcher::ProcessLocal()
 
 		if (TaskState == ETaskState::ProcessOk)
 		{
-			FString CurrentPath = FPaths::GetPath(FullPath);
 			const TSet<FString>& ExternalRefSet = FileParser.GetExternalRefSet();
 			if (ExternalRefSet.Num() > 0)
 			{
 				for (const FString& ExternalFile : ExternalRefSet)
 				{
-					AddTask(FPaths::Combine(CurrentPath, ExternalFile));
+					AddTask(ExternalFile);
 				}
 			}
 			LinkCTFileToUnrealCacheFile(FileParser.GetCADFileName(), FileParser.GetSceneGraphFile(), FileParser.GetMeshFileName());
diff --git a/Engine/Plugins/Enterprise/DatasmithCADImporter/Source/DatasmithDispatcher/Private/DatasmithWorkerHandler.cpp b/Engine/Plugins/Enterprise/DatasmithCADImporter/Source/DatasmithDispatcher/Private/DatasmithWorkerHandler.cpp
--- a/Engine/Plugins/Enterprise/DatasmithCADImporter/Source/DatasmithDispatcher/Private/DatasmithWorkerHandler.cpp
+++ b/Engine/Plugins/Enterprise/DatasmithCADImporter/Source/DatasmithDispatcher/Private/DatasmithWorkerHandler.cpp
@@ -330,10 +330,9 @@ void FDatasmithWorkerHandler::ProcessCommand(FCompletedTaskCommand& CompletedTas
 		return;
 	}
 
-	FString CurrentPath = FPaths::GetPath(CurrentTask->FileName);
 	for (const FString& ExternalReferenceFile : CompletedTaskCommand.ExternalReferences)
 	{
-		Dispatcher.AddTask(FPaths::Combine(CurrentPath, ExternalReferenceFile));
+		Dispatcher.AddTask(ExternalReferenceFile);
 	}
 	FString CurrentFileName = FPaths::GetCleanFilename(CurrentTask->FileName);
 	Dispatcher.SetTaskState(CurrentTask->Index, CompletedTaskCommand.ProcessResult);
