Undo //UE4/Release-4.24/Engine/Plugins/Enterprise/DatasmithCADImporter/Source/... changelist 10076693

The revert of KernelIO sp2 to sp1 require to revert CL done to workaround kernelIO regression

diff --git a/Engine/Plugins/Enterprise/DatasmithCADImporter/Source/CADInterfaces/Private/CTSession.cpp b/Engine/Plugins/Enterprise/DatasmithCADImporter/Source/CADInterfaces/Private/CTSession.cpp
--- a/Engine/Plugins/Enterprise/DatasmithCADImporter/Source/CADInterfaces/Private/CTSession.cpp
+++ b/Engine/Plugins/Enterprise/DatasmithCADImporter/Source/CADInterfaces/Private/CTSession.cpp
@@ -31,11 +31,6 @@ namespace CADLibrary
 		return CTKIO_SaveFile(ObjectList, *FilePath, L"Ct");
 	}
 
-	CheckedCTError CTSession::CleanBRep()
-	{
-		return CTKIO_CleanBody(MainObjectId);
-	}
-
 	CheckedCTError CTSession::TopoFixes(double SewingToleranceFactor)
 	{
 		return CADLibrary::Repair(MainObjectId, ImportParams.StitchingTechnique, SewingToleranceFactor);
diff --git a/Engine/Plugins/Enterprise/DatasmithCADImporter/Source/CADInterfaces/Private/CoreTechTypes.cpp b/Engine/Plugins/Enterprise/DatasmithCADImporter/Source/CADInterfaces/Private/CoreTechTypes.cpp
--- a/Engine/Plugins/Enterprise/DatasmithCADImporter/Source/CADInterfaces/Private/CoreTechTypes.cpp
+++ b/Engine/Plugins/Enterprise/DatasmithCADImporter/Source/CADInterfaces/Private/CoreTechTypes.cpp
@@ -60,20 +60,6 @@ namespace CADLibrary
 		return CT_KERNEL_IO::LoadFile(FileName, MainObject, LoadFlags, Lod, StringOption);
 	}
 
-	CT_IO_ERROR CTKIO_CleanBody(CT_OBJECT_ID MainObjectId)
-	{
-		CT_STR format = "Ct";          
-		CT_STR system = "Evolution";
-
-		CT_IO_ERROR result = CT_CONVERT_IO::SetConvertConfiguration(system, format);
-		if (result == IO_OK)
-		{
-			result = CT_CONVERT_IO::Convert(MainObjectId);
-		}
-		return result;
-	}
-
-
 	CT_IO_ERROR CTKIO_SaveFile(CT_LIST_IO& objects_list_to_save, const TCHAR* file_name, const TCHAR* format, const CT_OBJECT_ID coordsystem)
 	{
 		return CT_KERNEL_IO::SaveFile(objects_list_to_save, file_name, format, coordsystem);
diff --git a/Engine/Plugins/Enterprise/DatasmithCADImporter/Source/CADInterfaces/Public/CTSession.h b/Engine/Plugins/Enterprise/DatasmithCADImporter/Source/CADInterfaces/Public/CTSession.h
--- a/Engine/Plugins/Enterprise/DatasmithCADImporter/Source/CADInterfaces/Public/CTSession.h
+++ b/Engine/Plugins/Enterprise/DatasmithCADImporter/Source/CADInterfaces/Public/CTSession.h
@@ -44,12 +44,6 @@ public:
 	 */
 	CheckedCTError TopoFixes(double SewingToleranceFactor = 100);
 
-	/**
-	 * In case of patch with cyclic boundary, a process has to be done by kernel_IO
-	 */
-	CheckedCTError CleanBRep();
-
-
 	/**
 	 * @param InScaleFactor : use to scale meshing from Kernel-IO
 	 */
diff --git a/Engine/Plugins/Enterprise/DatasmithCADImporter/Source/CADInterfaces/Public/CoreTechTypes.h b/Engine/Plugins/Enterprise/DatasmithCADImporter/Source/CADInterfaces/Public/CoreTechTypes.h
--- a/Engine/Plugins/Enterprise/DatasmithCADImporter/Source/CADInterfaces/Public/CoreTechTypes.h
+++ b/Engine/Plugins/Enterprise/DatasmithCADImporter/Source/CADInterfaces/Public/CoreTechTypes.h
@@ -12,7 +12,6 @@
 #pragma warning(push)
 #pragma warning(disable:4996) // unsafe sprintf
 #pragma warning(disable:4828) // illegal character
-#include "kernel_io/convert_io/convert_io.h"
 #include "kernel_io/kernel_io.h"
 #include "kernel_io/kernel_io_error.h"
 #include "kernel_io/kernel_io_type.h"
@@ -49,8 +48,14 @@ namespace CADLibrary
 	};
 
 	/**
+>>>> ORIGINAL //UE4/Release-4.24/Engine/Plugins/Enterprise/DatasmithCADImporter/Source/CADInterfaces/Public/CoreTechTypes.h#2
+	 * Mandatory: Kernel_IO has to be initialize with the final MetricUnit, set function doesn't work
+==== THEIRS //UE4/Release-4.24/Engine/Plugins/Enterprise/DatasmithCADImporter/Source/CADInterfaces/Public/CoreTechTypes.h#4
 	 * Mandatory: Kernel_IO has to be initialize with the final MetricUnit, set function doesn't work
 	 * This methode set Tolerance to 0.00001 m whether 0.01 mm
+==== YOURS //David.Lesage_YUL-Release-4.24/Engine/Plugins/Enterprise/DatasmithCADImporter/Source/CADInterfaces/Public/CoreTechTypes.h
+	 * Mandatory: Kernel_IO has to be initialize with the final MetricUnit, set function doesn`t work
+<<<<
 	 * @param MetricUnit: Length unit express in meter
 	 */
 	CADINTERFACES_API CT_IO_ERROR CTKIO_InitializeKernel(double MetricUnit, const TCHAR* = TEXT(""));
@@ -100,12 +105,6 @@ namespace CADLibrary
 		const TCHAR*  string_option = TEXT("")               
 	);
 
-	/**
-	 * In case of patch with cyclic boundary, a process has to be done by kernel_IO
-	 * @param MainObjectId: Body to process
-	 */
-	CADINTERFACES_API CT_IO_ERROR CTKIO_CleanBody(CT_OBJECT_ID MainObjectId);
-
 	CADINTERFACES_API CT_IO_ERROR CTKIO_SaveFile
 	(
 		CT_LIST_IO&        objects_list_to_save,
@@ -146,7 +145,7 @@ namespace CADLibrary
 		 * Make sure CT is initialized, and a main object is ready.
 		 * Handle input file unit and an output unit
 		 * @param FileMetricUnit number of meters per file unit.
-		 * e.g. For a file in inches, arg should be 0.0254
+		 * eg. For a file in inches, arg should be 0.0254
 		 */
 		CoreTechSessionBase(const TCHAR* Owner, double Unit);
 		bool IsSessionValid() { return Owner != nullptr && MainObjectId != 0; }
diff --git a/Engine/Plugins/Enterprise/DatasmithCADImporter/Source/DatasmithOpenNurbsTranslator/Private/DatasmithOpenNurbsTranslator.cpp b/Engine/Plugins/Enterprise/DatasmithCADImporter/Source/DatasmithOpenNurbsTranslator/Private/DatasmithOpenNurbsTranslator.cpp
--- a/Engine/Plugins/Enterprise/DatasmithCADImporter/Source/DatasmithOpenNurbsTranslator/Private/DatasmithOpenNurbsTranslator.cpp
+++ b/Engine/Plugins/Enterprise/DatasmithCADImporter/Source/DatasmithOpenNurbsTranslator/Private/DatasmithOpenNurbsTranslator.cpp
@@ -2839,17 +2839,10 @@ bool FOpenNurbsTranslatorImpl::TranslateBRep(ON_Brep* Brep, const ON_3dmObjectAt
 
 	LocalSession->ClearData();
 
-	bool bBRepHasLoopNonManifold = false;
-	Result = LocalSession->AddBRep(*Brep, bBRepHasLoopNonManifold);
+	Result = LocalSession->AddBRep(*Brep);
 
 	FString Filename = FString::Printf(TEXT("%s.ct"), *Name);
 	FString FilePath = FPaths::Combine(OutputPath, Filename);
-
-	if(bBRepHasLoopNonManifold)
-	{
-		LocalSession->CleanBRep();
-	}
-
 	Result = LocalSession->SaveBrep(FilePath);
 	if (Result)
 	{
diff --git a/Engine/Plugins/Enterprise/DatasmithCADImporter/Source/DatasmithOpenNurbsTranslator/Private/RhinoCoretechWrapper.cpp b/Engine/Plugins/Enterprise/DatasmithCADImporter/Source/DatasmithOpenNurbsTranslator/Private/RhinoCoretechWrapper.cpp
--- a/Engine/Plugins/Enterprise/DatasmithCADImporter/Source/DatasmithOpenNurbsTranslator/Private/RhinoCoretechWrapper.cpp
+++ b/Engine/Plugins/Enterprise/DatasmithCADImporter/Source/DatasmithOpenNurbsTranslator/Private/RhinoCoretechWrapper.cpp
@@ -4,10 +4,8 @@
 
 
 #ifdef CAD_LIBRARY
-
 #include "CoreTechHelper.h"
-
-#include "Containers/Set.h"
+//#include "TessellationHelper.h"
 
 #pragma warning(push)
 #pragma warning(disable:4265)
@@ -195,12 +193,15 @@ CT_OBJECT_ID CreateCTSurface(ON_NurbsSurface& Surface)
 	return CTSurfaceID;
 }
 
-void CreateCTFace_internal(const ON_BrepFace& Face, CT_LIST_IO& dest, ON_BoundingBox& outerBBox, ON_NurbsSurface& Surface, bool ignoreInner, bool& bOutFaceHasLoopNonManifold)
+void CreateCTFace_internal(const ON_BrepFace& Face, CT_LIST_IO& dest, ON_BoundingBox& outerBBox, ON_NurbsSurface& Surface, bool ignoreInner)
 {
 	CT_OBJECT_ID SurfaceID = CreateCTSurface(Surface);
 	if (SurfaceID == 0)
 		return;
 
+	ON_BoundingBox realOuterBBox = Face.OuterLoop()->BoundingBox();
+	bool outerRedefined = realOuterBBox != outerBBox;
+
 	int LoopCount = ignoreInner ? 1 : Face.LoopCount();
 	CT_LIST_IO Loops;
 	for (int LoopIndex = 0; LoopIndex < LoopCount; ++LoopIndex)
@@ -211,12 +212,10 @@ void CreateCTFace_internal(const ON_BrepFace& Face, CT_LIST_IO& dest, ON_Boundin
 
 		ON_BrepLoop::TYPE onLoopType = on_Loop.m_type;
 		bool bIsOuter = (onLoopType == ON_BrepLoop::TYPE::outer);
-
-		int TrimCount = on_Loop.TrimCount();
+		//assert(bIsOuter == (LoopIndex == 0));
 
 		CT_LIST_IO Coedges;
-		TSet<ON_BrepVertex*> VertexSet;
-		VertexSet.Reserve(TrimCount);
+		int TrimCount = on_Loop.TrimCount();
 		for (int i = 0; i < TrimCount; ++i)
 		{
 			CADLibrary::CheckedCTError err = IO_OK;
@@ -227,7 +226,7 @@ void CreateCTFace_internal(const ON_BrepFace& Face, CT_LIST_IO& dest, ON_Boundin
 				continue;
 
 			ON_NurbsCurve nurbs_curve;
-			int nurbFormSuccess = Trim.GetNurbForm(nurbs_curve); // 0:No ok 1:Ok 2:OkBut
+			int nurbFormSuccess = Trim.GetNurbForm(nurbs_curve); // 0:Nok 1:Ok 2:OkBut
 			if (nurbFormSuccess == 0)
 				continue;
 
@@ -236,13 +235,6 @@ void CreateCTFace_internal(const ON_BrepFace& Face, CT_LIST_IO& dest, ON_Boundin
 			if (err != IO_OK)
 				continue;
 
-			// check if the loop is non manifold i.e. is auto-intersecting like a 8. We get the second vertex of each edge. If the vertex is already in the set, this mean that the loop is at least like a 8
-			if (!bOutFaceHasLoopNonManifold)
-			{
-				ON_BrepVertex* Vertex = on_edge->Vertex(1);
-				VertexSet.Add(Vertex, &bOutFaceHasLoopNonManifold);
-			}
-
 			// fill edge data
 			CT_UINT32 order = nurbs_curve.Order();
 
@@ -270,6 +262,13 @@ void CreateCTFace_internal(const ON_BrepFace& Face, CT_LIST_IO& dest, ON_Boundin
 			{
 				nurbs_curve.GetCV(j, nurbs_curve.IsRational() ? ON::point_style::euclidean_rational : ON::point_style::not_rational, wp);
 			}
+			if (outerRedefined && bIsOuter)
+			{
+				for (int j = 0; j < cv_data.size(); j += ctrl_hull_dim)
+				{
+					cv_data[j] = min(max(cv_data[j], outerBBox.m_min.x), outerBBox.m_max.x);
+				}
+			}
 
 			// knot multiplicity (ignored as knots are stored multiple times already)
 			std::vector<CT_UINT32> knotMult;
@@ -311,7 +310,7 @@ void CreateCTFace_internal(const ON_BrepFace& Face, CT_LIST_IO& dest, ON_Boundin
 	dest.PushBack(FaceID);
 }
 
-void CreateCTFace(const ON_Brep& brep, const ON_BrepFace& Face, CT_LIST_IO& dest, bool& bOutFaceHasLoopNonManifold)
+void CreateCTFace(const ON_Brep& brep, const ON_BrepFace& Face, CT_LIST_IO& dest)
 {
 	const ON_BrepLoop* outerLoop = Face.OuterLoop();
 	if (outerLoop == nullptr)
@@ -339,12 +338,47 @@ void CreateCTFace(const ON_Brep& brep, const ON_BrepFace& Face, CT_LIST_IO& dest
 		}
 		bBadLoopHack &= hasSingularTrim;
 	}
-	bOutFaceHasLoopNonManifold = bBadLoopHack;
-	
+
+	if (bBadLoopHack)
+	{
+		// we try to split in 2 faces: one west one east.
+		double span = outerBBox.Diagonal().x;
+
+		ON_BoundingBox innerBBox = Face.Loop(1)->BoundingBox();
+		for (int LoopIndex = 2; LoopIndex < LoopCount; ++LoopIndex)
+		{
+			const ON_BrepLoop& innerLoop = *Face.Loop(LoopIndex);
+			innerBBox.Union(innerLoop.BoundingBox());
+		}
+		double spaceWest = innerBBox.Min().x - outerBBox.Min().x;
+		double spaceEast = outerBBox.Max().x - innerBBox.Max().x;
+		double e = 0.01;
+		if (spaceEast < e*span && spaceWest < e*span)
+		{
+			// can't split... We could ignore inner loops:
+			CreateCTFace_internal(Face, dest, outerBBox, Surface, true);
+			return;
+		}
+
+		// split the outer bbox in a normal part (with outer) and a 'rest' (with no iner)
+		ON_BoundingBox outerTrimedPartBBox = outerBBox;
+		if (spaceEast > spaceWest)
+		{
+			outerBBox.m_max.x -= 0.5 * spaceEast;
+			outerTrimedPartBBox.m_min.x = outerBBox.m_max.x;
+		}
+		else
+		{
+			outerBBox.m_min.x += 0.5 * spaceWest;
+			outerTrimedPartBBox.m_max.x = outerBBox.m_min.x;
+		}
+
+		// We now need to create a bonus face to handle the hole we just created
+		CreateCTFace_internal(Face, dest, outerTrimedPartBBox, Surface, true);
+	}
 #endif // FIX_HOLE_IN_WHOLE_FACE
-	bool bHasCyclicLoop = false;
-	CreateCTFace_internal(Face, dest, outerBBox, Surface, false, bHasCyclicLoop);
-	bOutFaceHasLoopNonManifold |= bHasCyclicLoop;
+
+	CreateCTFace_internal(Face, dest, outerBBox, Surface, false);
 }
 
 
@@ -353,7 +387,7 @@ CT_IO_ERROR FRhinoCoretechWrapper::Tessellate(FMeshDescription& Mesh, CADLibrary
 	return CADLibrary::Tessellate(MainObjectId, ImportParams, Mesh, MeshParameters);
 }
 
-CADLibrary::CheckedCTError FRhinoCoretechWrapper::AddBRep(ON_Brep& Brep, bool& bOutBRepHasLoopNonManifold)
+CADLibrary::CheckedCTError FRhinoCoretechWrapper::AddBRep(ON_Brep& Brep)
 {
 	CADLibrary::CheckedCTError Result;
 	if (!IsSessionValid())
@@ -369,9 +403,7 @@ CADLibrary::CheckedCTError FRhinoCoretechWrapper::AddBRep(ON_Brep& Brep, bool& b
 	for (int index = 0; index < FaceCount; index++)
 	{
 		const ON_BrepFace& on_face = Brep.m_F[index];
-		bool bIsFaceHasLoopNonManifold = false;
-		CreateCTFace(Brep, on_face, FaceList, bIsFaceHasLoopNonManifold);
-		bOutBRepHasLoopNonManifold |= bIsFaceHasLoopNonManifold;
+		CreateCTFace(Brep, on_face, FaceList);
 	}
 
 	if (FaceList.IsEmpty())
diff --git a/Engine/Plugins/Enterprise/DatasmithCADImporter/Source/DatasmithOpenNurbsTranslator/Private/RhinoCoretechWrapper.h b/Engine/Plugins/Enterprise/DatasmithCADImporter/Source/DatasmithOpenNurbsTranslator/Private/RhinoCoretechWrapper.h
--- a/Engine/Plugins/Enterprise/DatasmithCADImporter/Source/DatasmithOpenNurbsTranslator/Private/RhinoCoretechWrapper.h
+++ b/Engine/Plugins/Enterprise/DatasmithCADImporter/Source/DatasmithOpenNurbsTranslator/Private/RhinoCoretechWrapper.h
@@ -25,7 +25,7 @@ public:
 	{
 	}
 
-	CADLibrary::CheckedCTError AddBRep(ON_Brep& brep, bool& bOutBRepHasLoopNonManifold);
+	CADLibrary::CheckedCTError AddBRep(ON_Brep& brep);
 	static TSharedPtr<FRhinoCoretechWrapper> GetSharedSession(double SceneUnit, double ScaleFactor);
 
 	CT_IO_ERROR Tessellate(FMeshDescription& Mesh, CADLibrary::FMeshParameters& MeshParameters);
