Undo //UE4/Release-4.24/Engine/Plugins/Enterprise/DatasmithCADImporter/Source/CADLibrary/Private/... changelist 9961301

The revert of KernelIO sp2 to sp1 require to revert CL done to workaround kernelIO regression

This workaround force a minimum size of mesh edges of small parts to ensure their mesh.
With KernelIO sp1, all parts are well meshed

diff --git a/Engine/Plugins/Enterprise/DatasmithCADImporter/Source/CADLibrary/Private/CoreTechHelper.cpp b/Engine/Plugins/Enterprise/DatasmithCADImporter/Source/CADLibrary/Private/CoreTechHelper.cpp
--- a/Engine/Plugins/Enterprise/DatasmithCADImporter/Source/CADLibrary/Private/CoreTechHelper.cpp
+++ b/Engine/Plugins/Enterprise/DatasmithCADImporter/Source/CADLibrary/Private/CoreTechHelper.cpp
@@ -687,37 +687,6 @@ uint32 GetSize(CT_TESS_DATA_TYPE type)
 	return 0;
 }
 
-// Adapt the parameters to try and make sure a valid tessellation can be generated for the object
-FImportParameters GetAdjustedTessellationParameters(CT_OBJECT_ID MainObjectId, const FImportParameters& ImportParams)
-{
-	FImportParameters Result = ImportParams;
-
-	CT_COORDINATE Min;
-	CT_COORDINATE Max;
-	if (CT_OBJECT_IO::ComputeMinmax(MainObjectId, Min, Max) == CT_IO_ERROR::IO_OK)
-	{
-		double BBSizeMm = (FMath::Abs(Max.xyz[0] - Min.xyz[0]) +
-					       FMath::Abs(Max.xyz[1] - Min.xyz[1]) +
-					       FMath::Abs(Max.xyz[2] - Min.xyz[2])) / 3.0;
-
-		// Heuristic to force smaller tessellation options if the bounding box of the object is small
-		if (BBSizeMm > 0.0 && BBSizeMm < 5.0)
-		{
-			const double ClampedMaxEdgeLength = BBSizeMm / 5.0;
-
-			// Mirror some of the conversions that will happen in SetCoreTechTessellationState
-			if (Result.MaxEdgeLength < SMALL_NUMBER)
-			{
-				Result.MaxEdgeLength = DBL_MAX;
-			}
-			Result.MaxEdgeLength = (Result.MaxEdgeLength / Result.ScaleFactor < ClampedMaxEdgeLength) ?
-									Result.MaxEdgeLength : ClampedMaxEdgeLength * Result.ScaleFactor;
-		}
-	}
-
-	return Result;
-}
-
 CT_IO_ERROR Tessellate(CT_OBJECT_ID MainObjectId, const FImportParameters& ImportParams, FMeshDescription& MeshDesc, FMeshParameters& MeshParameters)
 {
 	CheckedCTError Result;
@@ -727,8 +696,7 @@ CT_IO_ERROR Tessellate(CT_OBJECT_ID MainObjectId, const FImportParameters& Impor
 	if (!Result)
 		return Result;
 
-	FImportParameters FixedParams = GetAdjustedTessellationParameters(MainObjectId, ImportParams);
-	SetCoreTechTessellationState(FixedParams);
+	SetCoreTechTessellationState(ImportParams);
 
 	FString FullPath;
 	FString CachePath;
