Copy up //UE4/Dev-Enterprise/...at  Changelist 10508152  to //UE4/Dev-Tools-Staging/...

diff --git a/Engine/Plugins/Enterprise/AxFImporter/Source/AxFImporter/Private/AxFImporter.cpp b/Engine/Plugins/Enterprise/AxFImporter/Source/AxFImporter/Private/AxFImporter.cpp
--- a/Engine/Plugins/Enterprise/AxFImporter/Source/AxFImporter/Private/AxFImporter.cpp
+++ b/Engine/Plugins/Enterprise/AxFImporter/Source/AxFImporter/Private/AxFImporter.cpp
@@ -2164,7 +2164,7 @@ private:
 // 					}
 					else if (TextureName == AXF_SVBRDF_TEXTURE_NAME_ANISO_ROTATION)
 					{
-						Log.Warn(TEXT("Anisotropic AxF materials are supported yet - will use isotropic approximation."));
+						Log.Warn(TEXT("Anisotropic AxF materials are not supported yet - will use isotropic approximation."));
 					}
 					else
 					{
diff --git a/Engine/Plugins/Enterprise/DataprepEditor/DataprepEditor.uplugin b/Engine/Plugins/Enterprise/DataprepEditor/DataprepEditor.uplugin
--- a/Engine/Plugins/Enterprise/DataprepEditor/DataprepEditor.uplugin
+++ b/Engine/Plugins/Enterprise/DataprepEditor/DataprepEditor.uplugin
@@ -3,8 +3,8 @@
 	"Version": 1,
 	"VersionName": "0.9",
 	"FriendlyName": "Dataprep Editor",
-	"Description": "[WIP] Allows editing of Data preparation pipeline from within the Unreal editor",
-	"Category": "Programming",
+	"Description": "[WIP] A tool to simplify creation and execution of data preparation pipelines from within the Unreal Editor.",
+	"Category": "Dataprep",
 	"CreatedBy": "Epic Games, Inc.",
 	"CreatedByURL": "http://epicgames.com",
 	"DocsURL": "",
@@ -29,6 +29,11 @@
 			"Name": "DataprepLibraries",
 			"Type": "Editor",
 			"LoadingPhase": "Default"
+		},
+		{
+			"Name": "DataprepEditorScriptingUtilities",
+			"Type": "Editor",
+			"LoadingPhase": "Default"
 		}
 	],
 	"Plugins": [
diff --git a/Engine/Plugins/Enterprise/DataprepEditor/Source/DataprepCore/Private/DataprepActionAsset.cpp b/Engine/Plugins/Enterprise/DataprepEditor/Source/DataprepCore/Private/DataprepActionAsset.cpp
--- a/Engine/Plugins/Enterprise/DataprepEditor/Source/DataprepCore/Private/DataprepActionAsset.cpp
+++ b/Engine/Plugins/Enterprise/DataprepEditor/Source/DataprepCore/Private/DataprepActionAsset.cpp
@@ -411,17 +411,17 @@ void UDataprepActionAsset::ExecuteAction(const TSharedPtr<FDataprepActionContext
 					ExecuteOneStep( Step );
 				}
 
-				if ( ContextPtr->ProgressReporterPtr->IsWorkCancelled() )
+				if ( ContextPtr->ProgressReporterPtr && ContextPtr->ProgressReporterPtr->IsWorkCancelled() )
 				{
 					break;
 				}
 
-				if(ContextPtr->ContinueCallback && !ContextPtr->ContinueCallback(this, Step->Operation, Step->Filter))
+				if( ContextPtr->ContinueCallback && !ContextPtr->ContinueCallback( this, Step->Operation, Step->Filter ) )
 				{
 					break;
 				}
 
-				if(Step == SpecificStep)
+				if( Step == SpecificStep )
 				{
 					break;
 				}
diff --git a/Engine/Plugins/Enterprise/DataprepEditor/Source/DataprepCore/Private/DataprepCoreModule.cpp b/Engine/Plugins/Enterprise/DataprepEditor/Source/DataprepCore/Private/DataprepCoreModule.cpp
--- a/Engine/Plugins/Enterprise/DataprepEditor/Source/DataprepCore/Private/DataprepCoreModule.cpp
+++ b/Engine/Plugins/Enterprise/DataprepEditor/Source/DataprepCore/Private/DataprepCoreModule.cpp
@@ -1,54 +1,24 @@
 // Copyright 1998-2019 Epic Games, Inc. All Rights Reserved.
 
-#include "DataprepCoreModule.h"
-
-#include "AssetTypeActions_DataprepAssetInterface.h"
 #include "DataprepCorePrivateUtils.h"
 
-#include "Developer/AssetTools/Public/IAssetTools.h"
 #include "Misc/PackageName.h"
 #include "Modules/ModuleManager.h"
 
-#define LOCTEXT_NAMESPACE "DataprepCoreModule"
-
-EAssetTypeCategories::Type IDataprepCoreModule::DataprepCategoryBit;
-
-class FDataprepCoreModule : public IDataprepCoreModule
+class FDataprepCoreModule : public IModuleInterface
 {
 public:
 	virtual void StartupModule() override
 	{
-		// Register asset type actions for DataPrepRecipe class
-		IAssetTools& AssetTools = FModuleManager::LoadModuleChecked<FAssetToolsModule>("AssetTools").Get();
-
-		// Register Dataprep category to group together asset type actions related to Dataprep
-		DataprepCategoryBit = AssetTools.RegisterAdvancedAssetCategory(FName(TEXT("Dataprep")), LOCTEXT("DataprepCoreModule", "Dataprep"));
-
-		DataprepAssetInterfaceTypeAction = MakeShareable(new FAssetTypeActions_DataprepAssetInterface);
-		AssetTools.RegisterAssetTypeActions(DataprepAssetInterfaceTypeAction.ToSharedRef());
-
 		// Register mount point for Dataprep core library root package folder
 		FPackageName::RegisterMountPoint( DataprepCorePrivateUtils::GetRootPackagePath() + TEXT("/"), DataprepCorePrivateUtils::GetRootTemporaryDir()	);
 	}
 
 	virtual void ShutdownModule() override
 	{
-		// Unregister asset type actions
-		if (FModuleManager::Get().IsModuleLoaded(TEXT("AssetTools")))
-		{
-			IAssetTools& AssetTools = FModuleManager::LoadModuleChecked<FAssetToolsModule>("AssetTools").Get();
-			AssetTools.UnregisterAssetTypeActions(DataprepAssetInterfaceTypeAction.ToSharedRef());
-		}
-		DataprepAssetInterfaceTypeAction.Reset();
-
 		// Unregister mount point for Dataprep core library root package folder
 		FPackageName::UnRegisterMountPoint( DataprepCorePrivateUtils::GetRootPackagePath() + TEXT("/"), DataprepCorePrivateUtils::GetRootTemporaryDir() );
 	}
-
-private:
-	TSharedPtr<FAssetTypeActions_DataprepAssetInterface> DataprepAssetInterfaceTypeAction;
 };
 
-#undef LOCTEXT_NAMESPACE
-
 IMPLEMENT_MODULE( FDataprepCoreModule, DataprepCore )
diff --git a/Engine/Plugins/Enterprise/DataprepEditor/Source/DataprepCore/Private/DataprepCoreUtils.cpp b/Engine/Plugins/Enterprise/DataprepEditor/Source/DataprepCore/Private/DataprepCoreUtils.cpp
--- a/Engine/Plugins/Enterprise/DataprepEditor/Source/DataprepCore/Private/DataprepCoreUtils.cpp
+++ b/Engine/Plugins/Enterprise/DataprepEditor/Source/DataprepCore/Private/DataprepCoreUtils.cpp
@@ -2,27 +2,31 @@
 
 #include "DataprepCoreUtils.h"
 
-#include "DataprepAssetInterface.h"
-
 #ifdef NEW_DATASMITHSCENE_WORKFLOW
 #include "DataprepAssetUserData.h"
 #endif
 
 #include "DataprepAsset.h"
+#include "DataprepAssetInterface.h"
+#include "DataprepContentConsumer.h"
+#include "DataprepContentProducer.h"
 #include "DataprepCoreLogCategory.h"
 #include "DataprepCorePrivateUtils.h"
 #include "IDataprepProgressReporter.h"
 
 #include "Engine/StaticMesh.h"
+#include "Engine/World.h"
 #include "GameFramework/Actor.h"
+#include "HAL/FileManager.h"
 #include "LevelSequence.h"
-#include "Materials/MaterialFunction.h"
-#include "Materials/MaterialInstance.h"
-#include "Materials/Material.h"
 #include "MaterialGraph/MaterialGraph.h"
 #include "MaterialShared.h"
+#include "Materials/Material.h"
+#include "Materials/MaterialFunction.h"
+#include "Materials/MaterialInstance.h"
 #include "Misc/ScopedSlowTask.h"
 #include "RenderingThread.h"
+#include "UObject/StrongObjectPtr.h"
 #include "UObject/UObjectHash.h"
 
 #if WITH_EDITOR
@@ -158,6 +162,161 @@ bool FDataprepCoreUtils::IsAsset(UObject* Object)
 	return true;
 }
 
+bool FDataprepCoreUtils::ExecuteDataprep(UDataprepAssetInterface* DataprepAssetInterface, const TSharedPtr<IDataprepLogger>& Logger, const TSharedPtr<IDataprepProgressReporter>& Reporter)
+{
+	if ( DataprepAssetInterface != nullptr )
+	{
+		// The temporary folders are used for the whole session of the Unreal Editor
+		static FString RelativeTempFolder = FString::FromInt( FPlatformProcess::GetCurrentProcessId() ) / FGuid::NewGuid().ToString();
+		static FString TransientContentFolder = DataprepCorePrivateUtils::GetRootPackagePath() / RelativeTempFolder;
+
+		// Create transient world to host data from producer
+		FName UniqueWorldName = MakeUniqueObjectName( GetTransientPackage(), UWorld::StaticClass(), FName( *( LOCTEXT("TransientWorld", "Preview").ToString() ) ) );
+		TStrongObjectPtr<UWorld> TransientWorld = TStrongObjectPtr<UWorld>( NewObject<UWorld>( GetTransientPackage(), UniqueWorldName ) );
+		TransientWorld->WorldType = EWorldType::EditorPreview;
+
+		FWorldContext& WorldContext = GEngine->CreateNewWorldContext( TransientWorld->WorldType );
+		WorldContext.SetCurrentWorld( TransientWorld.Get() );
+
+		TransientWorld->InitializeNewWorld( UWorld::InitializationValues()
+			.AllowAudioPlayback( false )
+			.CreatePhysicsScene( false )
+			.RequiresHitProxies( false )
+			.CreateNavigation( false )
+			.CreateAISystem( false )
+			.ShouldSimulatePhysics( false )
+			.SetTransactional( false )
+			);
+
+		TArray<TWeakObjectPtr<UObject>> Assets;
+
+		FText DataprepAssetTextName = FText::FromString( DataprepAssetInterface->GetName() );
+		FText TaskDescription = FText::Format( LOCTEXT("ExecutingDataprepAsset", "Executing Dataprep Asset \"{0}\" ..."), DataprepAssetTextName );
+		FDataprepWorkReporter ProgressTask( Reporter, TaskDescription, 3.0f, 1.0f );
+
+		bool bSuccessfulExecute = true;
+
+		// Run the producers
+		{
+			// Create package to pass to the producers
+			UPackage* TransientPackage = NewObject<UPackage>( nullptr, *TransientContentFolder, RF_Transient );
+			TransientPackage->FullyLoad();
+
+			TSharedPtr<FDataprepCoreUtils::FDataprepFeedbackContext> FeedbackContext = MakeShared<FDataprepCoreUtils::FDataprepFeedbackContext>();
+
+			FDataprepProducerContext Context;
+			Context.SetWorld( TransientWorld.Get() )
+				.SetRootPackage( TransientPackage )
+				.SetLogger( Logger )
+				.SetProgressReporter( Reporter );
+
+			FText Message = FText::Format( LOCTEXT("Running_Producers", "Running \"{0}\'s Producers ..."), DataprepAssetTextName );
+			ProgressTask.ReportNextStep( Message );
+			Assets = DataprepAssetInterface->GetProducers()->Produce( Context );
+		}
+
+		// Trigger execution of data preparation operations on world attached to recipe
+		TSet<TWeakObjectPtr<UObject>> CachedAssets;
+		{
+			DataprepActionAsset::FCanExecuteNextStepFunc CanExecuteNextStepFunc = [](UDataprepActionAsset* ActionAsset, UDataprepOperation* Operation, UDataprepFilter* Filter) -> bool
+			{
+				return true;
+			};
+
+
+			TSharedPtr<FDataprepActionContext> ActionsContext = MakeShared<FDataprepActionContext>();
+
+			ActionsContext->SetTransientContentFolder( TransientContentFolder / TEXT("Pipeline") )
+				.SetLogger( Logger )
+				.SetProgressReporter( Reporter )
+				.SetCanExecuteNextStep( CanExecuteNextStepFunc )
+				.SetWorld( TransientWorld.Get() )
+				.SetAssets( Assets );
+
+			FText Message = FText::Format( LOCTEXT("Executing_Recipe", "Executing \"{0}\'s Recipe ..."), DataprepAssetTextName );
+			ProgressTask.ReportNextStep( Message );
+			DataprepAssetInterface->ExecuteRecipe( ActionsContext );
+
+			// Update list of assets with latest ones
+			Assets = ActionsContext->Assets.Array();
+
+			for ( TWeakObjectPtr<UObject>& Asset : Assets )
+			{
+				if ( Asset.IsValid() )
+				{
+					CachedAssets.Add( Asset );
+				}
+			}
+
+		}
+
+		// Run consumer to output result of recipe
+		{
+			FDataprepConsumerContext Context;
+			Context.SetWorld( TransientWorld.Get() )
+				.SetAssets( Assets )
+				.SetTransientContentFolder( TransientContentFolder )
+				.SetLogger( Logger )
+				.SetProgressReporter( Reporter );
+
+			FText Message = FText::Format( LOCTEXT("Running_Consumer", "Running \"{0}\'s Consumer ..."), DataprepAssetTextName );
+			ProgressTask.ReportNextStep( Message );
+
+			bSuccessfulExecute = DataprepAssetInterface->RunConsumer( Context );
+		}
+
+		// Clean all temporary data created by the Dataprep asset
+		{
+			// Delete all actors of the transient world
+			TArray<AActor*> TransientActors;
+			DataprepCorePrivateUtils::GetActorsFromWorld( TransientWorld.Get(), TransientActors );
+			for ( AActor* Actor : TransientActors )
+			{
+				if ( Actor && !Actor->IsPendingKill() )
+				{
+					TransientWorld->EditorDestroyActor( Actor, true );
+
+					// Since deletion can be delayed, rename to avoid future name collision
+					// Call UObject::Rename directly on actor to avoid AActor::Rename which unnecessarily sunregister and re-register components
+					Actor->UObject::Rename( nullptr, GetTransientPackage(), REN_DontCreateRedirectors | REN_ForceNoResetLoaders );
+				}
+			}
+
+			// Delete assets which are still in the transient content folder
+			TArray<UObject*> ObjectsToDelete;
+			for ( TWeakObjectPtr<UObject>& Asset : CachedAssets )
+			{
+				if ( UObject* ObjectToDelete = Asset.Get() )
+				{
+					FString PackagePath = ObjectToDelete->GetOutermost()->GetName();
+					if ( PackagePath.StartsWith( TransientContentFolder ) )
+					{
+						FDataprepCoreUtils::MoveToTransientPackage( ObjectToDelete );
+						ObjectsToDelete.Add( ObjectToDelete );
+					}
+				}
+			}
+
+			// Disable warnings from LogStaticMesh because FDataprepCoreUtils::PurgeObjects is pretty verbose on harmless warnings
+			ELogVerbosity::Type PrevLogStaticMeshVerbosity = LogStaticMesh.GetVerbosity();
+			LogStaticMesh.SetVerbosity( ELogVerbosity::Error );
+
+			FDataprepCoreUtils::PurgeObjects( MoveTemp( ObjectsToDelete ) );
+
+			// Restore LogStaticMesh verbosity
+			LogStaticMesh.SetVerbosity( PrevLogStaticMeshVerbosity );
+
+			// Erase all temporary files created by the Dataprep asset
+			static FString AbsolutePath = FPaths::ConvertRelativePathToFull( DataprepCorePrivateUtils::GetRootTemporaryDir() / RelativeTempFolder );
+			IFileManager::Get().DeleteDirectory( *AbsolutePath, false, true );
+		}
+
+		return bSuccessfulExecute;
+	}
+
+	return false;
+}
+
 FDataprepWorkReporter::FDataprepWorkReporter(const TSharedPtr<IDataprepProgressReporter>& InReporter, const FText& InDescription, float InAmountOfWork, float InIncrementOfWork, bool bInterruptible )
 	: Reporter( InReporter )
 	, DefaultIncrementOfWork( InIncrementOfWork )
@@ -170,7 +329,7 @@ FDataprepWorkReporter::FDataprepWorkReporter(const TSharedPtr<IDataprepProgressR
 
 FDataprepWorkReporter::~FDataprepWorkReporter()
 {
-	if( Reporter.IsValid())
+	if( Reporter.IsValid() )
 	{
 		Reporter->EndWork();
 	}
@@ -178,7 +337,7 @@ FDataprepWorkReporter::~FDataprepWorkReporter()
 
 void FDataprepWorkReporter::ReportNextStep(const FText & InMessage, float InIncrementOfWork )
 {
-	if( Reporter.IsValid())
+	if( Reporter.IsValid() )
 	{
 		Reporter->ReportProgress( InIncrementOfWork, InMessage );
 	}
diff --git a/Engine/Plugins/Enterprise/DataprepEditor/Source/DataprepCore/Private/SelectionSystem/DataprepBoolFilter.cpp b/Engine/Plugins/Enterprise/DataprepEditor/Source/DataprepCore/Private/SelectionSystem/DataprepBoolFilter.cpp
--- a/Engine/Plugins/Enterprise/DataprepEditor/Source/DataprepCore/Private/SelectionSystem/DataprepBoolFilter.cpp
+++ b/Engine/Plugins/Enterprise/DataprepEditor/Source/DataprepCore/Private/SelectionSystem/DataprepBoolFilter.cpp
@@ -19,7 +19,6 @@ TArray<UObject*> UDataprepBoolFilter::FilterObjects(const TArray<UObject *>& Obj
 		return DataprepSelectionSystemUtils::FilterObjects< UDataprepBoolFilter, UDataprepBoolFetcher, bool >( *this, * BoolFetcher, Objects );
 	}
 
-	ensure( false );
 	UE_LOG( LogDataprepCore, Error, TEXT("UDataprepBoolFilter::FilterObjects: There was no Fetcher") );
 	return {};
 }
diff --git a/Engine/Plugins/Enterprise/DataprepEditor/Source/DataprepCore/Public/DataprepAssetInstance.h b/Engine/Plugins/Enterprise/DataprepEditor/Source/DataprepCore/Public/DataprepAssetInstance.h
--- a/Engine/Plugins/Enterprise/DataprepEditor/Source/DataprepCore/Public/DataprepAssetInstance.h
+++ b/Engine/Plugins/Enterprise/DataprepEditor/Source/DataprepCore/Public/DataprepAssetInstance.h
@@ -18,7 +18,7 @@ class UDataprepParameterizationInstance;
  * It has its own inputs and output. It has overrides of the externalized parameters
  * of the pipeline.
  */
-UCLASS(Experimental)
+UCLASS(Experimental, BlueprintType)
 class DATAPREPCORE_API UDataprepAssetInstance : public UDataprepAssetInterface
 {
 	GENERATED_BODY()
diff --git a/Engine/Plugins/Enterprise/DataprepEditor/Source/DataprepCore/Public/DataprepAssetInterface.h b/Engine/Plugins/Enterprise/DataprepEditor/Source/DataprepCore/Public/DataprepAssetInterface.h
--- a/Engine/Plugins/Enterprise/DataprepEditor/Source/DataprepCore/Public/DataprepAssetInterface.h
+++ b/Engine/Plugins/Enterprise/DataprepEditor/Source/DataprepCore/Public/DataprepAssetInterface.h
@@ -49,7 +49,7 @@ protected:
 	FOnDataprepRecipeChanged OnChanged;
 };
 
-UCLASS(Experimental, Abstract)
+UCLASS(Experimental, Abstract, BlueprintType)
 class DATAPREPCORE_API UDataprepAssetInterface : public UObject
 {
 	GENERATED_BODY()
diff --git a/Engine/Plugins/Enterprise/DataprepEditor/Source/DataprepCore/Public/DataprepCoreUtils.h b/Engine/Plugins/Enterprise/DataprepEditor/Source/DataprepCore/Public/DataprepCoreUtils.h
--- a/Engine/Plugins/Enterprise/DataprepEditor/Source/DataprepCore/Public/DataprepCoreUtils.h
+++ b/Engine/Plugins/Enterprise/DataprepEditor/Source/DataprepCore/Public/DataprepCoreUtils.h
@@ -15,6 +15,7 @@
 #include "UObject/Package.h"
 
 class UDataprepAsset;
+class UDataprepAssetInterface;
 
 struct FScopedSlowTask;
 
@@ -69,13 +70,21 @@ public:
 		}
 	}
 
+	/**
+	 * Execute a dataprep asset (import, execute and commit)
+	 * @param DataprepAssetInterface The dataprep asset to execute
+	 * @param Logger A optional log manager for the dataprep producers, operations and consumer
+	 * @param Reporter A optional progress reporter for the dataprep producers, operations and consumer
+	 */
+	static bool ExecuteDataprep(UDataprepAssetInterface* DataprepAssetInterface, const TSharedPtr<IDataprepLogger>& Logger, const TSharedPtr<IDataprepProgressReporter>& Reporter);
+
 	/**
 	 * Helper function to build assets for use in the Dataprep pipeline
 
 	 * @param	Assets					Array of weak pointer on the assets to be build
 	 * @param	ProgressReporterPtr		Pointer to a IDataprepProgressReporter interface. This pointer can be invalid.
 	 */
-	static void BuildAssets(const TArray< TWeakObjectPtr<UObject> >& Assets, const TSharedPtr< IDataprepProgressReporter >& ProgressReporterPtr );
+	static void BuildAssets(const TArray< TWeakObjectPtr<UObject> >& Assets, const TSharedPtr<IDataprepProgressReporter>& ProgressReporterPtr );
 
 	class DATAPREPCORE_API FDataprepLogger : public IDataprepLogger
 	{
diff --git a/Engine/Plugins/Enterprise/DataprepEditor/Source/DataprepEditor/Private/AssetTypeActions_Dataprep.cpp b/Engine/Plugins/Enterprise/DataprepEditor/Source/DataprepEditor/Private/AssetTypeActions_Dataprep.cpp
--- a/Engine/Plugins/Enterprise/DataprepEditor/Source/DataprepEditor/Private/AssetTypeActions_Dataprep.cpp
+++ b/Engine/Plugins/Enterprise/DataprepEditor/Source/DataprepEditor/Private/AssetTypeActions_Dataprep.cpp
@@ -2,15 +2,15 @@
 
 #include "AssetTypeActions_Dataprep.h"
 
+#include "DataprepEditorModule.h"
 #include "DataprepRecipe.h"
-#include "DataprepCoreModule.h"
 
 /* FAssetTypeActions_Base interface
 *****************************************************************************/
 
 uint32 FAssetTypeActions_Dataprep::GetCategories()
 {
-	return IDataprepCoreModule::DataprepCategoryBit;
+	return IDataprepEditorModule::DataprepCategoryBit;
 }
 
 
diff --git a/Engine/Plugins/Enterprise/DataprepEditor/Source/DataprepEditor/Private/AssetTypeActions_DataprepAssetInterface.cpp b/Engine/Plugins/Enterprise/DataprepEditor/Source/DataprepEditor/Private/AssetTypeActions_DataprepAssetInterface.cpp
new file mode 100644
--- /dev/null
+++ b/Engine/Plugins/Enterprise/DataprepEditor/Source/DataprepEditor/Private/AssetTypeActions_DataprepAssetInterface.cpp
@@ -0,0 +1,155 @@
+// Copyright 1998-2019 Epic Games, Inc. All Rights Reserved.
+
+#include "AssetTypeActions_DataprepAssetInterface.h"
+
+#include "DataprepAssetInstance.h"
+#include "DataprepAssetProducers.h"
+#include "DataprepCoreUtils.h"
+#include "DataprepEditorModule.h"
+#include "DataprepFactories.h"
+
+#include "Developer/AssetTools/Public/IAssetTools.h"
+#include "Developer/AssetTools/Public/AssetToolsModule.h"
+#include "IContentBrowserSingleton.h"
+#include "ContentBrowserModule.h"
+
+// UI
+#include "Framework/MultiBox/MultiBoxBuilder.h"
+
+#define LOCTEXT_NAMESPACE "AssetTypeActions_DataprepAssetInterface"
+
+uint32 FAssetTypeActions_DataprepAssetInterface::GetCategories()
+{
+	return IDataprepEditorModule::DataprepCategoryBit;
+}
+
+FText FAssetTypeActions_DataprepAssetInterface::GetName() const
+{
+	return LOCTEXT( "Name", "Dataprep Interface" );
+}
+
+UClass* FAssetTypeActions_DataprepAssetInterface::GetSupportedClass() const
+{
+	return UDataprepAssetInterface::StaticClass();
+}
+
+void FAssetTypeActions_DataprepAssetInterface::CreateInstance(TArray<TWeakObjectPtr<UDataprepAssetInterface>> DataprepAssetInterfaces)
+{
+	// Code is inspired from FAssetTypeActions_MaterialInterface::ExecuteNewMIC
+	const FString DefaultSuffix = TEXT("_Inst");
+
+	if(DataprepAssetInterfaces.Num() == 1)
+	{
+		if(UDataprepAssetInterface* Object = DataprepAssetInterfaces[0].Get())
+		{
+			// Determine an appropriate and unique name 
+			FString Name;
+			FString PackageName;
+			CreateUniqueAssetName(Object->GetOutermost()->GetName(), DefaultSuffix, PackageName, Name);
+
+			// Create the factory used to generate the asset
+			UDataprepAssetInstanceFactory* Factory = NewObject<UDataprepAssetInstanceFactory>();
+			Factory->InitialParent = Object;
+
+			// Create asset in 
+			FContentBrowserModule& ContentBrowserModule = FModuleManager::LoadModuleChecked<FContentBrowserModule>("ContentBrowser");
+			ContentBrowserModule.Get().CreateNewAsset(Name, FPackageName::GetLongPackagePath(PackageName), UDataprepAssetInstance::StaticClass(), Factory);
+		}
+	}
+	else
+	{
+		TArray<UObject*> AssetsToSync;
+		for(auto AssetIt = DataprepAssetInterfaces.CreateConstIterator(); AssetIt; ++AssetIt)
+		{
+			if(UDataprepAssetInterface* DataprepAssetInterface = (*AssetIt).Get())
+			{
+				// Determine an appropriate name
+				FString Name;
+				FString PackageName;
+				CreateUniqueAssetName(DataprepAssetInterface->GetOutermost()->GetName(), DefaultSuffix, PackageName, Name);
+
+				// Create the factory used to generate the asset
+				UDataprepAssetInstanceFactory* Factory = NewObject<UDataprepAssetInstanceFactory>();
+				Factory->InitialParent = DataprepAssetInterface;
+
+				IAssetTools& AssetTools = FModuleManager::LoadModuleChecked<FAssetToolsModule>("AssetTools").Get();
+				UObject* NewAsset = AssetTools.CreateAsset(Name, FPackageName::GetLongPackagePath(PackageName), UDataprepAssetInstance::StaticClass(), Factory);
+
+				if(NewAsset)
+				{
+					AssetsToSync.Add(NewAsset);
+				}
+			}
+		}
+
+		if(AssetsToSync.Num() > 0)
+		{
+			FContentBrowserModule& ContentBrowserModule = FModuleManager::Get().LoadModuleChecked<FContentBrowserModule>("ContentBrowser");
+			ContentBrowserModule.Get().SyncBrowserToAssets(AssetsToSync, /*bAllowLockedBrowsers=*/true);
+		}
+	}
+}
+
+void FAssetTypeActions_DataprepAssetInterface::ExecuteDataprepAssets(TArray<TWeakObjectPtr<UDataprepAssetInterface>> DataprepAssetInterfaces)
+{
+	for(TWeakObjectPtr<UDataprepAssetInterface>& DataprepAssetInterfacePtr : DataprepAssetInterfaces)
+	{
+		if( UDataprepAssetInterface* DataprepAssetInterface = DataprepAssetInterfacePtr.Get() )
+		{
+			// Nothing to do if the Dataprep asset does not have any inputs
+			if(DataprepAssetInterface->GetProducers()->GetProducersCount() > 0)
+			{
+				FDataprepCoreUtils::ExecuteDataprep( DataprepAssetInterface
+					, MakeShared<FDataprepCoreUtils::FDataprepLogger>()
+					, MakeShared<FDataprepCoreUtils::FDataprepProgressUIReporter>() );
+			}
+		}
+	}
+}
+
+void FAssetTypeActions_DataprepAssetInterface::GetActions(const TArray<UObject*>& InObjects, FMenuBuilder& MenuBuilder)
+{
+	auto DataprepAssetInterfaces = GetTypedWeakObjectPtrs<UDataprepAssetInterface>(InObjects);
+
+	if(DataprepAssetInterfaces.Num() == 0)
+	{
+		return;
+	}
+
+	// #ueent_remark: An instance of an instance is not supported for 4.24.
+	// Do not expose 'Create Instance' menu entry if at least one Dataprep asset is an instance
+	bool bContainsAnInstance  = false;
+	for (UObject* Object : InObjects)
+	{
+		if (Object && Object->GetClass() == UDataprepAssetInstance::StaticClass())
+		{
+			bContainsAnInstance = true;
+			break;
+		}
+	}
+
+	if (!bContainsAnInstance)
+	{
+		MenuBuilder.AddMenuEntry(
+			LOCTEXT("CreateInstance", "Create Instance"),
+			LOCTEXT("CreateInstanceTooltip", "Creates a parameterized Dataprep asset using this Dataprep asset as a base."),
+			FSlateIcon(),
+			FUIAction(
+				FExecuteAction::CreateSP(this, &FAssetTypeActions_DataprepAssetInterface::CreateInstance, DataprepAssetInterfaces),
+				FCanExecuteAction()
+			)
+		);
+	}
+
+	MenuBuilder.AddMenuEntry(
+		LOCTEXT("RunAsset", "Execute"),
+		LOCTEXT("RunAssetTooltip", "Runs the Dataprep asset's producers, execute its recipe, finally runs the consumer"),
+		FSlateIcon(),
+		FUIAction(
+			FExecuteAction::CreateSP(this, &FAssetTypeActions_DataprepAssetInterface::ExecuteDataprepAssets, DataprepAssetInterfaces),
+			FCanExecuteAction()
+		)
+	);
+}
+
+#undef LOCTEXT_NAMESPACE
\ No newline at end of file
diff --git a/Engine/Plugins/Enterprise/DataprepEditor/Source/DataprepEditor/Private/AssetTypeActions_DataprepAssetInterface.h b/Engine/Plugins/Enterprise/DataprepEditor/Source/DataprepEditor/Private/AssetTypeActions_DataprepAssetInterface.h
new file mode 100644
--- /dev/null
+++ b/Engine/Plugins/Enterprise/DataprepEditor/Source/DataprepEditor/Private/AssetTypeActions_DataprepAssetInterface.h
@@ -0,0 +1,27 @@
+// Copyright 1998-2019 Epic Games, Inc. All Rights Reserved.
+
+#pragma once
+
+#include "CoreMinimal.h"
+#include "AssetTypeActions_Base.h"
+
+class UDataprepAssetInterface;
+
+/** Asset type actions for UDataprepAssetInterface class */
+class DATAPREPEDITOR_API FAssetTypeActions_DataprepAssetInterface : public FAssetTypeActions_Base
+{
+public:
+	// Begin IAssetTypeActions interface
+	virtual FText GetName() const override;
+	virtual uint32 GetCategories() override;
+	virtual FColor GetTypeColor() const override { return FColor(0, 0, 0); }
+	virtual UClass* GetSupportedClass() const override;
+	virtual bool HasActions ( const TArray<UObject*>& InObjects ) const override { return true; }
+	virtual void GetActions( const TArray<UObject*>& InObjects, FMenuBuilder& MenuBuilder ) override;
+	virtual bool IsImportedAsset() const override { return false; }
+
+private:
+	void CreateInstance(TArray<TWeakObjectPtr<UDataprepAssetInterface>> Objects);
+	void ExecuteDataprepAssets(TArray<TWeakObjectPtr<UDataprepAssetInterface>> Objects);
+};
+
diff --git a/Engine/Plugins/Enterprise/DataprepEditor/Source/DataprepEditor/Private/DataprepEditorModule.cpp b/Engine/Plugins/Enterprise/DataprepEditor/Source/DataprepEditor/Private/DataprepEditorModule.cpp
--- a/Engine/Plugins/Enterprise/DataprepEditor/Source/DataprepEditor/Private/DataprepEditorModule.cpp
+++ b/Engine/Plugins/Enterprise/DataprepEditor/Source/DataprepEditor/Private/DataprepEditorModule.cpp
@@ -4,9 +4,9 @@
 
 #include "AssetTypeActions_Dataprep.h"
 #include "AssetTypeActions_DataprepAsset.h"
+#include "AssetTypeActions_DataprepAssetInterface.h"
 #include "DataprepAssetProducers.h"
 #include "DataprepEditor.h"
-
 #include "DataprepEditorStyle.h"
 #include "Widgets/DataprepWidgets.h"
 #include "Widgets/SDataprepEditorViewport.h"
@@ -23,6 +23,8 @@
 
 const FName DataprepEditorAppIdentifier = FName(TEXT("DataprepEditorApp"));
 
+EAssetTypeCategories::Type IDataprepEditorModule::DataprepCategoryBit;
+
 #define LOCTEXT_NAMESPACE "DataprepEditorModule"
 
 class FDataprepEditorModule : public IDataprepEditorModule
@@ -39,6 +41,13 @@ public:
 		// Register asset type actions for DataPrepRecipe class
 		IAssetTools& AssetTools = FModuleManager::LoadModuleChecked<FAssetToolsModule>("AssetTools").Get();
 
+		// Register Dataprep category to group together asset type actions related to Dataprep
+		DataprepCategoryBit = AssetTools.RegisterAdvancedAssetCategory(FName(TEXT("Dataprep")), LOCTEXT("DataprepAssetCategory", "Dataprep"));
+
+		TSharedPtr<FAssetTypeActions_DataprepAssetInterface> DataprepAssetInterfaceTypeAction = MakeShared<FAssetTypeActions_DataprepAssetInterface>();
+		AssetTools.RegisterAssetTypeActions(DataprepAssetInterfaceTypeAction.ToSharedRef());
+		AssetTypeActionsArray.Add(DataprepAssetInterfaceTypeAction);
+
 		TSharedPtr<FAssetTypeActions_Dataprep> DataprepRecipeTypeAction = MakeShareable(new FAssetTypeActions_Dataprep);
 		AssetTools.RegisterAssetTypeActions(DataprepRecipeTypeAction.ToSharedRef());
 		AssetTypeActionsArray.Add(DataprepRecipeTypeAction);
diff --git a/Engine/Plugins/Enterprise/DataprepEditor/Source/DataprepEditor/Public/DataprepEditorModule.h b/Engine/Plugins/Enterprise/DataprepEditor/Source/DataprepEditor/Public/DataprepEditorModule.h
--- a/Engine/Plugins/Enterprise/DataprepEditor/Source/DataprepEditor/Public/DataprepEditorModule.h
+++ b/Engine/Plugins/Enterprise/DataprepEditor/Source/DataprepEditor/Public/DataprepEditorModule.h
@@ -4,6 +4,7 @@
 
 #include "CoreMinimal.h"
 
+#include "AssetTypeCategories.h"
 #include "Developer/AssetTools/Public/AssetTypeCategories.h"
 #include "Modules/ModuleInterface.h"
 #include "Modules/ModuleManager.h"		// For inline LoadModuleChecked()
@@ -44,4 +45,6 @@ public:
 
 	virtual TSharedRef<SWidget> CreateDataprepProducersWidget(UDataprepAssetProducers* AssetProducers) = 0;
 	virtual TSharedRef<SWidget> CreateDataprepDetailsView(UObject* ObjectToDetail) = 0;
+
+	static EAssetTypeCategories::Type DataprepCategoryBit;
 };
diff --git a/Engine/Plugins/Enterprise/DataprepEditor/Source/DataprepEditorScriptingUtilities/Private/DataprepEditorScriptingUtilitiesModule.cpp b/Engine/Plugins/Enterprise/DataprepEditor/Source/DataprepEditorScriptingUtilities/Private/DataprepEditorScriptingUtilitiesModule.cpp
new file mode 100644
--- /dev/null
+++ b/Engine/Plugins/Enterprise/DataprepEditor/Source/DataprepEditorScriptingUtilities/Private/DataprepEditorScriptingUtilitiesModule.cpp
@@ -0,0 +1,6 @@
+// Copyright 1998-2019 Epic Games, Inc. All Rights Reserved.
+
+#include "Modules/ModuleManager.h"
+
+IMPLEMENT_MODULE( FDefaultModuleImpl, DataprepEditorScriptingUtilities );
+
diff --git a/Engine/Plugins/Enterprise/DatasmithC4DImporter/Source/DatasmithC4DTranslator/DatasmithC4DTranslator.Build.cs b/Engine/Plugins/Enterprise/DatasmithC4DImporter/Source/DatasmithC4DTranslator/DatasmithC4DTranslator.Build.cs
--- a/Engine/Plugins/Enterprise/DatasmithC4DImporter/Source/DatasmithC4DTranslator/DatasmithC4DTranslator.Build.cs
+++ b/Engine/Plugins/Enterprise/DatasmithC4DImporter/Source/DatasmithC4DTranslator/DatasmithC4DTranslator.Build.cs
@@ -35,7 +35,7 @@ namespace UnrealBuildTool.Rules
 				new string[]
 				{
 					"DatasmithContent",
-					"DatasmithImporter"
+					"DatasmithTranslator"
 				}
 			);
 
diff --git a/Engine/Plugins/Enterprise/DatasmithC4DImporter/Source/DatasmithC4DTranslator/Private/DatasmithC4DImporter.cpp b/Engine/Plugins/Enterprise/DatasmithC4DImporter/Source/DatasmithC4DTranslator/Private/DatasmithC4DImporter.cpp
--- a/Engine/Plugins/Enterprise/DatasmithC4DImporter/Source/DatasmithC4DTranslator/Private/DatasmithC4DImporter.cpp
+++ b/Engine/Plugins/Enterprise/DatasmithC4DImporter/Source/DatasmithC4DTranslator/Private/DatasmithC4DImporter.cpp
@@ -13,11 +13,11 @@
 #include "DatasmithImportOptions.h"
 #include "DatasmithMesh.h"
 #include "DatasmithMeshExporter.h"
-#include "DatasmithMeshHelper.h"
 #include "DatasmithSceneExporter.h"
 #include "DatasmithSceneFactory.h"
 #include "DatasmithUtils.h"
 #include "IDatasmithSceneElements.h"
+#include "Utility/DatasmithMeshHelper.h"
 
 #include "AssetRegistryModule.h"
 #include "Curves/RichCurve.h"
@@ -2403,9 +2403,18 @@ TSharedPtr<IDatasmithActorElement> FDatasmithC4DImporter::ImportObjectAndChildre
 
 		if (InstanceObjects)
 		{
-			if (InstanceObjects->Num() > 0 && (*InstanceObjects)[0]->GetType() == ObjectType)
+			if (InstanceObjects->Num() > 0)
 			{
-				DataObject = (*InstanceObjects)[0];
+				melange::BaseObject* RealDataObject = (*InstanceObjects)[0];
+				InstanceObjects->RemoveAt(0);
+				if (RealDataObject->GetType() == ObjectType)
+				{
+					DataObject = RealDataObject;
+				}
+				else
+				{
+					bSuccess = false;
+				}
 			}
 			else
 			{
diff --git a/Engine/Plugins/Enterprise/DatasmithC4DImporter/Source/DatasmithC4DTranslator/Private/DatasmithC4DTranslator.h b/Engine/Plugins/Enterprise/DatasmithC4DImporter/Source/DatasmithC4DTranslator/Private/DatasmithC4DTranslator.h
--- a/Engine/Plugins/Enterprise/DatasmithC4DImporter/Source/DatasmithC4DTranslator/Private/DatasmithC4DTranslator.h
+++ b/Engine/Plugins/Enterprise/DatasmithC4DImporter/Source/DatasmithC4DTranslator/Private/DatasmithC4DTranslator.h
@@ -2,9 +2,10 @@
 
 #pragma once
 
-#include "Translators/DatasmithTranslator.h"
 #include "DatasmithC4DImporter.h"
 #include "DatasmithC4DImportOptions.h"
+#include "DatasmithTranslator.h"
+
 #include "CoreMinimal.h"
 
 class FDatasmithC4DTranslator : public IDatasmithTranslator
@@ -12,7 +13,6 @@ class FDatasmithC4DTranslator : public IDatasmithTranslator
 public:
 	virtual FName GetFName() const override { return "DatasmithC4DTranslator"; };
 
-	// IDatasmithTranslator interface
 #ifndef _MELANGE_SDK_
 	virtual void Initialize(FDatasmithTranslatorCapabilities& OutCapabilities) override { OutCapabilities.bIsEnabled = false; }
 #else
@@ -26,7 +26,6 @@ public:
 
 	virtual void GetSceneImportOptions(TArray<TStrongObjectPtr<UObject>>& Options) override;
 	virtual void SetSceneImportOptions(TArray<TStrongObjectPtr<UObject>>& Options) override;
-	//~ End IDatasmithTranslator interface
 
 private:
 	TStrongObjectPtr<UDatasmithC4DImportOptions> ImportOptions;
diff --git a/Engine/Plugins/Enterprise/DatasmithC4DImporter/Source/DatasmithC4DTranslator/Private/DatasmithC4DTranslatorModule.cpp b/Engine/Plugins/Enterprise/DatasmithC4DImporter/Source/DatasmithC4DTranslator/Private/DatasmithC4DTranslatorModule.cpp
--- a/Engine/Plugins/Enterprise/DatasmithC4DImporter/Source/DatasmithC4DTranslator/Private/DatasmithC4DTranslatorModule.cpp
+++ b/Engine/Plugins/Enterprise/DatasmithC4DImporter/Source/DatasmithC4DTranslator/Private/DatasmithC4DTranslatorModule.cpp
@@ -36,7 +36,7 @@ public:
 	virtual void StartupModule() override
 	{
 		// Make sure the DatasmithImporter module exists and has been initialized before adding FDatasmithC4DTranslator's material selector
-		FModuleManager::Get().LoadModule(TEXT("DatasmithImporter"));
+		FModuleManager::Get().LoadModule(TEXT("DatasmithTranslator"));
 
 		FDatasmithMasterMaterialManager::Get().RegisterSelector(TEXT("C4DTranslator"), MakeShared< FDatasmithC4DImporterMaterialSelector >());
 
diff --git a/Engine/Plugins/Enterprise/DatasmithCADImporter/Source/CADInterfaces/Private/CADSceneGraph.cpp b/Engine/Plugins/Enterprise/DatasmithCADImporter/Source/CADInterfaces/Private/CADSceneGraph.cpp
--- a/Engine/Plugins/Enterprise/DatasmithCADImporter/Source/CADInterfaces/Private/CADSceneGraph.cpp
+++ b/Engine/Plugins/Enterprise/DatasmithCADImporter/Source/CADInterfaces/Private/CADSceneGraph.cpp
@@ -4,8 +4,7 @@
 
 #include "Misc/FileHelper.h"
 #include "Serialization/Archive.h"
-#include "Serialization/MemoryWriter.h"
-#include "Serialization/MemoryReader.h"
+#include "HAL/FileManager.h"
 
 namespace CADLibrary
 {
@@ -87,22 +86,18 @@ FArchive& operator<<(FArchive& Ar, FArchiveMockUp& MockUp)
 	return Ar;
 }
 
-void SerializeMockUp(FArchiveMockUp& MockUp, const TCHAR* Filename)
+void FArchiveMockUp::SerializeMockUp(const TCHAR* Filename)
 {
-	TArray<uint8> OutBuffer;
-	FMemoryWriter ArWriter(OutBuffer);
-	ArWriter << MockUp;
-
-	FFileHelper::SaveArrayToFile(OutBuffer, Filename);
+	TUniquePtr<FArchive> Archive(IFileManager::Get().CreateFileWriter(Filename));
+	*Archive << *this;
+	Archive->Close();
 }
 
-void DeserializeMockUpFile(const TCHAR* Filename, FArchiveMockUp& MockUp)
+void FArchiveMockUp::DeserializeMockUpFile(const TCHAR* Filename)
 {
-	TArray<uint8> Buffer;
-	FFileHelper::LoadFileToArray(Buffer, Filename);
-
-	FMemoryReader ArReader(Buffer);
-	ArReader << MockUp;
+	TUniquePtr<FArchive> Archive(IFileManager::Get().CreateFileReader(Filename));
+	*Archive << *this;
+	Archive->Close();
 }
 
 }
diff --git a/Engine/Plugins/Enterprise/DatasmithCADImporter/Source/CADInterfaces/Private/CoreTechFileParser.cpp b/Engine/Plugins/Enterprise/DatasmithCADImporter/Source/CADInterfaces/Private/CoreTechFileParser.cpp
--- a/Engine/Plugins/Enterprise/DatasmithCADImporter/Source/CADInterfaces/Private/CoreTechFileParser.cpp
+++ b/Engine/Plugins/Enterprise/DatasmithCADImporter/Source/CADInterfaces/Private/CoreTechFileParser.cpp
@@ -152,8 +152,10 @@ namespace {
 	}
 }
 
-uint32 GetFileHash(const FString& FileName, const FFileStatData& FileStatData, const FString& Config, const FImportParameters& ImportParam)
+uint32 GetFileHash(const FString& FullPath, const FString& FileName, const FString& Config, const FImportParameters& ImportParam)
 {
+	FFileStatData FileStatData = IFileManager::Get().GetStatData(*FullPath);
+
 	int64 FileSize = FileStatData.FileSize;
 	FDateTime ModificationTime = FileStatData.ModificationTime;
 
@@ -426,7 +428,7 @@ uint32 GetStaticMeshUuid(const TCHAR* OutSgFile, const int32 BodyId)
 
 void FCoreTechFileParser::ExportSceneGraphFile()
 {
-	SerializeMockUp(MockUpDescription, *FPaths::Combine(CachePath, TEXT("scene"), MockUpDescription.SceneGraphArchive + TEXT(".sg")));
+	MockUpDescription.SerializeMockUp(*FPaths::Combine(CachePath, TEXT("scene"), MockUpDescription.SceneGraphArchive + TEXT(".sg")));
 }
 
 void FCoreTechFileParser::ExportMeshArchiveFile()
@@ -436,7 +438,7 @@ void FCoreTechFileParser::ExportMeshArchiveFile()
 
 void FCoreTechFileParser::LoadSceneGraphArchive(const FString& SGFile)
 {
-	DeserializeMockUpFile(*SGFile, MockUpDescription);
+	MockUpDescription.DeserializeMockUpFile(*SGFile);
 }
 
 uint32 FCoreTechFileParser::GetMaterialNum()
@@ -490,21 +492,21 @@ void FCoreTechFileParser::ReadMaterials()
 	}
 }
 
-FCoreTechFileParser::FCoreTechFileParser(const FString& InCADFullPath, const FString& InCachePath, const FImportParameters& ImportParams, const TCHAR* KernelIOPath)
+FCoreTechFileParser::FCoreTechFileParser(const FImportParameters& ImportParams, const FString& EnginePluginsPath, const FString& InCachePath)
 	: CachePath(InCachePath)
-	, FullPath(InCADFullPath)
 	, bNeedSaveCTFile(false)
 	, ImportParameters(ImportParams)
 {
-	CTKIO_InitializeKernel(ImportParameters.MetricUnit, KernelIOPath);
+	CTKIO_InitializeKernel(ImportParameters.MetricUnit, *EnginePluginsPath);
 }
 
-FCoreTechFileParser::EProcessResult FCoreTechFileParser::ProcessFile()
+FCoreTechFileParser::EProcessResult FCoreTechFileParser::ProcessFile(const FString& InFullPath)
 {
-	FileConfiguration.Empty();
-
+	FString FullPath = InFullPath;
 	CADFile = FPaths::GetCleanFilename(*FullPath);
 
+	FileConfiguration.Empty();
+
 	// Check if configuration is passed with file name
 	FString NewCADFile;
 	const FString CutOffMarker = TEXT("|");
@@ -518,8 +520,7 @@ FCoreTechFileParser::EProcessResult FCoreTechFileParser::ProcessFile()
 		return EProcessResult::FileNotFound;
 	}
 
-	FFileStatData FileStatData = IFileManager::Get().GetStatData(*FullPath);
-	uint32 FileHash = GetFileHash(CADFile, FileStatData, FileConfiguration, ImportParameters);
+	uint32 FileHash = GetFileHash(FullPath, CADFile, FileConfiguration, ImportParameters);
 
 	MockUpDescription.SceneGraphArchive = FString::Printf(TEXT("UEx%08x"), FileHash);
 
@@ -556,17 +557,17 @@ FCoreTechFileParser::EProcessResult FCoreTechFileParser::ProcessFile()
 	}
 #endif
 	// Process the file
-	return ReadFileWithKernelIO();
+	return ReadFileWithKernelIO(FullPath);
 }
 
-FCoreTechFileParser::EProcessResult FCoreTechFileParser::ReadFileWithKernelIO()
+FCoreTechFileParser::EProcessResult FCoreTechFileParser::ReadFileWithKernelIO(const FString& FullPath)
 {
 	CT_IO_ERROR Result = IO_OK;
 	CT_OBJECT_ID MainId = 0;
 
 	Result = CT_KERNEL_IO::UnloadModel();
 
-	CT_FLAGS CTImportOption = SetCoreTechImportOption(FPaths::GetExtension(CADFile));
+	CT_FLAGS CTImportOption = SetCoreTechImportOption(FPaths::GetExtension(FullPath));
 
 	FString LoadOption;
 	CT_UINT32 NumberOfIds = 1;
@@ -875,7 +876,7 @@ uint32 GetBodiesFaceSetNum(TArray<CT_OBJECT_ID>& BodySet)
 	return size;
 }
 
-void FCoreTechFileParser::GetBodyTessellation(CT_OBJECT_ID BodyId, FBodyMesh& OutBodyMesh, const FImportParameters& ImportParams, uint32 DefaultMaterialHash)
+void FCoreTechFileParser::GetBodyTessellation(CT_OBJECT_ID BodyId, FBodyMesh& OutBodyMesh, uint32 DefaultMaterialHash)
 {
 	CT_LIST_IO FaceList;
 	CT_BODY_IO::AskFaces(BodyId, FaceList);
@@ -901,7 +902,7 @@ void FCoreTechFileParser::GetBodyTessellation(CT_OBJECT_ID BodyId, FBodyMesh& Ou
 	CT_OBJECT_ID FaceID;
 	while ((FaceID = FaceList.IteratorIter()) != 0)
 	{
-		uint32 TriangleNum = GetFaceTessellation(FaceID, OutBodyMesh.Faces, ImportParams);
+		uint32 TriangleNum = GetFaceTessellation(FaceID, OutBodyMesh.Faces, ImportParameters);
 
 		if (TriangleNum == 0)
 		{
@@ -949,7 +950,7 @@ bool FCoreTechFileParser::ReadBody(CT_OBJECT_ID BodyId, uint32 DefaultMaterialHa
 	MockUpDescription.BodySet[Index].MeshActorName = GetStaticMeshUuid(*MockUpDescription.SceneGraphArchive, BodyId);
 	BodyMeshes[BodyMeshIndex].MeshActorName = MockUpDescription.BodySet[Index].MeshActorName;
 
-	GetBodyTessellation(BodyId, BodyMeshes[BodyMeshIndex], ImportParameters, DefaultMaterialHash);
+	GetBodyTessellation(BodyId, BodyMeshes[BodyMeshIndex], DefaultMaterialHash);
 
 	MockUpDescription.BodySet[Index].ColorFaceSet = BodyMeshes[BodyMeshIndex].ColorSet;
 	MockUpDescription.BodySet[Index].MaterialFaceSet = BodyMeshes[BodyMeshIndex].MaterialSet;
@@ -970,7 +971,10 @@ void FCoreTechFileParser::GetAttributeValue(CT_ATTRIB_TYPE AttributType, int Ith
 
 	Value = "";
 
-	if (CT_ATTRIB_DEFINITION_IO::AskFieldDefinition(AttributType, IthField, FieldType, FieldName) != IO_OK) return;
+	if (CT_ATTRIB_DEFINITION_IO::AskFieldDefinition(AttributType, IthField, FieldType, FieldName) != IO_OK) 
+	{
+		return;
+	}
 
 	switch (FieldType) {
 		case CT_ATTRIB_FIELD_UNKNOWN:
@@ -980,21 +984,30 @@ void FCoreTechFileParser::GetAttributeValue(CT_ATTRIB_TYPE AttributType, int Ith
 		case CT_ATTRIB_FIELD_INTEGER:
 		{
 			int IValue;
-			if (CT_CURRENT_ATTRIB_IO::AskIntField(IthField, IValue) != IO_OK) break;
+			if (CT_CURRENT_ATTRIB_IO::AskIntField(IthField, IValue) != IO_OK) 
+			{
+				break;
+			}
 			Value = FString::FromInt(IValue);
 			break;
 		}
 		case CT_ATTRIB_FIELD_DOUBLE:
 		{
 			double DValue;
-			if (CT_CURRENT_ATTRIB_IO::AskDblField(IthField, DValue) != IO_OK) break;
+			if (CT_CURRENT_ATTRIB_IO::AskDblField(IthField, DValue) != IO_OK)
+			{
+				break;
+			}
 			Value = FString::Printf(TEXT("%lf"), DValue);
 			break;
 		}
 		case CT_ATTRIB_FIELD_STRING:
 		{
 			CT_STR StrValue;
-			if (CT_CURRENT_ATTRIB_IO::AskStrField(IthField, StrValue) != IO_OK) break;
+			if (CT_CURRENT_ATTRIB_IO::AskStrField(IthField, StrValue) != IO_OK)
+			{
+				break;
+			}
 			Value = StrValue.toUnicode();
 			break;
 		}
@@ -1011,11 +1024,17 @@ void FCoreTechFileParser::GetStringMetaDataValue(CT_OBJECT_ID NodeId, const TCHA
 	CT_UINT32 IthAttrib = 0;
 	while (CT_OBJECT_IO::SearchAttribute(NodeId, CT_ATTRIB_STRING_METADATA, IthAttrib++) == IO_OK)
 	{
-		if (CT_CURRENT_ATTRIB_IO::AskStrField(ITH_STRING_METADATA_NAME, FieldName) != IO_OK) break;
+		if (CT_CURRENT_ATTRIB_IO::AskStrField(ITH_STRING_METADATA_NAME, FieldName) != IO_OK)
+		{
+			break;
+		}
 		if (!FCString::Strcmp(InMetaDataName, FieldName.toUnicode()))
 		{
 			CT_STR FieldStrValue;
-			if (CT_CURRENT_ATTRIB_IO::AskStrField(ITH_STRING_METADATA_VALUE, FieldStrValue) != IO_OK) break;
+			if (CT_CURRENT_ATTRIB_IO::AskStrField(ITH_STRING_METADATA_VALUE, FieldStrValue) != IO_OK)
+			{
+				break;
+			}
 			if (FieldStrValue.IsEmpty())
 			{
 				return;
@@ -1071,45 +1090,85 @@ void FCoreTechFileParser::ReadNodeMetaData(CT_OBJECT_ID NodeId, TMap<FString, FS
 		FString              FieldValue;
 
 
-		if (CT_CURRENT_ATTRIB_IO::AskAttributeType(AttributeType) != IO_OK) continue;;
+		if (CT_CURRENT_ATTRIB_IO::AskAttributeType(AttributeType) != IO_OK) 
+		{
+			continue;
+		}
+
 		switch (AttributeType) {
 
 		case CT_ATTRIB_SPLT:
 			break;
 
 		case CT_ATTRIB_NAME:
-			if (CT_CURRENT_ATTRIB_IO::AskStrField(ITH_NAME_VALUE, FieldStrValue) != IO_OK) break;
-			if (FieldStrValue.IsEmpty()) { break; }
+			if (CT_CURRENT_ATTRIB_IO::AskStrField(ITH_NAME_VALUE, FieldStrValue) != IO_OK)
+			{
+				break;
+			}
+			if (FieldStrValue.IsEmpty())
+			{
+				break;
+			}
 			OutMetaData.Add(TEXT("CTName"), FieldStrValue.toUnicode());
 			break;
 
 		case CT_ATTRIB_ORIGINAL_NAME:
-			if (CT_CURRENT_ATTRIB_IO::AskStrField(ITH_NAME_VALUE, FieldStrValue) != IO_OK) break;
-			if (FieldStrValue.IsEmpty()) { break; }
+			if (CT_CURRENT_ATTRIB_IO::AskStrField(ITH_NAME_VALUE, FieldStrValue) != IO_OK)
+			{
+				break;
+			}
+			if (FieldStrValue.IsEmpty())
+			{
+				break;
+			}
 			OutMetaData.Add(TEXT("Name"), FieldStrValue.toUnicode());
 			break;
 
 		case CT_ATTRIB_ORIGINAL_FILENAME:
-			if (CT_CURRENT_ATTRIB_IO::AskStrField(ITH_FILENAME_VALUE, FieldStrValue) != IO_OK) break;
-			if (FieldStrValue.IsEmpty()) { break; }
+			if (CT_CURRENT_ATTRIB_IO::AskStrField(ITH_FILENAME_VALUE, FieldStrValue) != IO_OK)
+			{
+				break;
+			}
+			if (FieldStrValue.IsEmpty())
+			{
+				break;
+			}
 			OutMetaData.Add(TEXT("FileName"), FieldStrValue.toUnicode());
 			break;
 
 		case CT_ATTRIB_UUID:
-			if (CT_CURRENT_ATTRIB_IO::AskStrField(ITH_UUID_VALUE, FieldStrValue) != IO_OK) break;
-			if (FieldStrValue.IsEmpty()) { break; }
+			if (CT_CURRENT_ATTRIB_IO::AskStrField(ITH_UUID_VALUE, FieldStrValue) != IO_OK)
+			{
+				break;
+			}
+			if (FieldStrValue.IsEmpty())
+			{
+				break;
+			}
 			OutMetaData.Add(TEXT("UUID"), FieldStrValue.toUnicode());
 			break;
 
 		case CT_ATTRIB_INPUT_FORMAT_AND_EMETTOR:
-			if (CT_CURRENT_ATTRIB_IO::AskStrField(ITH_INPUT_FORMAT_AND_EMETTOR, FieldStrValue) != IO_OK) break;
-			if (FieldStrValue.IsEmpty()) { break; }
+			if (CT_CURRENT_ATTRIB_IO::AskStrField(ITH_INPUT_FORMAT_AND_EMETTOR, FieldStrValue) != IO_OK)
+			{
+				break;
+			}
+			if (FieldStrValue.IsEmpty())
+			{
+				break;
+			}
 			OutMetaData.Add(TEXT("Input_Format_and_Emitter"), FieldStrValue.toUnicode());
 			break;
 
 		case CT_ATTRIB_CONFIGURATION_NAME:
-			if (CT_CURRENT_ATTRIB_IO::AskStrField(ITH_NAME_VALUE, FieldStrValue) != IO_OK) break;
-			if (FieldStrValue.IsEmpty()) { break; }
+			if (CT_CURRENT_ATTRIB_IO::AskStrField(ITH_NAME_VALUE, FieldStrValue) != IO_OK)
+			{
+				break;
+			}
+			if (FieldStrValue.IsEmpty())
+			{
+				break;
+			}
 			OutMetaData.Add(TEXT("ConfigurationName"), FieldStrValue.toUnicode());
 			break;
 
@@ -1124,7 +1183,10 @@ void FCoreTechFileParser::ReadNodeMetaData(CT_OBJECT_ID NodeId, TMap<FString, FS
 
 		case CT_ATTRIB_COLORID:
 			{
-				if (CT_CURRENT_ATTRIB_IO::AskIntField(ITH_COLORID_VALUE, FieldIntValue) != IO_OK) break;
+				if (CT_CURRENT_ATTRIB_IO::AskIntField(ITH_COLORID_VALUE, FieldIntValue) != IO_OK)
+				{
+					break;
+				}
 				uint32 ColorId = FieldIntValue;
 
 				uint8 Alpha = 255;
@@ -1147,7 +1209,10 @@ void FCoreTechFileParser::ReadNodeMetaData(CT_OBJECT_ID NodeId, TMap<FString, FS
 
 		case CT_ATTRIB_MATERIALID:
 		{
-			if (CT_CURRENT_ATTRIB_IO::AskIntField(ITH_MATERIALID_VALUE, FieldIntValue) != IO_OK) break;
+			if (CT_CURRENT_ATTRIB_IO::AskIntField(ITH_MATERIALID_VALUE, FieldIntValue) != IO_OK)
+			{
+				break;
+			}
 			if (FArchiveMaterial* Material = MockUpDescription.MaterialHIdToMaterial.Find(FieldIntValue))
 			{
 				OutMetaData.Add(TEXT("MaterialName"), FString::FromInt(Material->UEMaterialName));
@@ -1156,7 +1221,10 @@ void FCoreTechFileParser::ReadNodeMetaData(CT_OBJECT_ID NodeId, TMap<FString, FS
 		}
 
 		case CT_ATTRIB_TRANSPARENCY:
-			if (CT_CURRENT_ATTRIB_IO::AskDblField(ITH_TRANSPARENCY_VALUE, FieldDoubleValue0) != IO_OK) break;
+			if (CT_CURRENT_ATTRIB_IO::AskDblField(ITH_TRANSPARENCY_VALUE, FieldDoubleValue0) != IO_OK)
+			{
+				break;
+			}
 			FieldIntValue = FMath::Max((1. - FieldDoubleValue0), FieldDoubleValue0) * 255.;
 			OutMetaData.Add(TEXT("Transparency"), FString::FromInt(FieldIntValue));
 			break;
@@ -1166,7 +1234,10 @@ void FCoreTechFileParser::ReadNodeMetaData(CT_OBJECT_ID NodeId, TMap<FString, FS
 			break;
 
 		case CT_ATTRIB_REFCOUNT:
-			if (CT_CURRENT_ATTRIB_IO::AskIntField(ITH_REFCOUNT_VALUE, FieldIntValue) != IO_OK) break;
+			if (CT_CURRENT_ATTRIB_IO::AskIntField(ITH_REFCOUNT_VALUE, FieldIntValue) != IO_OK)
+			{
+				break;
+			}
 			//OutMetaData.Add(TEXT("RefCount"), FString::FromInt(FieldIntValue));
 			break;
 
@@ -1179,13 +1250,25 @@ void FCoreTechFileParser::ReadNodeMetaData(CT_OBJECT_ID NodeId, TMap<FString, FS
 			break;
 
 		case CT_ATTRIB_MASS_PROPERTIES:
-			if (CT_CURRENT_ATTRIB_IO::AskDblField(ITH_MASS_PROPERTIES_AREA, FieldDoubleValue0) != IO_OK) break;
+			if (CT_CURRENT_ATTRIB_IO::AskDblField(ITH_MASS_PROPERTIES_AREA, FieldDoubleValue0) != IO_OK)
+			{
+				break;
+			}
 			OutMetaData.Add(TEXT("Area"), FString::Printf(TEXT("%lf"), FieldDoubleValue0));
-			if (CT_CURRENT_ATTRIB_IO::AskDblField(ITH_MASS_PROPERTIES_VOLUME, FieldDoubleValue0) != IO_OK) break;
+			if (CT_CURRENT_ATTRIB_IO::AskDblField(ITH_MASS_PROPERTIES_VOLUME, FieldDoubleValue0) != IO_OK)
+			{
+				break;
+			}
 			OutMetaData.Add(TEXT("Volume"), FString::Printf(TEXT("%lf"), FieldDoubleValue0));
-			if (CT_CURRENT_ATTRIB_IO::AskDblField(ITH_MASS_PROPERTIES_MASS, FieldDoubleValue0) != IO_OK) break;
+			if (CT_CURRENT_ATTRIB_IO::AskDblField(ITH_MASS_PROPERTIES_MASS, FieldDoubleValue0) != IO_OK)
+			{
+				break;
+			}
 			OutMetaData.Add(TEXT("Mass"), FString::Printf(TEXT("%lf"), FieldDoubleValue0));
-			if (CT_CURRENT_ATTRIB_IO::AskDblField(ITH_MASS_PROPERTIES_LENGTH, FieldDoubleValue0) != IO_OK) break;
+			if (CT_CURRENT_ATTRIB_IO::AskDblField(ITH_MASS_PROPERTIES_LENGTH, FieldDoubleValue0) != IO_OK)
+			{
+				break;
+			}
 			OutMetaData.Add(TEXT("Length"), FString::Printf(TEXT("%lf"), FieldDoubleValue0));
 			//ITH_MASS_PROPERTIES_COGX, ITH_MASS_PROPERTIES_COGY, ITH_MASS_PROPERTIES_COGZ
 			//ITH_MASS_PROPERTIES_M1, ITH_MASS_PROPERTIES_M2, ITH_MASS_PROPERTIES_M3
@@ -1199,20 +1282,38 @@ void FCoreTechFileParser::ReadNodeMetaData(CT_OBJECT_ID NodeId, TMap<FString, FS
 			break;
 
 		case CT_ATTRIB_INTEGER_METADATA:
-			if (CT_CURRENT_ATTRIB_IO::AskStrField(ITH_INTEGER_METADATA_NAME, FieldName) != IO_OK) break;
-			if (CT_CURRENT_ATTRIB_IO::AskIntField(ITH_INTEGER_METADATA_VALUE, FieldIntValue) != IO_OK) break;
+			if (CT_CURRENT_ATTRIB_IO::AskStrField(ITH_INTEGER_METADATA_NAME, FieldName) != IO_OK)
+			{
+				break;
+			}
+			if (CT_CURRENT_ATTRIB_IO::AskIntField(ITH_INTEGER_METADATA_VALUE, FieldIntValue) != IO_OK)
+			{
+				break;
+			}
 			OutMetaData.Add(FieldName.toUnicode(), FString::FromInt(FieldIntValue));
 			break;
 
 		case CT_ATTRIB_DOUBLE_METADATA:
-			if (CT_CURRENT_ATTRIB_IO::AskStrField(ITH_DOUBLE_METADATA_NAME, FieldName) != IO_OK) break;
-			if (CT_CURRENT_ATTRIB_IO::AskDblField(ITH_DOUBLE_METADATA_VALUE, FieldDoubleValue0) != IO_OK) break;
+			if (CT_CURRENT_ATTRIB_IO::AskStrField(ITH_DOUBLE_METADATA_NAME, FieldName) != IO_OK)
+			{
+				break;
+			}
+			if (CT_CURRENT_ATTRIB_IO::AskDblField(ITH_DOUBLE_METADATA_VALUE, FieldDoubleValue0) != IO_OK)
+			{
+				break;
+			}
 			OutMetaData.Add(FieldName.toUnicode(), FString::Printf(TEXT("%lf"), FieldDoubleValue0));
 			break;
 
 		case CT_ATTRIB_STRING_METADATA:
-			if (CT_CURRENT_ATTRIB_IO::AskStrField(ITH_STRING_METADATA_NAME, FieldName) != IO_OK) break;
-			if (CT_CURRENT_ATTRIB_IO::AskStrField(ITH_STRING_METADATA_VALUE, FieldStrValue) != IO_OK) break;
+			if (CT_CURRENT_ATTRIB_IO::AskStrField(ITH_STRING_METADATA_NAME, FieldName) != IO_OK)
+			{
+				break;
+			}
+			if (CT_CURRENT_ATTRIB_IO::AskStrField(ITH_STRING_METADATA_VALUE, FieldStrValue) != IO_OK)
+			{
+				break;
+			}
 			if (FieldStrValue.IsEmpty()) break;
 			if(ConfigName== FieldName.toUnicode())
 			{
@@ -1222,9 +1323,18 @@ void FCoreTechFileParser::ReadNodeMetaData(CT_OBJECT_ID NodeId, TMap<FString, FS
 			break;
 
 		case CT_ATTRIB_ORIGINAL_UNITS:
-			if (CT_CURRENT_ATTRIB_IO::AskDblField(ITH_ORIGINAL_UNITS_MASS, FieldDoubleValue0) != IO_OK) break;
-			if (CT_CURRENT_ATTRIB_IO::AskDblField(ITH_ORIGINAL_UNITS_LENGTH, FieldDoubleValue1) != IO_OK) break;
-			if (CT_CURRENT_ATTRIB_IO::AskDblField(ITH_ORIGINAL_UNITS_DURATION, FieldDoubleValue2) != IO_OK) break;
+			if (CT_CURRENT_ATTRIB_IO::AskDblField(ITH_ORIGINAL_UNITS_MASS, FieldDoubleValue0) != IO_OK)
+			{
+				break;
+			}
+			if (CT_CURRENT_ATTRIB_IO::AskDblField(ITH_ORIGINAL_UNITS_LENGTH, FieldDoubleValue1) != IO_OK)
+			{
+				break;
+			}
+			if (CT_CURRENT_ATTRIB_IO::AskDblField(ITH_ORIGINAL_UNITS_DURATION, FieldDoubleValue2) != IO_OK)
+			{
+				break;
+			}
 			OutMetaData.Add(TEXT("OriginalUnitsMass"), FString::Printf(TEXT("%lf"), FieldDoubleValue0));
 			OutMetaData.Add(TEXT("OriginalUnitsLength"), FString::Printf(TEXT("%lf"), FieldDoubleValue1));
 			OutMetaData.Add(TEXT("OriginalUnitsDuration"), FString::Printf(TEXT("%lf"), FieldDoubleValue2));
@@ -1236,20 +1346,50 @@ void FCoreTechFileParser::ReadNodeMetaData(CT_OBJECT_ID NodeId, TMap<FString, FS
 			break;
 
 		case CT_ATTRIB_PRODUCT:
-			if (CT_CURRENT_ATTRIB_IO::AskStrField(ITH_PRODUCT_REVISION, FieldStrValue) != IO_OK) break;
-			if (FieldStrValue.IsEmpty()) { break; }
+			if (CT_CURRENT_ATTRIB_IO::AskStrField(ITH_PRODUCT_REVISION, FieldStrValue) != IO_OK)
+			{
+				break;
+			}
+			if (FieldStrValue.IsEmpty())
+			{
+				break;
+			} 
 			OutMetaData.Add(TEXT("ProductRevision"), FieldStrValue.toUnicode());
-			if (CT_CURRENT_ATTRIB_IO::AskStrField(ITH_PRODUCT_DEFINITION, FieldStrValue) != IO_OK) break;
-			if (FieldStrValue.IsEmpty()) { break; }
+			if (CT_CURRENT_ATTRIB_IO::AskStrField(ITH_PRODUCT_DEFINITION, FieldStrValue) != IO_OK)
+			{
+				break;
+			}
+			if (FieldStrValue.IsEmpty())
+			{
+				break;
+			}
 			OutMetaData.Add(TEXT("ProductDefinition"), FieldStrValue.toUnicode());
-			if (CT_CURRENT_ATTRIB_IO::AskStrField(ITH_PRODUCT_NOMENCLATURE, FieldStrValue) != IO_OK) break;
-			if (FieldStrValue.IsEmpty()) { break; }
+			if (CT_CURRENT_ATTRIB_IO::AskStrField(ITH_PRODUCT_NOMENCLATURE, FieldStrValue) != IO_OK)
+			{
+				break;
+			}
+			if (FieldStrValue.IsEmpty())
+			{
+				break;
+			}
 			OutMetaData.Add(TEXT("ProductNomenclature"), FieldStrValue.toUnicode());
-			if (CT_CURRENT_ATTRIB_IO::AskStrField(ITH_PRODUCT_SOURCE, FieldStrValue) != IO_OK) break;
-			if (FieldStrValue.IsEmpty()) { break; }
+			if (CT_CURRENT_ATTRIB_IO::AskStrField(ITH_PRODUCT_SOURCE, FieldStrValue) != IO_OK)
+			{
+				break;
+			}
+			if (FieldStrValue.IsEmpty())
+			{
+				break;
+			}
 			OutMetaData.Add(TEXT("ProductSource"), FieldStrValue.toUnicode());
-			if (CT_CURRENT_ATTRIB_IO::AskStrField(ITH_PRODUCT_DESCRIPTION, FieldStrValue) != IO_OK) break;
-			if (FieldStrValue.IsEmpty()) { break; }
+			if (CT_CURRENT_ATTRIB_IO::AskStrField(ITH_PRODUCT_DESCRIPTION, FieldStrValue) != IO_OK)
+			{
+				break;
+			}
+			if (FieldStrValue.IsEmpty())
+			{
+				break;
+			}
 			OutMetaData.Add(TEXT("ProductDescription"), FieldStrValue.toUnicode());
 			break;
 
@@ -1280,21 +1420,42 @@ void FCoreTechFileParser::ReadNodeMetaData(CT_OBJECT_ID NodeId, TMap<FString, FS
 			break;
 
 		case CT_ATTRIB_INTEGER_PARAMETER:
-			if (CT_CURRENT_ATTRIB_IO::AskStrField(ITH_INTEGER_PARAMETER_NAME, FieldName) != IO_OK) break;
-			if (CT_CURRENT_ATTRIB_IO::AskIntField(ITH_INTEGER_PARAMETER_VALUE, FieldIntValue) != IO_OK) break;
+			if (CT_CURRENT_ATTRIB_IO::AskStrField(ITH_INTEGER_PARAMETER_NAME, FieldName) != IO_OK)
+			{
+				break;
+			}
+			if (CT_CURRENT_ATTRIB_IO::AskIntField(ITH_INTEGER_PARAMETER_VALUE, FieldIntValue) != IO_OK)
+			{
+				break;
+			}
 			OutMetaData.Add(FieldName.toUnicode(), FString::FromInt(FieldIntValue));
 			break;
 
 		case CT_ATTRIB_DOUBLE_PARAMETER:
-			if (CT_CURRENT_ATTRIB_IO::AskStrField(ITH_DOUBLE_PARAMETER_NAME, FieldName) != IO_OK) break;
-			if (CT_CURRENT_ATTRIB_IO::AskDblField(ITH_DOUBLE_PARAMETER_VALUE, FieldDoubleValue0) != IO_OK) break;
+			if (CT_CURRENT_ATTRIB_IO::AskStrField(ITH_DOUBLE_PARAMETER_NAME, FieldName) != IO_OK)
+			{
+				break;
+			}
+			if (CT_CURRENT_ATTRIB_IO::AskDblField(ITH_DOUBLE_PARAMETER_VALUE, FieldDoubleValue0) != IO_OK)
+			{
+				break;
+			}
 			OutMetaData.Add(FieldName.toUnicode(), FString::Printf(TEXT("%lf"), FieldDoubleValue0));
 			break;
 
 		case CT_ATTRIB_STRING_PARAMETER:
-			if (CT_CURRENT_ATTRIB_IO::AskStrField(ITH_STRING_PARAMETER_NAME, FieldName) != IO_OK) break;
-			if (CT_CURRENT_ATTRIB_IO::AskStrField(ITH_STRING_PARAMETER_VALUE, FieldStrValue) != IO_OK) break;
-			if (FieldStrValue.IsEmpty()) { break; }
+			if (CT_CURRENT_ATTRIB_IO::AskStrField(ITH_STRING_PARAMETER_NAME, FieldName) != IO_OK)
+			{
+				break;
+			}
+			if (CT_CURRENT_ATTRIB_IO::AskStrField(ITH_STRING_PARAMETER_VALUE, FieldStrValue) != IO_OK)
+			{
+				break;
+			}
+			if (FieldStrValue.IsEmpty())
+			{
+				break;
+			}
 			OutMetaData.Add(FieldName.toUnicode(), FieldStrValue.toUnicode());
 			break;
 
@@ -1305,24 +1466,54 @@ void FCoreTechFileParser::ReadNodeMetaData(CT_OBJECT_ID NodeId, TMap<FString, FS
 			break;
 
 		case CT_ATTRIB_SAVE_OPTION:
-			if (CT_CURRENT_ATTRIB_IO::AskStrField(ITH_SAVE_OPTION_AUTHOR, FieldStrValue) != IO_OK) break;
-			if (FieldStrValue.IsEmpty()) { break; }
+			if (CT_CURRENT_ATTRIB_IO::AskStrField(ITH_SAVE_OPTION_AUTHOR, FieldStrValue) != IO_OK)
+			{
+				break;
+			}
+			if (FieldStrValue.IsEmpty())
+			{
+				break;
+			}
 			OutMetaData.Add(TEXT("SaveOptionAuthor"), FieldStrValue.toUnicode());
 
-			if (CT_CURRENT_ATTRIB_IO::AskStrField(ITH_SAVE_OPTION_ORGANIZATION, FieldStrValue) != IO_OK) break;
-			if (FieldStrValue.IsEmpty()) { break; }
+			if (CT_CURRENT_ATTRIB_IO::AskStrField(ITH_SAVE_OPTION_ORGANIZATION, FieldStrValue) != IO_OK)
+			{
+				break;
+			}
+			if (FieldStrValue.IsEmpty())
+			{
+				break;
+			}
 			OutMetaData.Add(TEXT("SaveOptionOrganization"), FieldStrValue.toUnicode());
 	
-			if (CT_CURRENT_ATTRIB_IO::AskStrField(ITH_SAVE_OPTION_FILE_DESCRIPTION, FieldStrValue) != IO_OK) break;
-			if (FieldStrValue.IsEmpty()) { break; }
+			if (CT_CURRENT_ATTRIB_IO::AskStrField(ITH_SAVE_OPTION_FILE_DESCRIPTION, FieldStrValue) != IO_OK)
+			{
+				break;
+			}
+			if (FieldStrValue.IsEmpty())
+			{
+				break;
+			}
 			OutMetaData.Add(TEXT("SaveOptionFileDescription"), FieldStrValue.toUnicode());
 
-			if (CT_CURRENT_ATTRIB_IO::AskStrField(ITH_SAVE_OPTION_AUTHORISATION, FieldStrValue) != IO_OK) break;
-			if (FieldStrValue.IsEmpty()) { break; }
+			if (CT_CURRENT_ATTRIB_IO::AskStrField(ITH_SAVE_OPTION_AUTHORISATION, FieldStrValue) != IO_OK)
+			{
+				break;
+			}
+			if (FieldStrValue.IsEmpty())
+			{
+				break;
+			}
 			OutMetaData.Add(TEXT("SaveOptionAuthorisation"), FieldStrValue.toUnicode());
 
-			if (CT_CURRENT_ATTRIB_IO::AskStrField(ITH_SAVE_OPTION_PREPROCESSOR, FieldStrValue) != IO_OK) break;
-			if (FieldStrValue.IsEmpty()) { break; }
+			if (CT_CURRENT_ATTRIB_IO::AskStrField(ITH_SAVE_OPTION_PREPROCESSOR, FieldStrValue) != IO_OK) 
+			{
+				break;
+			}
+			if (FieldStrValue.IsEmpty())
+			{
+				break;
+			}
 			OutMetaData.Add(TEXT("SaveOptionPreprocessor"), FieldStrValue.toUnicode());
 			break;
 
@@ -1332,15 +1523,30 @@ void FCoreTechFileParser::ReadNodeMetaData(CT_OBJECT_ID NodeId, TMap<FString, FS
 			break;
 
 		case CT_ATTRIB_ORIGINAL_ID_STRING:
-			if (CT_CURRENT_ATTRIB_IO::AskStrField(ITH_ORIGINAL_ID_VALUE_STRING, FieldStrValue) != IO_OK) break;
-			if (FieldStrValue.IsEmpty()) { break; }
+			if (CT_CURRENT_ATTRIB_IO::AskStrField(ITH_ORIGINAL_ID_VALUE_STRING, FieldStrValue) != IO_OK)
+			{
+				break;
+			}
+			if (FieldStrValue.IsEmpty())
+			{
+				break;
+			}
 			OutMetaData.Add(TEXT("OriginalIdStr"), FieldStrValue.toUnicode());
 			break;
 
 		case CT_ATTRIB_COLOR_RGB_DOUBLE:
-			if (CT_CURRENT_ATTRIB_IO::AskDblField(ITH_ATTRIB_COLOR_R_DOUBLE, FieldDoubleValue0) != IO_OK) break;
-			if (CT_CURRENT_ATTRIB_IO::AskDblField(ITH_ATTRIB_COLOR_G_DOUBLE, FieldDoubleValue1) != IO_OK) break;
-			if (CT_CURRENT_ATTRIB_IO::AskDblField(ITH_ATTRIB_COLOR_B_DOUBLE, FieldDoubleValue2) != IO_OK) break;
+			if (CT_CURRENT_ATTRIB_IO::AskDblField(ITH_ATTRIB_COLOR_R_DOUBLE, FieldDoubleValue0) != IO_OK)
+			{
+				break;
+			}
+			if (CT_CURRENT_ATTRIB_IO::AskDblField(ITH_ATTRIB_COLOR_G_DOUBLE, FieldDoubleValue1) != IO_OK)
+			{
+				break;
+			}
+			if (CT_CURRENT_ATTRIB_IO::AskDblField(ITH_ATTRIB_COLOR_B_DOUBLE, FieldDoubleValue2) != IO_OK)
+			{
+				break;
+			}
 			FieldValue = FString::Printf(TEXT("%lf"), FieldDoubleValue0) + TEXT(", ") + FString::Printf(TEXT("%lf"), FieldDoubleValue1) + TEXT(", ") + FString::Printf(TEXT("%lf"), FieldDoubleValue2);
 			//OutMetaData.Add(TEXT("ColorRGBDouble"), FieldValue);
 			break;
@@ -1355,21 +1561,42 @@ void FCoreTechFileParser::ReadNodeMetaData(CT_OBJECT_ID NodeId, TMap<FString, FS
 			break;
 
 		case CT_ATTRIB_INTEGER_VALIDATION_ATTRIBUTE:
-			if (CT_CURRENT_ATTRIB_IO::AskStrField(ITH_INTEGER_VALIDATION_NAME, FieldName) != IO_OK) break;
-			if (CT_CURRENT_ATTRIB_IO::AskIntField(ITH_INTEGER_VALIDATION_VALUE, FieldIntValue) != IO_OK) break;
+			if (CT_CURRENT_ATTRIB_IO::AskStrField(ITH_INTEGER_VALIDATION_NAME, FieldName) != IO_OK)
+			{
+				break;
+			}
+			if (CT_CURRENT_ATTRIB_IO::AskIntField(ITH_INTEGER_VALIDATION_VALUE, FieldIntValue) != IO_OK)
+			{
+				break;
+			}
 			OutMetaData.Add(FieldName.toUnicode(), FString::FromInt(FieldIntValue));
 			break;
 
 		case CT_ATTRIB_DOUBLE_VALIDATION_ATTRIBUTE:
-			if (CT_CURRENT_ATTRIB_IO::AskStrField(ITH_DOUBLE_VALIDATION_NAME, FieldName) != IO_OK) break;
-			if (CT_CURRENT_ATTRIB_IO::AskDblField(ITH_DOUBLE_VALIDATION_VALUE, FieldDoubleValue0) != IO_OK) break;
+			if (CT_CURRENT_ATTRIB_IO::AskStrField(ITH_DOUBLE_VALIDATION_NAME, FieldName) != IO_OK)
+			{
+				break;
+			}
+			if (CT_CURRENT_ATTRIB_IO::AskDblField(ITH_DOUBLE_VALIDATION_VALUE, FieldDoubleValue0) != IO_OK)
+			{
+				break;
+			}
 			OutMetaData.Add(FieldName.toUnicode(), FString::Printf(TEXT("%lf"), FieldDoubleValue0));
 			break;
 
 		case CT_ATTRIB_STRING_VALIDATION_ATTRIBUTE:
-			if (CT_CURRENT_ATTRIB_IO::AskStrField(ITH_STRING_VALIDATION_NAME, FieldName) != IO_OK) break;
-			if (CT_CURRENT_ATTRIB_IO::AskStrField(ITH_STRING_VALIDATION_VALUE, FieldStrValue) != IO_OK) break;
-			if (FieldStrValue.IsEmpty()) { break; }
+			if (CT_CURRENT_ATTRIB_IO::AskStrField(ITH_STRING_VALIDATION_NAME, FieldName) != IO_OK)
+			{
+				break;
+			}
+			if (CT_CURRENT_ATTRIB_IO::AskStrField(ITH_STRING_VALIDATION_VALUE, FieldStrValue) != IO_OK)
+			{
+				break;
+			}
+			if (FieldStrValue.IsEmpty())
+			{
+				break;
+			}
 			OutMetaData.Add(FieldName.toUnicode(), FieldStrValue.toUnicode());
 			break;
 
@@ -1388,8 +1615,14 @@ void FCoreTechFileParser::ReadNodeMetaData(CT_OBJECT_ID NodeId, TMap<FString, FS
 			break;
 
 		case CT_ATTRIB_GROUPNAME:
-			if (CT_CURRENT_ATTRIB_IO::AskStrField(ITH_GROUPNAME_VALUE, FieldStrValue) != IO_OK) break;
-			if (FieldStrValue.IsEmpty()) { break; }
+			if (CT_CURRENT_ATTRIB_IO::AskStrField(ITH_GROUPNAME_VALUE, FieldStrValue) != IO_OK)
+			{
+				break;
+			}
+			if (FieldStrValue.IsEmpty())
+			{
+				break;
+			}
 			OutMetaData.Add(TEXT("GroupName"), FieldStrValue.toUnicode());
 			break;
 
diff --git a/Engine/Plugins/Enterprise/DatasmithCADImporter/Source/CADInterfaces/Private/CoreTechTypes.cpp b/Engine/Plugins/Enterprise/DatasmithCADImporter/Source/CADInterfaces/Private/CoreTechTypes.cpp
--- a/Engine/Plugins/Enterprise/DatasmithCADImporter/Source/CADInterfaces/Private/CoreTechTypes.cpp
+++ b/Engine/Plugins/Enterprise/DatasmithCADImporter/Source/CADInterfaces/Private/CoreTechTypes.cpp
@@ -14,10 +14,21 @@ namespace CADLibrary
 	// Note: CTKIO_* functions are not functionally useful.
 	// This wrapping allows a correct profiling of the CT API.
 
-	CT_IO_ERROR CTKIO_InitializeKernel(double Unit, const TCHAR* KernelIOPath)
+	CT_IO_ERROR CTKIO_InitializeKernel(double Unit, const TCHAR* EnginePluginsPath)
 	{
+		FString KernelIOPath;
+		if (FCString::Strlen(EnginePluginsPath))
+		{
+			KernelIOPath = FPaths::Combine(EnginePluginsPath, TEXT(KERNEL_IO_PLUGINSPATH));
+			KernelIOPath = FPaths::ConvertRelativePathToFull(KernelIOPath);
+			if (!FPaths::DirectoryExists(KernelIOPath))
+			{
+				KernelIOPath.Empty();
+			}
+		}
+
 		CT_STR appName = CoreTechLicenseKey;
-		return CT_KERNEL_IO::InitializeKernel(appName, Unit, 0.00001 / Unit, KernelIOPath);
+		return CT_KERNEL_IO::InitializeKernel(appName, Unit, 0.00001 / Unit, *KernelIOPath);
 	}
 
 	CT_IO_ERROR CTKIO_ShutdownKernel()
diff --git a/Engine/Plugins/Enterprise/DatasmithCADImporter/Source/CADInterfaces/Public/CADSceneGraph.h b/Engine/Plugins/Enterprise/DatasmithCADImporter/Source/CADInterfaces/Public/CADSceneGraph.h
--- a/Engine/Plugins/Enterprise/DatasmithCADImporter/Source/CADInterfaces/Public/CADSceneGraph.h
+++ b/Engine/Plugins/Enterprise/DatasmithCADImporter/Source/CADInterfaces/Public/CADSceneGraph.h
@@ -126,6 +126,9 @@ class CADINTERFACES_API FArchiveMockUp
 public:
 	friend FArchive& operator<<(FArchive& Ar, FArchiveMockUp& C);
 
+	void SerializeMockUp(const TCHAR* Filename);
+	void DeserializeMockUpFile(const TCHAR* Filename);
+
 public:
 	FString CADFile;
 	FString SceneGraphArchive;
@@ -146,8 +149,6 @@ public:
 	TMap<CadId, int32> CADIdToInstanceIndex;
 };
 
-CADINTERFACES_API void SerializeMockUp(FArchiveMockUp& Object, const TCHAR* Filename);
-CADINTERFACES_API void DeserializeMockUpFile(const TCHAR* Filename, FArchiveMockUp& MockUp);
 
 }
 
diff --git a/Engine/Plugins/Enterprise/DatasmithCADImporter/Source/CADInterfaces/Public/CoreTechFileParser.h b/Engine/Plugins/Enterprise/DatasmithCADImporter/Source/CADInterfaces/Public/CoreTechFileParser.h
--- a/Engine/Plugins/Enterprise/DatasmithCADImporter/Source/CADInterfaces/Public/CoreTechFileParser.h
+++ b/Engine/Plugins/Enterprise/DatasmithCADImporter/Source/CADInterfaces/Public/CoreTechFileParser.h
@@ -56,17 +56,19 @@ public:
 	using EProcessResult = ECoretechParsingResult;
 
 	/**
-	 * @param InCTFullPath Full path of the CAD file to parse
-	 * @param InCachePath Full path of the cache in which the data will be saved
 	 * @param ImportParams Parameters that setting import data like mesh SAG...
-	 * @param KernelIOPath Full Path of KernelIO libraries (oda_translator.exe, ...). Mandatory to import DWG, or DGN files
+	 * @param EnginePluginsPath Full Path of EnginePlugins. Mandatory to set KernelIO to import DWG, or DGN files
+	 * @param InCachePath Full path of the cache in which the data will be saved
 	 */
-	FCoreTechFileParser(const FString& InCTFullPath, const FString& InCachePath, const FImportParameters& ImportParams, const TCHAR* KernelIOPath = TEXT(""));
+	FCoreTechFileParser(const FImportParameters& ImportParams, const FString& EnginePluginsPath = TEXT(""), const FString& InCachePath = TEXT(""));
 
 	double GetMetricUnit() const { return 0.01; }
 
-	EProcessResult ProcessFile();
-	void GetBodyTessellation(CT_OBJECT_ID BodyId, FBodyMesh& OutBodyMesh, const FImportParameters& ImportParams, uint32 ParentMaterialHash);
+	/**
+	 * @param InCTFullPath Full path of the CAD file to parse
+	 */
+	EProcessResult ProcessFile(const FString& InCTFullPath);
+	void GetBodyTessellation(CT_OBJECT_ID BodyId, FBodyMesh& OutBodyMesh, uint32 ParentMaterialHash);
 
 
 	const TSet<FString>& GetExternalRefSet()
@@ -88,7 +90,7 @@ public:
 	}
 
 private:
-	EProcessResult ReadFileWithKernelIO();
+	EProcessResult ReadFileWithKernelIO(const FString& FullPath);
 	bool ReadNode(CT_OBJECT_ID NodeId, uint32 ParentMaterialHash);
 	bool ReadInstance(CT_OBJECT_ID NodeId, uint32 ParentMaterialHash);
 	bool ReadComponent(CT_OBJECT_ID NodeId, uint32 ParentMaterialHash);
@@ -115,9 +117,8 @@ private:
 	void ExportMeshArchiveFile();
 
 protected:
-	const FString& CachePath;
+	FString CachePath;
 
-	FString FullPath;
 	FString CADFile;
 
 	FString FileConfiguration;
diff --git a/Engine/Plugins/Enterprise/DatasmithCADImporter/Source/CADLibrary/Private/CoreTechHelper.cpp b/Engine/Plugins/Enterprise/DatasmithCADImporter/Source/CADLibrary/Private/CoreTechHelper.cpp
--- a/Engine/Plugins/Enterprise/DatasmithCADImporter/Source/CADLibrary/Private/CoreTechHelper.cpp
+++ b/Engine/Plugins/Enterprise/DatasmithCADImporter/Source/CADLibrary/Private/CoreTechHelper.cpp
@@ -668,16 +668,14 @@ CT_IO_ERROR Tessellate(CT_OBJECT_ID MainObjectId, const FImportParameters& Impor
 
 	SetCoreTechTessellationState(ImportParams);
 
-	FString FullPath;
-	FString CachePath;
-	FCoreTechFileParser Parser = FCoreTechFileParser(FullPath, CachePath, ImportParams);
+	FCoreTechFileParser Parser = FCoreTechFileParser(ImportParams);
 
 	FBodyMesh BodyMesh;
 	BodyMesh.BodyID = 1;
 
 	while (CT_OBJECT_ID BodyId = Objects.IteratorIter())
 	{
-		Parser.GetBodyTessellation(BodyId, BodyMesh, ImportParams, 0);
+		Parser.GetBodyTessellation(BodyId, BodyMesh, 0);
 	}
 
 	bool bTessellated = ConvertCTBodySetToMeshDescription(ImportParams, MeshParameters, BodyMesh, MeshDesc);
diff --git a/Engine/Plugins/Enterprise/DatasmithCADImporter/Source/CADTools/Private/CADData.cpp b/Engine/Plugins/Enterprise/DatasmithCADImporter/Source/CADTools/Private/CADData.cpp
--- a/Engine/Plugins/Enterprise/DatasmithCADImporter/Source/CADTools/Private/CADData.cpp
+++ b/Engine/Plugins/Enterprise/DatasmithCADImporter/Source/CADTools/Private/CADData.cpp
@@ -4,9 +4,6 @@
 #include "CADOptions.h"
 
 #include "Misc/FileHelper.h"
-#include "Serialization/Archive.h"
-#include "Serialization/MemoryWriter.h"
-#include "Serialization/MemoryReader.h"
 
 namespace CADLibrary
 {
@@ -102,30 +99,30 @@ FArchive& operator<<(FArchive& Ar, FBodyMesh& BodyMesh)
 
 void SerializeBodyMeshSet(const TCHAR* Filename, TArray<FBodyMesh>& InBodySet)
 {
-	TArray<uint8> OutBuffer;
-	FMemoryWriter ArWriter(OutBuffer);
+	TUniquePtr<FArchive> Archive(IFileManager::Get().CreateFileWriter(Filename));
+
 	uint32 type = 234561;
-	ArWriter << type;
+	*Archive << type;
 
-	ArWriter << InBodySet;
+	*Archive << InBodySet;
 
-	FFileHelper::SaveArrayToFile(OutBuffer, Filename);
+	Archive->Close();
 }
 
 void DeserializeBodyMeshFile(const TCHAR* Filename, TArray<FBodyMesh>& OutBodySet)
 {
-	TArray<uint8> InBuffer;
-	FFileHelper::LoadFileToArray(InBuffer, Filename);
+	TUniquePtr<FArchive> Archive(IFileManager::Get().CreateFileReader(Filename));
 
-	FMemoryReader ArReader(InBuffer);
 	uint32 type = 0;
-	ArReader << type;
+	*Archive << type;
 	if (type != 234561)
 	{
+		Archive->Close();
 		return;
 	}
 
-	ArReader << OutBodySet;
+	*Archive << OutBodySet;
+	Archive->Close();
 }
 
 }
diff --git a/Engine/Plugins/Enterprise/DatasmithCADImporter/Source/DatasmithCADTranslator/DatasmithCADTranslator.Build.cs b/Engine/Plugins/Enterprise/DatasmithCADImporter/Source/DatasmithCADTranslator/DatasmithCADTranslator.Build.cs
--- a/Engine/Plugins/Enterprise/DatasmithCADImporter/Source/DatasmithCADTranslator/DatasmithCADTranslator.Build.cs
+++ b/Engine/Plugins/Enterprise/DatasmithCADImporter/Source/DatasmithCADTranslator/DatasmithCADTranslator.Build.cs
@@ -25,8 +25,8 @@ namespace UnrealBuildTool.Rules
 					"DatasmithCore",
 					"DatasmithCoreTechParametricSurfaceData",
 					"DatasmithContent",
-					"DatasmithImporter",
 					"DatasmithDispatcher",
+					"DatasmithTranslator",
 					"Engine",
 					"MeshDescription",
 					"Sockets",
diff --git a/Engine/Plugins/Enterprise/DatasmithCADImporter/Source/DatasmithCADTranslator/Private/DatasmithMeshBuilder.cpp b/Engine/Plugins/Enterprise/DatasmithCADImporter/Source/DatasmithCADTranslator/Private/DatasmithMeshBuilder.cpp
--- a/Engine/Plugins/Enterprise/DatasmithCADImporter/Source/DatasmithCADTranslator/Private/DatasmithMeshBuilder.cpp
+++ b/Engine/Plugins/Enterprise/DatasmithCADImporter/Source/DatasmithCADTranslator/Private/DatasmithMeshBuilder.cpp
@@ -6,8 +6,8 @@
 #endif // CAD_INTERFACE
 
 #include "CoreTechFileParser.h"
-#include "DatasmithMeshHelper.h"
 #include "IDatasmithSceneElements.h"
+#include "Utility/DatasmithMeshHelper.h"
 
 #include "HAL/FileManager.h"
 #include "MeshDescription.h"
diff --git a/Engine/Plugins/Enterprise/DatasmithCADImporter/Source/DatasmithCADTranslator/Private/DatasmithSceneGraphBuilder.cpp b/Engine/Plugins/Enterprise/DatasmithCADImporter/Source/DatasmithCADTranslator/Private/DatasmithSceneGraphBuilder.cpp
--- a/Engine/Plugins/Enterprise/DatasmithCADImporter/Source/DatasmithCADTranslator/Private/DatasmithSceneGraphBuilder.cpp
+++ b/Engine/Plugins/Enterprise/DatasmithCADImporter/Source/DatasmithCADTranslator/Private/DatasmithSceneGraphBuilder.cpp
@@ -124,7 +124,7 @@ void FDatasmithSceneGraphBuilder::LoadSceneGraphDescriptionFiles()
 
 		CADFileToArchiveMockUp.Add(FilePair.Key, &MockUpDescription);
 	
-		CADLibrary::DeserializeMockUpFile(*MockUpDescriptionFile, MockUpDescription);
+		MockUpDescription.DeserializeMockUpFile(*MockUpDescriptionFile);
 
 		for(const auto& ColorPair : MockUpDescription.ColorHIdToColor)
 		{
diff --git a/Engine/Plugins/Enterprise/DatasmithCADImporter/Source/DatasmithCADTranslator/Public/DatasmithSceneGraphBuilder.h b/Engine/Plugins/Enterprise/DatasmithCADImporter/Source/DatasmithCADTranslator/Public/DatasmithSceneGraphBuilder.h
--- a/Engine/Plugins/Enterprise/DatasmithCADImporter/Source/DatasmithCADTranslator/Public/DatasmithSceneGraphBuilder.h
+++ b/Engine/Plugins/Enterprise/DatasmithCADImporter/Source/DatasmithCADTranslator/Public/DatasmithSceneGraphBuilder.h
@@ -15,36 +15,12 @@
 #include "Containers/Queue.h"
 #include "Misc/Paths.h"
 
-static const FString SPACE(" ");
-
-class FCADSceneGraphDescriptionFile;
 class FDatasmithSceneSource;
 class IDatasmithActorElement;
 class IDatasmithMeshElement;
 class IDatasmithScene;
 class IDatasmithUEPbrMaterialElement;
 
-enum NODE_TYPE
-{
-	INSTANCE,
-	COMPONENT,
-	EXTERNALCOMPONENT,
-	BODY,
-	UNDEFINED,
-	UNKNOWN
-};
-
-namespace CADLibrary		
-{
-class FArchiveMockUp;
-class ICADArchiveObject;
-class FArchiveInstance;
-class FArchiveComponent;
-class FArchiveBody;
-class FArchiveMaterial;
-class FArchiveColor;
-}
-
 class ActorData  //#ueent_CAD
 {
 public:
diff --git a/Engine/Plugins/Enterprise/DatasmithCADImporter/Source/DatasmithCoreTechExtension/Private/CoreTechBlueprintLibrary.cpp b/Engine/Plugins/Enterprise/DatasmithCADImporter/Source/DatasmithCoreTechExtension/Private/CoreTechBlueprintLibrary.cpp
--- a/Engine/Plugins/Enterprise/DatasmithCADImporter/Source/DatasmithCoreTechExtension/Private/CoreTechBlueprintLibrary.cpp
+++ b/Engine/Plugins/Enterprise/DatasmithCADImporter/Source/DatasmithCoreTechExtension/Private/CoreTechBlueprintLibrary.cpp
@@ -9,8 +9,8 @@
 #include "CoreTechRetessellateAction.h"
 #include "DatasmithAdditionalData.h"
 #include "DatasmithStaticMeshImporter.h" // Call to BuildStaticMesh
+#include "DatasmithTranslator.h"
 #include "DatasmithUtils.h"
-#include "Translators/DatasmithTranslator.h"
 #include "UI/DatasmithDisplayHelper.h"
 
 #include "Algo/AnyOf.h"
diff --git a/Engine/Plugins/Enterprise/DatasmithCADImporter/Source/DatasmithCoreTechExtension/Private/CoreTechRetessellateAction.cpp b/Engine/Plugins/Enterprise/DatasmithCADImporter/Source/DatasmithCoreTechExtension/Private/CoreTechRetessellateAction.cpp
--- a/Engine/Plugins/Enterprise/DatasmithCADImporter/Source/DatasmithCoreTechExtension/Private/CoreTechRetessellateAction.cpp
+++ b/Engine/Plugins/Enterprise/DatasmithCADImporter/Source/DatasmithCoreTechExtension/Private/CoreTechRetessellateAction.cpp
@@ -11,7 +11,7 @@
 #include "DatasmithAdditionalData.h"
 #include "DatasmithStaticMeshImporter.h" // Call to BuildStaticMesh
 #include "DatasmithUtils.h"
-#include "Translators/DatasmithTranslator.h"
+#include "DatasmithTranslator.h"
 #include "UI/DatasmithDisplayHelper.h"
 
 #include "Algo/AnyOf.h"
diff --git a/Engine/Plugins/Enterprise/DatasmithCADImporter/Source/DatasmithCoreTechParametricSurfaceData/DatasmithCoreTechParametricSurfaceData.Build.cs b/Engine/Plugins/Enterprise/DatasmithCADImporter/Source/DatasmithCoreTechParametricSurfaceData/DatasmithCoreTechParametricSurfaceData.Build.cs
--- a/Engine/Plugins/Enterprise/DatasmithCADImporter/Source/DatasmithCoreTechParametricSurfaceData/DatasmithCoreTechParametricSurfaceData.Build.cs
+++ b/Engine/Plugins/Enterprise/DatasmithCADImporter/Source/DatasmithCoreTechParametricSurfaceData/DatasmithCoreTechParametricSurfaceData.Build.cs
@@ -22,7 +22,7 @@ namespace UnrealBuildTool.Rules
 					"CoreUObject",
 					"MeshDescription",
 					"DatasmithContent",
-					"DatasmithImporter",
+					"DatasmithTranslator",
 					"Engine",
 					"StaticMeshDescription",
 					"CADLibrary",
diff --git a/Engine/Plugins/Enterprise/DatasmithCADImporter/Source/DatasmithCoreTechParametricSurfaceData/Public/DatasmithCoreTechTranslator.h b/Engine/Plugins/Enterprise/DatasmithCADImporter/Source/DatasmithCoreTechParametricSurfaceData/Public/DatasmithCoreTechTranslator.h
--- a/Engine/Plugins/Enterprise/DatasmithCADImporter/Source/DatasmithCoreTechParametricSurfaceData/Public/DatasmithCoreTechTranslator.h
+++ b/Engine/Plugins/Enterprise/DatasmithCADImporter/Source/DatasmithCoreTechParametricSurfaceData/Public/DatasmithCoreTechTranslator.h
@@ -5,7 +5,7 @@
 #include "CoreMinimal.h"
 
 #include "DatasmithImportOptions.h"
-#include "Translators/DatasmithTranslator.h"
+#include "DatasmithTranslator.h"
 #include "UObject/ObjectMacros.h"
 
 class UDatasmithCommonTessellationOptions;
diff --git a/Engine/Plugins/Enterprise/DatasmithCADImporter/Source/DatasmithDispatcher/Private/DatasmithDispatcher.cpp b/Engine/Plugins/Enterprise/DatasmithCADImporter/Source/DatasmithDispatcher/Private/DatasmithDispatcher.cpp
--- a/Engine/Plugins/Enterprise/DatasmithCADImporter/Source/DatasmithDispatcher/Private/DatasmithDispatcher.cpp
+++ b/Engine/Plugins/Enterprise/DatasmithCADImporter/Source/DatasmithDispatcher/Private/DatasmithDispatcher.cpp
@@ -203,15 +203,12 @@ void FDatasmithDispatcher::CloseHandlers()
 void FDatasmithDispatcher::ProcessLocal()
 {
 #ifdef CAD_INTERFACE
-	FString KernelIOPath = FPaths::Combine(FPaths::EnginePluginsDir(), TEXT(KERNEL_IO_PLUGINSPATH));
-	KernelIOPath = FPaths::ConvertRelativePathToFull(KernelIOPath);
-
 	while (TOptional<FTask> Task = GetNextTask())
 	{
 		FString FullPath = Task->FileName;
 
-		CADLibrary::FCoreTechFileParser FileParser(FullPath, ProcessCacheFolder, ImportParameters, *KernelIOPath);
-		ETaskState ProcessResult = FileParser.ProcessFile();
+		CADLibrary::FCoreTechFileParser FileParser(ImportParameters, *FPaths::EnginePluginsDir(), ProcessCacheFolder);
+		ETaskState ProcessResult = FileParser.ProcessFile(FullPath);
 
 		ETaskState TaskState = ProcessResult;
 		SetTaskState(Task->Index, TaskState);
diff --git a/Engine/Plugins/Enterprise/DatasmithCADImporter/Source/DatasmithDispatcher/Private/DatasmithWorkerHandler.cpp b/Engine/Plugins/Enterprise/DatasmithCADImporter/Source/DatasmithDispatcher/Private/DatasmithWorkerHandler.cpp
--- a/Engine/Plugins/Enterprise/DatasmithCADImporter/Source/DatasmithDispatcher/Private/DatasmithWorkerHandler.cpp
+++ b/Engine/Plugins/Enterprise/DatasmithCADImporter/Source/DatasmithDispatcher/Private/DatasmithWorkerHandler.cpp
@@ -57,12 +57,6 @@ FDatasmithWorkerHandler::~FDatasmithWorkerHandler()
 
 void FDatasmithWorkerHandler::StartWorkerProcess()
 {
-	FString KernelIOPath;
-#ifdef CAD_INTERFACE
-	KernelIOPath = FPaths::Combine(FPaths::EnginePluginsDir(), TEXT(KERNEL_IO_PLUGINSPATH));
-	KernelIOPath = FPaths::ConvertRelativePathToFull(KernelIOPath);
-#endif
-
 	ensure(ErrorState == EWorkerErrorState::Ok);
 	FString ProcessorPath = GetWorkerExecutablePath();
 	if (FPaths::FileExists(ProcessorPath))
@@ -78,7 +72,7 @@ void FDatasmithWorkerHandler::StartWorkerProcess()
 		CommandToProcess += TEXT(" -ServerPID ") + FString::FromInt(FPlatformProcess::GetCurrentProcessId());
 		CommandToProcess += TEXT(" -ServerPort ") + FString::FromInt(ListenPort);
 		CommandToProcess += TEXT(" -CacheDir \"") + CachePath + TEXT('"');
-		CommandToProcess += TEXT(" -KernelIOPath \"") + KernelIOPath + TEXT('"');
+		CommandToProcess += TEXT(" -EnginePluginsDir \"") + FPaths::EnginePluginsDir() + TEXT('"');
 		UE_LOG(LogDatasmithDispatcher, Verbose, TEXT("CommandToProcess: %s"), *CommandToProcess);
 
 		uint32 WorkerProcessId = 0;
diff --git a/Engine/Plugins/Enterprise/DatasmithCADImporter/Source/DatasmithOpenNurbsTranslator/Private/DatasmithOpenNurbsTranslator.cpp b/Engine/Plugins/Enterprise/DatasmithCADImporter/Source/DatasmithOpenNurbsTranslator/Private/DatasmithOpenNurbsTranslator.cpp
--- a/Engine/Plugins/Enterprise/DatasmithCADImporter/Source/DatasmithOpenNurbsTranslator/Private/DatasmithOpenNurbsTranslator.cpp
+++ b/Engine/Plugins/Enterprise/DatasmithCADImporter/Source/DatasmithOpenNurbsTranslator/Private/DatasmithOpenNurbsTranslator.cpp
@@ -10,10 +10,10 @@
 #include "DatasmithMaterialElements.h"
 #include "DatasmithMaterialsUtils.h"
 #include "DatasmithMesh.h"
-#include "DatasmithMeshHelper.h"
 #include "DatasmithSceneFactory.h"
 #include "DatasmithSceneSource.h"
 #include "DatasmithUtils.h"
+#include "Utility/DatasmithMeshHelper.h"
 
 #include "StaticMeshAttributes.h"
 #include "StaticMeshOperations.h"
diff --git a/Engine/Plugins/Enterprise/DatasmithCADImporter/Source/DatasmithOpenNurbsTranslator/Private/DatasmithOpenNurbsTranslator.h b/Engine/Plugins/Enterprise/DatasmithCADImporter/Source/DatasmithOpenNurbsTranslator/Private/DatasmithOpenNurbsTranslator.h
--- a/Engine/Plugins/Enterprise/DatasmithCADImporter/Source/DatasmithOpenNurbsTranslator/Private/DatasmithOpenNurbsTranslator.h
+++ b/Engine/Plugins/Enterprise/DatasmithCADImporter/Source/DatasmithOpenNurbsTranslator/Private/DatasmithOpenNurbsTranslator.h
@@ -15,7 +15,6 @@ class FDatasmithOpenNurbsTranslator : public FDatasmithCoreTechTranslator
 public:
 	FDatasmithOpenNurbsTranslator();
 
-	// Begin IDatasmithTranslator overrides
 	virtual FName GetFName() const override { return "DatasmithOpenNurbsTranslator"; };
 
 	virtual void Initialize(FDatasmithTranslatorCapabilities& OutCapabilities) override;
@@ -26,9 +25,7 @@ public:
 	virtual void UnloadScene() override;
 
 	virtual bool LoadStaticMesh(const TSharedRef<IDatasmithMeshElement> MeshElement, FDatasmithMeshElementPayload& OutMeshPayload) override;
-	// End IDatasmithTranslator overrides
 
-	// Begin ADatasmithCoreTechTranslator overrides
 	virtual void SetSceneImportOptions(TArray<TStrongObjectPtr<UObject>>& Options) override;
 
 protected:
@@ -36,7 +33,6 @@ protected:
 	{
 		TessellationOptions.StitchingTechnique = EDatasmithCADStitchingTechnique::StitchingNone;
 	}
-	// End ADatasmithCoreTechTranslator overrides
 
 private:
 
diff --git a/Engine/Plugins/Enterprise/DatasmithCADImporter/Source/DatasmithWireTranslator/Private/DatasmithWireTranslator.cpp b/Engine/Plugins/Enterprise/DatasmithCADImporter/Source/DatasmithWireTranslator/Private/DatasmithWireTranslator.cpp
--- a/Engine/Plugins/Enterprise/DatasmithCADImporter/Source/DatasmithWireTranslator/Private/DatasmithWireTranslator.cpp
+++ b/Engine/Plugins/Enterprise/DatasmithCADImporter/Source/DatasmithWireTranslator/Private/DatasmithWireTranslator.cpp
@@ -4,15 +4,15 @@
 
 #include "Containers/List.h"
 #include "DatasmithImportOptions.h"
-#include "DatasmithMeshHelper.h"
 #include "DatasmithSceneFactory.h"
+#include "DatasmithTranslator.h"
 #include "DatasmithUtils.h"
 #include "DatasmithWireTranslatorModule.h"
 #include "IDatasmithSceneElements.h"
 #include "Misc/FileHelper.h"
 #include "Misc/Paths.h"
 #include "OpenModelUtils.h"
-#include "Translators\DatasmithTranslator.h"
+#include "Utility/DatasmithMeshHelper.h"
 
 #include "StaticMeshDescription.h"
 #include "StaticMeshOperations.h"
diff --git a/Engine/Plugins/Enterprise/DatasmithCADImporter/Source/DatasmithWireTranslator/Private/OpenModelUtils.cpp b/Engine/Plugins/Enterprise/DatasmithCADImporter/Source/DatasmithWireTranslator/Private/OpenModelUtils.cpp
--- a/Engine/Plugins/Enterprise/DatasmithCADImporter/Source/DatasmithWireTranslator/Private/OpenModelUtils.cpp
+++ b/Engine/Plugins/Enterprise/DatasmithCADImporter/Source/DatasmithWireTranslator/Private/OpenModelUtils.cpp
@@ -5,7 +5,7 @@
 
 #include "CADOptions.h"
 #include "DatasmithUtils.h"
-#include "Translators/DatasmithTranslator.h"
+#include "DatasmithTranslator.h"
 #include "IDatasmithSceneElements.h"
 
 #include "AlDagNode.h"
@@ -1578,7 +1578,7 @@ void OpenModelUtils::SetActorTransform(const TSharedPtr< IDatasmithActorElement
 		FTransform LocalUETransform = FDatasmithUtils::ConvertTransform(FDatasmithUtils::EModelCoordSystem::ZUp_RightHanded, LocalTransform);
 		FQuat Quat;
 		
-		FDatasmithTransformUtils::GetRotation(LocalUETransform, Quat);
+		FDatasmithTransformUtils::GetSafeRotation(LocalUETransform, Quat);
 
 		ActorElement->SetTranslation(LocalUETransform.GetTranslation());
 		ActorElement->SetScale(LocalUETransform.GetScale3D());
diff --git a/Engine/Plugins/Enterprise/DatasmithFBXImporter/Source/DatasmithDeltaGenTranslator/Private/DatasmithDeltaGenTranslator.h b/Engine/Plugins/Enterprise/DatasmithFBXImporter/Source/DatasmithDeltaGenTranslator/Private/DatasmithDeltaGenTranslator.h
--- a/Engine/Plugins/Enterprise/DatasmithFBXImporter/Source/DatasmithDeltaGenTranslator/Private/DatasmithDeltaGenTranslator.h
+++ b/Engine/Plugins/Enterprise/DatasmithFBXImporter/Source/DatasmithDeltaGenTranslator/Private/DatasmithDeltaGenTranslator.h
@@ -2,7 +2,7 @@
 
 #pragma once
 
-#include "Translators/DatasmithTranslator.h"
+#include "DatasmithTranslator.h"
 #include "DatasmithDeltaGenImporter.h"
 #include "DatasmithDeltaGenImportOptions.h"
 #include "CoreMinimal.h"
@@ -12,7 +12,6 @@ class FDatasmithDeltaGenTranslator : public IDatasmithTranslator
 public:
 	virtual FName GetFName() const override { return "DatasmithDeltaGenTranslator"; };
 
-	// IDatasmithTranslator interface
 	virtual void Initialize(FDatasmithTranslatorCapabilities& OutCapabilities) override;
 	virtual bool IsSourceSupported(const FDatasmithSceneSource& Source);
 
@@ -24,7 +23,6 @@ public:
 
 	virtual void GetSceneImportOptions(TArray<TStrongObjectPtr<UObject>>& Options) override;
 	virtual void SetSceneImportOptions(TArray<TStrongObjectPtr<UObject>>& Options) override;
-	//~ End IDatasmithTranslator interface
 
 private:
     TStrongObjectPtr<UDatasmithDeltaGenImportOptions> ImportOptions;
diff --git a/Engine/Plugins/Enterprise/DatasmithFBXImporter/Source/DatasmithDeltaGenTranslator/Private/DatasmithDeltaGenTranslatorModule.cpp b/Engine/Plugins/Enterprise/DatasmithFBXImporter/Source/DatasmithDeltaGenTranslator/Private/DatasmithDeltaGenTranslatorModule.cpp
--- a/Engine/Plugins/Enterprise/DatasmithFBXImporter/Source/DatasmithDeltaGenTranslator/Private/DatasmithDeltaGenTranslatorModule.cpp
+++ b/Engine/Plugins/Enterprise/DatasmithFBXImporter/Source/DatasmithDeltaGenTranslator/Private/DatasmithDeltaGenTranslatorModule.cpp
@@ -5,9 +5,8 @@
 #include "MasterMaterials/DatasmithMasterMaterialManager.h"
 #include "DatasmithDeltaGenImporterMaterialSelector.h"
 
-#include "Translators/DatasmithTranslator.h"
-#include "DatasmithImporterModule.h"
 #include "CoreMinimal.h"
+#include "DatasmithTranslator.h"
 
 class FDeltaGenTranslatorModule : public IDatasmithDeltaGenTranslatorModule
 {
@@ -15,7 +14,7 @@ public:
 	virtual void StartupModule() override
 	{
 		// Make sure the DatasmithImporter module exists and has been initialized before adding FDatasmithDeltaGenTranslator's material selector
-		FModuleManager::Get().LoadModule(TEXT("DatasmithImporter"));
+		FModuleManager::Get().LoadModule(TEXT("DatasmithTranslator"));
 
 		FDatasmithMasterMaterialManager::Get().RegisterSelector(TEXT("Deltagen"), MakeShared< FDatasmithDeltaGenImporterMaterialSelector >());
 
diff --git a/Engine/Plugins/Enterprise/DatasmithFBXImporter/Source/DatasmithFBXTranslator/DatasmithFBXTranslator.Build.cs b/Engine/Plugins/Enterprise/DatasmithFBXImporter/Source/DatasmithFBXTranslator/DatasmithFBXTranslator.Build.cs
--- a/Engine/Plugins/Enterprise/DatasmithFBXImporter/Source/DatasmithFBXTranslator/DatasmithFBXTranslator.Build.cs
+++ b/Engine/Plugins/Enterprise/DatasmithFBXImporter/Source/DatasmithFBXTranslator/DatasmithFBXTranslator.Build.cs
@@ -28,7 +28,7 @@ namespace UnrealBuildTool.Rules
 				new string[]
 				{
 					"DatasmithContent",
-					"DatasmithImporter"
+					"DatasmithTranslator"
 				}
 			);
 		}
diff --git a/Engine/Plugins/Enterprise/DatasmithFBXImporter/Source/DatasmithFBXTranslator/Private/DatasmithFBXImporter.cpp b/Engine/Plugins/Enterprise/DatasmithFBXImporter/Source/DatasmithFBXTranslator/Private/DatasmithFBXImporter.cpp
--- a/Engine/Plugins/Enterprise/DatasmithFBXImporter/Source/DatasmithFBXTranslator/Private/DatasmithFBXImporter.cpp
+++ b/Engine/Plugins/Enterprise/DatasmithFBXImporter/Source/DatasmithFBXTranslator/Private/DatasmithFBXImporter.cpp
@@ -5,17 +5,14 @@
 #include "DatasmithFBXImportOptions.h"
 #include "DatasmithFBXScene.h"
 #include "DatasmithFBXSceneProcessor.h"
-#include "DatasmithImportContext.h"
 #include "DatasmithImportedSequencesActor.h"
-#include "DatasmithLightImporter.h"
-#include "DatasmithMaterialExpressions.h"
-#include "DatasmithMeshHelper.h"
 #include "DatasmithScene.h"
 #include "DatasmithSceneActor.h"
 #include "DatasmithSceneFactory.h"
 #include "DatasmithUtils.h"
 #include "IDatasmithSceneElements.h"
 #include "ObjectTemplates/DatasmithStaticMeshTemplate.h"
+#include "Utility/DatasmithMeshHelper.h"
 
 #include "AssetRegistryModule.h"
 #include "Editor.h"
diff --git a/Engine/Plugins/Enterprise/DatasmithFBXImporter/Source/DatasmithVREDTranslator/Private/DatasmithVREDTranslator.h b/Engine/Plugins/Enterprise/DatasmithFBXImporter/Source/DatasmithVREDTranslator/Private/DatasmithVREDTranslator.h
--- a/Engine/Plugins/Enterprise/DatasmithFBXImporter/Source/DatasmithVREDTranslator/Private/DatasmithVREDTranslator.h
+++ b/Engine/Plugins/Enterprise/DatasmithFBXImporter/Source/DatasmithVREDTranslator/Private/DatasmithVREDTranslator.h
@@ -5,14 +5,13 @@
 #include "CoreMinimal.h"
 #include "DatasmithVREDImporter.h"
 #include "DatasmithVREDImportOptions.h"
-#include "Translators/DatasmithTranslator.h"
+#include "DatasmithTranslator.h"
 
 class FDatasmithVREDTranslator : public IDatasmithTranslator
 {
 public:
 	virtual FName GetFName() const override { return "DatasmithVREDTranslator"; };
 
-	// IDatasmithTranslator interface
 	virtual void Initialize(FDatasmithTranslatorCapabilities& OutCapabilities) override;
 	virtual bool IsSourceSupported(const FDatasmithSceneSource& Source) override;
 
@@ -24,7 +23,6 @@ public:
 
 	virtual void GetSceneImportOptions(TArray<TStrongObjectPtr<UObject>>& Options) override;
 	virtual void SetSceneImportOptions(TArray<TStrongObjectPtr<UObject>>& Options) override;
-	//~ End IDatasmithTranslator interface
 
 private:
     TStrongObjectPtr<UDatasmithVREDImportOptions> ImportOptions;
diff --git a/Engine/Plugins/Enterprise/DatasmithFBXImporter/Source/DatasmithVREDTranslator/Private/DatasmithVREDTranslatorModule.cpp b/Engine/Plugins/Enterprise/DatasmithFBXImporter/Source/DatasmithVREDTranslator/Private/DatasmithVREDTranslatorModule.cpp
--- a/Engine/Plugins/Enterprise/DatasmithFBXImporter/Source/DatasmithVREDTranslator/Private/DatasmithVREDTranslatorModule.cpp
+++ b/Engine/Plugins/Enterprise/DatasmithFBXImporter/Source/DatasmithVREDTranslator/Private/DatasmithVREDTranslatorModule.cpp
@@ -7,8 +7,7 @@
 #include "MasterMaterials/DatasmithMasterMaterialManager.h"
 
 #include "CoreMinimal.h"
-#include "DatasmithImporterModule.h"
-#include "Translators/DatasmithTranslator.h"
+#include "DatasmithTranslator.h"
 
 class FVREDTranslatorModule : public IDatasmithVREDTranslatorModule
 {
diff --git a/Engine/Plugins/Enterprise/DatasmithGLTFImporter/Source/DatasmithGLTFTranslator/Private/DatasmithGLTFImporter.cpp b/Engine/Plugins/Enterprise/DatasmithGLTFImporter/Source/DatasmithGLTFTranslator/Private/DatasmithGLTFImporter.cpp
--- a/Engine/Plugins/Enterprise/DatasmithGLTFImporter/Source/DatasmithGLTFTranslator/Private/DatasmithGLTFImporter.cpp
+++ b/Engine/Plugins/Enterprise/DatasmithGLTFImporter/Source/DatasmithGLTFTranslator/Private/DatasmithGLTFImporter.cpp
@@ -15,10 +15,10 @@
 
 #include "DatasmithAssetImportData.h"
 #include "DatasmithSceneFactory.h"
+#include "DatasmithUtils.h"
 #include "IDatasmithSceneElements.h"
 #include "ObjectTemplates/DatasmithStaticMeshTemplate.h"
-#include "Utility/DatasmithImporterUtils.h"
-#include "DatasmithMeshHelper.h"
+#include "Utility/DatasmithMeshHelper.h"
 
 #include "AssetRegistryModule.h"
 #include "Engine/StaticMesh.h"
diff --git a/Engine/Plugins/Enterprise/DatasmithGLTFImporter/Source/DatasmithGLTFTranslator/Private/DatasmithGLTFTranslator.h b/Engine/Plugins/Enterprise/DatasmithGLTFImporter/Source/DatasmithGLTFTranslator/Private/DatasmithGLTFTranslator.h
--- a/Engine/Plugins/Enterprise/DatasmithGLTFImporter/Source/DatasmithGLTFTranslator/Private/DatasmithGLTFTranslator.h
+++ b/Engine/Plugins/Enterprise/DatasmithGLTFImporter/Source/DatasmithGLTFTranslator/Private/DatasmithGLTFTranslator.h
@@ -2,10 +2,11 @@
 
 #pragma once
 
-#include "Translators/DatasmithTranslator.h"
 #include "DatasmithGLTFImporter.h"
 #include "DatasmithGLTFImportOptions.h"
+
 #include "CoreMinimal.h"
+#include "DatasmithTranslator.h"
 
 class FDatasmithGLTFTranslator : public IDatasmithTranslator
 {
diff --git a/Engine/Plugins/Enterprise/DatasmithGLTFImporter/Source/DatasmithGLTFTranslator/Private/DatasmithGLTFTranslatorModule.cpp b/Engine/Plugins/Enterprise/DatasmithGLTFImporter/Source/DatasmithGLTFTranslator/Private/DatasmithGLTFTranslatorModule.cpp
--- a/Engine/Plugins/Enterprise/DatasmithGLTFImporter/Source/DatasmithGLTFTranslator/Private/DatasmithGLTFTranslatorModule.cpp
+++ b/Engine/Plugins/Enterprise/DatasmithGLTFImporter/Source/DatasmithGLTFTranslator/Private/DatasmithGLTFTranslatorModule.cpp
@@ -3,9 +3,9 @@
 #include "DatasmithGLTFTranslatorModule.h"
 #include "DatasmithGLTFTranslator.h"
 
-#include "Translators/DatasmithTranslator.h"
-#include "DatasmithImporterModule.h"
 #include "CoreMinimal.h"
+#include "DatasmithTranslator.h"
+#include "DatasmithTranslatorModule.h"
 
 const TCHAR* IDatasmithGLTFTranslatorModule::ModuleName = TEXT("DatasmithGLTFTranslator");
 
@@ -14,7 +14,7 @@ class FGLTFTranslatorModule : public IDatasmithGLTFTranslatorModule
 public:
 	virtual void StartupModule() override
 	{
-		FModuleManager::Get().LoadModule(TEXT("DatasmithImporter"));
+		FModuleManager::Get().LoadModule(TEXT("DatasmithTranslator"));
 		Datasmith::RegisterTranslator<FDatasmithGLTFTranslator>();
 	}
 
diff --git a/Engine/Plugins/Enterprise/DatasmithIFCImporter/Source/DatasmithIFCTranslator/DatasmithIFCTranslator.Build.cs b/Engine/Plugins/Enterprise/DatasmithIFCImporter/Source/DatasmithIFCTranslator/DatasmithIFCTranslator.Build.cs
--- a/Engine/Plugins/Enterprise/DatasmithIFCImporter/Source/DatasmithIFCTranslator/DatasmithIFCTranslator.Build.cs
+++ b/Engine/Plugins/Enterprise/DatasmithIFCImporter/Source/DatasmithIFCTranslator/DatasmithIFCTranslator.Build.cs
@@ -45,8 +45,8 @@ namespace UnrealBuildTool.Rules
                 {
                     "Core",
                     "DatasmithContent",
-					"DatasmithImporter"
-				}
+                    "DatasmithTranslator"
+                }
             );
 
 			if (Target.Platform == UnrealTargetPlatform.Win64)
diff --git a/Engine/Plugins/Enterprise/DatasmithIFCImporter/Source/DatasmithIFCTranslator/Private/DatasmithIFCTranslator.h b/Engine/Plugins/Enterprise/DatasmithIFCImporter/Source/DatasmithIFCTranslator/Private/DatasmithIFCTranslator.h
--- a/Engine/Plugins/Enterprise/DatasmithIFCImporter/Source/DatasmithIFCTranslator/Private/DatasmithIFCTranslator.h
+++ b/Engine/Plugins/Enterprise/DatasmithIFCImporter/Source/DatasmithIFCTranslator/Private/DatasmithIFCTranslator.h
@@ -2,17 +2,17 @@
 
 #pragma once
 
-#include "Translators/DatasmithTranslator.h"
 #include "DatasmithIFCImporter.h"
 #include "DatasmithIFCImportOptions.h"
+
 #include "CoreMinimal.h"
+#include "DatasmithTranslator.h"
 
 class FDatasmithIFCTranslator : public IDatasmithTranslator
 {
 public:
 	virtual FName GetFName() const override { return "DatasmithIFCTranslator"; };
 
-	// IDatasmithTranslator interface
 	virtual void Initialize(FDatasmithTranslatorCapabilities& OutCapabilities) override;
 	virtual bool LoadScene(TSharedRef<IDatasmithScene> OutScene) override;
 	virtual void UnloadScene() override;
@@ -22,7 +22,6 @@ public:
 
 	virtual void GetSceneImportOptions(TArray<TStrongObjectPtr<UObject>>& Options) override;
 	virtual void SetSceneImportOptions(TArray<TStrongObjectPtr<UObject>>& Options) override;
-	//~ End IDatasmithTranslator interface
 
 private:
     TStrongObjectPtr<UDatasmithIFCImportOptions> ImportOptions;
diff --git a/Engine/Plugins/Enterprise/DatasmithIFCImporter/Source/DatasmithIFCTranslator/Private/DatasmithIFCTranslatorModule.cpp b/Engine/Plugins/Enterprise/DatasmithIFCImporter/Source/DatasmithIFCTranslator/Private/DatasmithIFCTranslatorModule.cpp
--- a/Engine/Plugins/Enterprise/DatasmithIFCImporter/Source/DatasmithIFCTranslator/Private/DatasmithIFCTranslatorModule.cpp
+++ b/Engine/Plugins/Enterprise/DatasmithIFCImporter/Source/DatasmithIFCTranslator/Private/DatasmithIFCTranslatorModule.cpp
@@ -3,9 +3,8 @@
 #include "DatasmithIFCTranslatorModule.h"
 #include "DatasmithIFCTranslator.h"
 
-#include "Translators/DatasmithTranslator.h"
-#include "DatasmithImporterModule.h"
 #include "CoreMinimal.h"
+#include "DatasmithTranslator.h"
 #include "Interfaces/IPluginManager.h"
 #include "Misc/Paths.h"
 
@@ -21,29 +20,29 @@ public:
 		check(LibHandle == nullptr);
 
 #ifdef WITH_IFC_ENGINE_LIB
-		const FString IFCEngineDll = TEXT("ifcengine.dll");
 
 		// determine directory paths
 		FString IFCEngineDllPath = FPaths::Combine(IPluginManager::Get().FindPlugin(TEXT("DatasmithIFCImporter"))->GetBaseDir(), TEXT("Binaries"), FPlatformProcess::GetBinariesSubdirectory());
 		FPlatformProcess::PushDllDirectory(*IFCEngineDllPath);
-		IFCEngineDllPath = FPaths::Combine(IFCEngineDllPath, IFCEngineDll);
+		FString IFCEngineDll = FPaths::Combine(IFCEngineDllPath, TEXT("ifcengine.dll"));
 
-		if (!FPaths::FileExists(IFCEngineDllPath))
+		if (!FPaths::FileExists(IFCEngineDll))
 		{
 			UE_LOG(LogIFCTranslator, Error, TEXT("Failed to find the binary folder for the IFCEngine dll. Plug-in will not be functional."));
 			return;
 		}
 
-		LibHandle = FPlatformProcess::GetDllHandle(*IFCEngineDllPath);
+		LibHandle = FPlatformProcess::GetDllHandle(*IFCEngineDll);
+		FPlatformProcess::PopDllDirectory(*IFCEngineDllPath);
 
 		if (LibHandle == nullptr)
 		{
-			UE_LOG(LogIFCTranslator, Error, TEXT("Failed to load required library %s. Plug-in will not be functional."), *IFCEngineDllPath);
+			UE_LOG(LogIFCTranslator, Error, TEXT("Failed to load required library %s. Plug-in will not be functional."), *IFCEngineDll);
 			return;
 		}
 #endif
 
-		FModuleManager::Get().LoadModule(TEXT("DatasmithImporter"));
+		FModuleManager::Get().LoadModule(TEXT("DatasmithTranslator"));
 		Datasmith::RegisterTranslator<FDatasmithIFCTranslator>();
 	}
 
diff --git a/Engine/Plugins/Enterprise/DatasmithImporter/DatasmithImporter.uplugin b/Engine/Plugins/Enterprise/DatasmithImporter/DatasmithImporter.uplugin
--- a/Engine/Plugins/Enterprise/DatasmithImporter/DatasmithImporter.uplugin
+++ b/Engine/Plugins/Enterprise/DatasmithImporter/DatasmithImporter.uplugin
@@ -20,6 +20,11 @@
 			"Type": "Editor",
 			"LoadingPhase": "Default"
 		},
+		{
+			"Name": "DatasmithTranslator",
+			"Type": "Editor",
+			"LoadingPhase": "Default"
+		},
 		{
 			"Name": "DatasmithNativeTranslator",
 			"Type": "Editor",
diff --git a/Engine/Plugins/Enterprise/DatasmithImporter/Source/DatasmithImporter/DatasmithImporter.Build.cs b/Engine/Plugins/Enterprise/DatasmithImporter/Source/DatasmithImporter/DatasmithImporter.Build.cs
--- a/Engine/Plugins/Enterprise/DatasmithImporter/Source/DatasmithImporter/DatasmithImporter.Build.cs
+++ b/Engine/Plugins/Enterprise/DatasmithImporter/Source/DatasmithImporter/DatasmithImporter.Build.cs
@@ -55,6 +55,7 @@ namespace UnrealBuildTool.Rules
 					"DataprepCore",
 					"DatasmithContent",
                     "DatasmithCore",
+					"DatasmithTranslator",
                     "DatasmithContentEditor",
 				}
 			);
