Fix SolidWorks import multiconfiguration

diff --git a/Engine/Plugins/Enterprise/DatasmithCADImporter/Source/CADInterfaces/Private/CoreTechFileParser.cpp b/Engine/Plugins/Enterprise/DatasmithCADImporter/Source/CADInterfaces/Private/CoreTechFileParser.cpp
--- a/Engine/Plugins/Enterprise/DatasmithCADImporter/Source/CADInterfaces/Private/CoreTechFileParser.cpp
+++ b/Engine/Plugins/Enterprise/DatasmithCADImporter/Source/CADInterfaces/Private/CoreTechFileParser.cpp
@@ -845,12 +845,12 @@ bool FCoreTechFileParser::ReadInstance(CT_OBJECT_ID InstanceNodeId, uint32 Defau
 		CT_COMPONENT_IO::AskExternalDefinition(ReferenceNodeId, ComponentFile, FileType);
 		FString ExternalRefFullPath = ComponentFile.toUnicode();
 
-		MockUpDescription.Instances[Index].ExternalRef = FPaths::GetCleanFilename(ExternalRefFullPath);
 		if(!NodeConfiguration.IsEmpty())
 		{
-			MockUpDescription.Instances[Index].ExternalRef += TEXT("|") + NodeConfiguration;
+			ExternalRefFullPath += TEXT("|") + NodeConfiguration;
 		}
 
+		MockUpDescription.Instances[Index].ExternalRef = FPaths::GetCleanFilename(ExternalRefFullPath);
 		MockUpDescription.ExternalRefSet.Add(ExternalRefFullPath);
 	}
 	else
