Improve stitching technique Sew

diff --git a/Engine/Plugins/Enterprise/DatasmithCADImporter/Source/CADInterfaces/Private/CoreTechFileParser.cpp b/Engine/Plugins/Enterprise/DatasmithCADImporter/Source/CADInterfaces/Private/CoreTechFileParser.cpp
--- a/Engine/Plugins/Enterprise/DatasmithCADImporter/Source/CADInterfaces/Private/CoreTechFileParser.cpp
+++ b/Engine/Plugins/Enterprise/DatasmithCADImporter/Source/CADInterfaces/Private/CoreTechFileParser.cpp
@@ -607,7 +607,6 @@ FCoreTechFileParser::EProcessResult FCoreTechFileParser::ReadFileWithKernelIO()
 		return EProcessResult::ProcessFailed;
 	}
 
-	Repair(MainId, ImportParameters.StitchingTechnique);
 	SetCoreTechTessellationState(ImportParameters);
 
 	MockUpDescription.FullPath = FullPath;
@@ -776,6 +775,44 @@ bool FCoreTechFileParser::ReadUnloadedComponent(CT_OBJECT_ID ComponentId)
 	return true;
 }
 
+
+void GetInstancesAndBodies(CT_OBJECT_ID InComponentId, TArray<CT_OBJECT_ID>& OutInstances, TArray<CT_OBJECT_ID>& OutBodies)
+{
+	CT_LIST_IO Children;
+	CT_COMPONENT_IO::AskChildren(InComponentId, Children);
+
+	int32 NbChildren = Children.Count();
+	OutInstances.Empty(NbChildren);
+	OutBodies.Empty(NbChildren);
+
+	Children.IteratorInitialize();
+	CT_OBJECT_ID ChildId;
+	while ((ChildId = Children.IteratorIter()) != 0)
+	{
+		CT_OBJECT_TYPE Type;
+		CT_OBJECT_IO::AskType(ChildId, Type);
+
+		switch (Type)
+		{
+		case CT_INSTANCE_TYPE:
+		{
+			OutInstances.Add(ChildId);
+			break;
+		}
+
+		case CT_BODY_TYPE:
+		{
+			OutBodies.Add(ChildId);
+			break;
+		}
+
+		// we don't manage CURVE, POINT, and COORDSYSTEM (the other kind of child of the component).
+		default:
+			break;
+		}
+	}
+}
+
 bool FCoreTechFileParser::ReadComponent(CT_OBJECT_ID ComponentId, uint32 DefaultMaterialHash)
 {
 	int32 Index = MockUpDescription.ComponentSet.Emplace(ComponentId);
@@ -787,17 +824,29 @@ bool FCoreTechFileParser::ReadComponent(CT_OBJECT_ID ComponentId, uint32 Default
 		DefaultMaterialHash = MaterialHash;
 	}
 
-	CT_LIST_IO Children;
-	CT_COMPONENT_IO::AskChildren(ComponentId, Children);
 
-	// Parse children
-	Children.IteratorInitialize();
-	CT_OBJECT_ID ChildId;
-	while ((ChildId = Children.IteratorIter()) != 0)
+	TArray<CT_OBJECT_ID> Instances, Bodies;
+	GetInstancesAndBodies(ComponentId, Instances, Bodies);
+
+	if (!Instances.Num() && (Bodies.Num() > 1) && ImportParameters.StitchingTechnique == StitchingSew)
+	{
+		Repair(ComponentId, StitchingSew);
+		GetInstancesAndBodies(ComponentId, Instances, Bodies);
+	}
+
+	for (CT_OBJECT_ID InstanceId : Instances)
+	{
+		if (ReadInstance(InstanceId, DefaultMaterialHash))
+		{
+			MockUpDescription.ComponentSet[Index].Children.Add(InstanceId);
+		}
+	}
+
+	for (CT_OBJECT_ID BodyId : Bodies)
 	{
-		if (ReadNode(ChildId, DefaultMaterialHash))
+		if (ReadBody(BodyId, DefaultMaterialHash))
 		{
-			MockUpDescription.ComponentSet[Index].Children.Add(ChildId);
+			MockUpDescription.ComponentSet[Index].Children.Add(BodyId);
 		}
 	}
 
