Update version number

diff --git a/Engine/Plugins/Enterprise/DatasmithCADImporter/Source/DatasmithCADTranslator/Private/DatasmithCADTranslatorModule.cpp b/Engine/Plugins/Enterprise/DatasmithCADImporter/Source/DatasmithCADTranslator/Private/DatasmithCADTranslatorModule.cpp
--- a/Engine/Plugins/Enterprise/DatasmithCADImporter/Source/DatasmithCADTranslator/Private/DatasmithCADTranslatorModule.cpp
+++ b/Engine/Plugins/Enterprise/DatasmithCADImporter/Source/DatasmithCADTranslator/Private/DatasmithCADTranslatorModule.cpp
@@ -13,13 +13,16 @@ IMPLEMENT_MODULE(FDatasmithCADTranslatorModule, DatasmithCADTranslator);
 
 void FDatasmithCADTranslatorModule::StartupModule()
 {
+	const int32 CacheVersion = 2;
+
 	// Create temporary directory which will be used by CoreTech to store tessellation data
-	const TCHAR* OldCacheVersion = TEXT("0");
-	FString OldCacheDir = FPaths::ConvertRelativePathToFull(FPaths::Combine(FPaths::ProjectIntermediateDir(), TEXT("DatasmithCADCache"), OldCacheVersion));
-	IFileManager::Get().DeleteDirectory(*OldCacheDir, true, true);
+	for (int32 Version = 0; Version < CacheVersion; ++Version)
+	{
+		FString OldCacheDir = FPaths::ConvertRelativePathToFull(FPaths::Combine(FPaths::ProjectIntermediateDir(), TEXT("DatasmithCADCache"), *FString::FromInt(Version)));
+		IFileManager::Get().DeleteDirectory(*OldCacheDir, true, true);
+	}
 
-	const TCHAR* CacheVersion = TEXT("1");
-	CacheDir = FPaths::ConvertRelativePathToFull(FPaths::Combine(FPaths::ProjectIntermediateDir(), TEXT("DatasmithCADCache"), CacheVersion));
+	CacheDir = FPaths::ConvertRelativePathToFull(FPaths::Combine(FPaths::ProjectIntermediateDir(), TEXT("DatasmithCADCache"), *FString::FromInt(CacheVersion)));
 	IFileManager::Get().MakeDirectory(*CacheDir);
 
 	Datasmith::RegisterTranslator<FDatasmithCADTranslator>();
