Non-unity build fix.

#rnx

diff --git a/Engine/Plugins/Enterprise/DatasmithCADImporter/Source/DatasmithPLMXMLTranslator/Private/DatasmithPlmXmlImporter.cpp b/Engine/Plugins/Enterprise/DatasmithCADImporter/Source/DatasmithPLMXMLTranslator/Private/DatasmithPlmXmlImporter.cpp
--- a/Engine/Plugins/Enterprise/DatasmithCADImporter/Source/DatasmithPLMXMLTranslator/Private/DatasmithPlmXmlImporter.cpp
+++ b/Engine/Plugins/Enterprise/DatasmithCADImporter/Source/DatasmithPLMXMLTranslator/Private/DatasmithPlmXmlImporter.cpp
@@ -20,6 +20,7 @@
 #include "DatasmithDispatcher.h"
 #include "DatasmithMeshBuilder.h"
 #include "DatasmithSceneGraphBuilder.h"
+#include "DatasmithUtils.h"
 
 DEFINE_LOG_CATEGORY_STATIC(LogDatasmithPlmXmlImport, Log, All);
 
