Merging 11214575 #jira UE-87707

diff --git a/Engine/Plugins/Enterprise/DatasmithCADImporter/Source/CADInterfaces/Private/CoreTechFileParser.cpp b/Engine/Plugins/Enterprise/DatasmithCADImporter/Source/CADInterfaces/Private/CoreTechFileParser.cpp
--- a/Engine/Plugins/Enterprise/DatasmithCADImporter/Source/CADInterfaces/Private/CoreTechFileParser.cpp
+++ b/Engine/Plugins/Enterprise/DatasmithCADImporter/Source/CADInterfaces/Private/CoreTechFileParser.cpp
@@ -1004,7 +1004,6 @@ void FCoreTechFileParser::GetBodyTessellation(CT_OBJECT_ID BodyId, FBodyMesh& Ou
 	CT_BODY_IO::AskVerticesArray(BodyId, VerticesArray.GetData());
 
 	FBox& BBox = OutBodyMesh.BBox;
-	BBox.Init();
 	for (const CT_COORDINATE& Point : VerticesArray)
 	{
 		BBox += FVector((float)Point.xyz[0], (float)Point.xyz[1], (float)Point.xyz[2]);
diff --git a/Engine/Plugins/Enterprise/DatasmithCADImporter/Source/CADTools/Public/CADData.h b/Engine/Plugins/Enterprise/DatasmithCADImporter/Source/CADTools/Public/CADData.h
--- a/Engine/Plugins/Enterprise/DatasmithCADImporter/Source/CADTools/Public/CADData.h
+++ b/Engine/Plugins/Enterprise/DatasmithCADImporter/Source/CADTools/Public/CADData.h
@@ -117,6 +117,7 @@ class CADTOOLS_API FBodyMesh
 public:
 	FBodyMesh(CadId InBodyID = 0) : BodyID(InBodyID)
 	{
+		BBox.Init();
 	}
 
 	friend FArchive& operator<<(FArchive& Ar, FBodyMesh& BodyMesh);
