Merge cl 11248544

diff --git a/Engine/Plugins/Enterprise/DatasmithCADImporter/Source/CADInterfaces/Private/CADInterfacesModule.cpp b/Engine/Plugins/Enterprise/DatasmithCADImporter/Source/CADInterfaces/Private/CADInterfacesModule.cpp
--- a/Engine/Plugins/Enterprise/DatasmithCADImporter/Source/CADInterfaces/Private/CADInterfacesModule.cpp
+++ b/Engine/Plugins/Enterprise/DatasmithCADImporter/Source/CADInterfaces/Private/CADInterfacesModule.cpp
@@ -34,6 +34,7 @@ void FCADInterfacesModule::StartupModule()
 {
 	check(KernelIOLibHandle == nullptr);
 
+#ifdef CAD_INTERFACE
 	FString KernelIODll = TEXT("kernel_io.dll");
 
 	// determine directory paths
@@ -43,7 +44,7 @@ void FCADInterfacesModule::StartupModule()
 
 	if (!FPaths::FileExists(KernelIODll))
 	{
-		UE_LOG(CADInterfaces, Warning, TEXT("Failed to find the binary folder for the CoreTech dll. Plug-in will not be functional."));
+		UE_LOG(CADInterfaces, Error, TEXT("Failed to find the binary folder for the CoreTech dll. Plug-in will not be functional."));
 		return;
 	}
 
@@ -54,6 +55,9 @@ void FCADInterfacesModule::StartupModule()
 	}
 
 	FPlatformProcess::PopDllDirectory(*CADImporterDllPath);
+#else
+	UE_LOG(CADInterfaces, Display, TEXT("Missing CoreTech module. Plug-in will not be functional."));
+#endif
 }
 
 void FCADInterfacesModule::ShutdownModule()
