Merged CL 11205097

PLMXML: Retessellate; and refactor adding retessellate data for other cad files

diff --git a/Engine/Plugins/Enterprise/DatasmithCADImporter/Source/DatasmithCADTranslator/Private/DatasmithCADTranslator.cpp b/Engine/Plugins/Enterprise/DatasmithCADImporter/Source/DatasmithCADTranslator/Private/DatasmithCADTranslator.cpp
--- a/Engine/Plugins/Enterprise/DatasmithCADImporter/Source/DatasmithCADTranslator/Private/DatasmithCADTranslator.cpp
+++ b/Engine/Plugins/Enterprise/DatasmithCADImporter/Source/DatasmithCADTranslator/Private/DatasmithCADTranslator.cpp
@@ -140,29 +140,7 @@ bool FDatasmithCADTranslator::LoadStaticMesh(const TSharedRef<IDatasmithMeshElem
 	{
 		OutMeshPayload.LodMeshes.Add(MoveTemp(Mesh.GetValue()));
 
-		// Store CoreTech additional data if provided
-		const TCHAR* CoretechFile = MeshElement->GetFile();
-		if (FPaths::FileExists(CoretechFile))
-		{
-			TArray<uint8> ByteArray;
-			if (FFileHelper::LoadFileToArray(ByteArray, CoretechFile))
-			{
-				UCoreTechParametricSurfaceData* CoreTechData = Datasmith::MakeAdditionalData<UCoreTechParametricSurfaceData>();
-				CoreTechData->SourceFile = CoretechFile;
-				CoreTechData->RawData = MoveTemp(ByteArray);
-				CoreTechData->SceneParameters.ModelCoordSys = uint8(ImportParameters.ModelCoordSys);
-				CoreTechData->SceneParameters.MetricUnit = ImportParameters.MetricUnit;
-				CoreTechData->SceneParameters.ScaleFactor = ImportParameters.ScaleFactor;
-
-				CoreTechData->MeshParameters.bNeedSwapOrientation = MeshParameters.bNeedSwapOrientation;
-				CoreTechData->MeshParameters.bIsSymmetric = MeshParameters.bIsSymmetric;
-				CoreTechData->MeshParameters.SymmetricNormal = MeshParameters.SymmetricNormal;
-				CoreTechData->MeshParameters.SymmetricOrigin = MeshParameters.SymmetricOrigin;
-
-				CoreTechData->LastTessellationOptions = GetCommonTessellationOptions();
-				OutMeshPayload.AdditionalData.Add(CoreTechData);
-			}
-		}
+		DatasmithCoreTechParametricSurfaceData::AddCoreTechSurfaceDataForMesh(MeshElement, ImportParameters, MeshParameters, GetCommonTessellationOptions(), OutMeshPayload);
 	}
 	return OutMeshPayload.LodMeshes.Num() > 0;
 }
diff --git a/Engine/Plugins/Enterprise/DatasmithCADImporter/Source/DatasmithCoreTechParametricSurfaceData/Private/CoreTechParametricSurfaceExtension.cpp b/Engine/Plugins/Enterprise/DatasmithCADImporter/Source/DatasmithCoreTechParametricSurfaceData/Private/CoreTechParametricSurfaceExtension.cpp
--- a/Engine/Plugins/Enterprise/DatasmithCADImporter/Source/DatasmithCoreTechParametricSurfaceData/Private/CoreTechParametricSurfaceExtension.cpp
+++ b/Engine/Plugins/Enterprise/DatasmithCADImporter/Source/DatasmithCoreTechParametricSurfaceData/Private/CoreTechParametricSurfaceExtension.cpp
@@ -1,8 +1,16 @@
 // Copyright Epic Games, Inc. All Rights Reserved.
 
 #include "CoreTechParametricSurfaceExtension.h"
+#include "DatasmithPayload.h"
+#include "IDatasmithSceneElements.h"
+#include "Misc/FileHelper.h"
+#include "Misc/Paths.h"
 #include "UObject/EnterpriseObjectVersion.h"
 
+#ifdef CAD_LIBRARY
+#include "CADOptions.h"
+#endif // CAD_LIBRARY
+
 void UCoreTechParametricSurfaceData::Serialize(FArchive& Ar)
 {
 	Ar.UsingCustomVersion(FEnterpriseObjectVersion::GUID);
@@ -19,3 +27,37 @@ void UCoreTechParametricSurfaceData::Serialize(FArchive& Ar)
 		RawData = MoveTemp(RawData_DEPRECATED);
 	}
 }
+
+namespace DatasmithCoreTechParametricSurfaceData
+{
+	// TODO: convert to FCoreTechSceneParameters/FCoreTechMeshParameters ?
+	void AddCoreTechSurfaceDataForMesh(const TSharedRef<IDatasmithMeshElement>& InMeshElement, const CADLibrary::FImportParameters& InSceneParameters, const CADLibrary::FMeshParameters& InMeshParameters, const FDatasmithTessellationOptions& InTessellationOptions, FDatasmithMeshElementPayload& OutMeshPayload)
+	{
+#ifdef CAD_LIBRARY
+		// Store CoreTech additional data if provided
+		const TCHAR* CoretechFile = InMeshElement->GetFile();
+		if (FPaths::FileExists(CoretechFile))
+		{
+			TArray<uint8> ByteArray;
+			if (FFileHelper::LoadFileToArray(ByteArray, CoretechFile))
+			{
+				UCoreTechParametricSurfaceData* CoreTechData = Datasmith::MakeAdditionalData<UCoreTechParametricSurfaceData>();
+				CoreTechData->SourceFile = CoretechFile;
+				CoreTechData->RawData = MoveTemp(ByteArray);
+				CoreTechData->SceneParameters.ModelCoordSys = uint8(InSceneParameters.ModelCoordSys);
+				CoreTechData->SceneParameters.MetricUnit = InSceneParameters.MetricUnit;
+				CoreTechData->SceneParameters.ScaleFactor = InSceneParameters.ScaleFactor;
+
+				CoreTechData->MeshParameters.bNeedSwapOrientation = InMeshParameters.bNeedSwapOrientation;
+				CoreTechData->MeshParameters.bIsSymmetric = InMeshParameters.bIsSymmetric;
+				CoreTechData->MeshParameters.SymmetricNormal = InMeshParameters.SymmetricNormal;
+				CoreTechData->MeshParameters.SymmetricOrigin = InMeshParameters.SymmetricOrigin;
+
+				CoreTechData->LastTessellationOptions = InTessellationOptions;
+				OutMeshPayload.AdditionalData.Add(CoreTechData);
+			}
+		}
+#endif // CAD_LIBRARY
+	}
+
+}
diff --git a/Engine/Plugins/Enterprise/DatasmithCADImporter/Source/DatasmithCoreTechParametricSurfaceData/Public/CoreTechParametricSurfaceExtension.h b/Engine/Plugins/Enterprise/DatasmithCADImporter/Source/DatasmithCoreTechParametricSurfaceData/Public/CoreTechParametricSurfaceExtension.h
--- a/Engine/Plugins/Enterprise/DatasmithCADImporter/Source/DatasmithCoreTechParametricSurfaceData/Public/CoreTechParametricSurfaceExtension.h
+++ b/Engine/Plugins/Enterprise/DatasmithCADImporter/Source/DatasmithCoreTechParametricSurfaceData/Public/CoreTechParametricSurfaceExtension.h
@@ -72,4 +72,18 @@ private:
 	TArray<uint8> RawData_DEPRECATED;
 
 	virtual void Serialize(FArchive& Ar) override;
-};
\ No newline at end of file
+};
+
+class IDatasmithMeshElement;
+struct FDatasmithMeshElementPayload;
+
+namespace CADLibrary
+{
+	struct FImportParameters;
+	struct FMeshParameters;
+}
+
+namespace DatasmithCoreTechParametricSurfaceData
+{
+	void DATASMITHCORETECHPARAMETRICSURFACEDATA_API AddCoreTechSurfaceDataForMesh(const TSharedRef<IDatasmithMeshElement>& InMeshElement, const CADLibrary::FImportParameters& InSceneParameters, const CADLibrary::FMeshParameters&, const FDatasmithTessellationOptions& InTessellationOptions, FDatasmithMeshElementPayload& OutMeshPayload);
+}
diff --git a/Engine/Plugins/Enterprise/DatasmithCADImporter/Source/DatasmithOpenNurbsTranslator/Private/DatasmithOpenNurbsTranslator.cpp b/Engine/Plugins/Enterprise/DatasmithCADImporter/Source/DatasmithOpenNurbsTranslator/Private/DatasmithOpenNurbsTranslator.cpp
--- a/Engine/Plugins/Enterprise/DatasmithCADImporter/Source/DatasmithOpenNurbsTranslator/Private/DatasmithOpenNurbsTranslator.cpp
+++ b/Engine/Plugins/Enterprise/DatasmithCADImporter/Source/DatasmithOpenNurbsTranslator/Private/DatasmithOpenNurbsTranslator.cpp
@@ -3203,29 +3203,21 @@ void FDatasmithOpenNurbsTranslator::UnloadScene()
 
 bool FDatasmithOpenNurbsTranslator::LoadStaticMesh(const TSharedRef<IDatasmithMeshElement> MeshElement, FDatasmithMeshElementPayload& OutMeshPayload)
 {
-	if ( TOptional< FMeshDescription > Mesh = Translator->GetMeshDescription( MeshElement ) )
+	if (TOptional< FMeshDescription > Mesh = Translator->GetMeshDescription(MeshElement))
 	{
 		OutMeshPayload.LodMeshes.Add(MoveTemp(Mesh.GetValue()));
 
 #ifdef CAD_LIBRARY
-		// Store CoreTech additional data if provided
-		const TCHAR* CoretechFile = MeshElement->GetFile();
-		if (FPaths::FileExists(CoretechFile))
-		{
-			TArray<uint8> ByteArray;
-			if (FFileHelper::LoadFileToArray(ByteArray, CoretechFile))
-			{
-				UCoreTechParametricSurfaceData* CoreTechData = Datasmith::MakeAdditionalData<UCoreTechParametricSurfaceData>();
-				CoreTechData->SourceFile = CoretechFile;
-				CoreTechData->RawData = MoveTemp(ByteArray);
-				CoreTechData->SceneParameters.ModelCoordSys = uint8(FDatasmithUtils::EModelCoordSystem::ZUp_RightHanded);
-				CoreTechData->SceneParameters.ScaleFactor = Translator->GetScalingFactor();
-				CoreTechData->SceneParameters.MetricUnit = Translator->GetMetricUnit();
-				CoreTechData->LastTessellationOptions = OpenNurbsOptions;
-				OutMeshPayload.AdditionalData.Add(CoreTechData);
-			}
-		}
+		CADLibrary::FImportParameters ImportParameters;
+		ImportParameters.ModelCoordSys = CADLibrary::EModelCoordSystem::ZUp_RightHanded;
+		ImportParameters.MetricUnit = Translator->GetMetricUnit();
+		ImportParameters.ScaleFactor = Translator->GetScalingFactor();
+
+		CADLibrary::FMeshParameters MeshParameters;
+
+		DatasmithCoreTechParametricSurfaceData::AddCoreTechSurfaceDataForMesh(MeshElement, ImportParameters, MeshParameters, OpenNurbsOptions, OutMeshPayload);
 #endif // CAD_LIBRARY
+
 	}
 
 	return OutMeshPayload.LodMeshes.Num() > 0;
diff --git a/Engine/Plugins/Enterprise/DatasmithCADImporter/Source/DatasmithPLMXMLTranslator/DatasmithPLMXMLTranslator.Build.cs b/Engine/Plugins/Enterprise/DatasmithCADImporter/Source/DatasmithPLMXMLTranslator/DatasmithPLMXMLTranslator.Build.cs
--- a/Engine/Plugins/Enterprise/DatasmithCADImporter/Source/DatasmithPLMXMLTranslator/DatasmithPLMXMLTranslator.Build.cs
+++ b/Engine/Plugins/Enterprise/DatasmithCADImporter/Source/DatasmithPLMXMLTranslator/DatasmithPLMXMLTranslator.Build.cs
@@ -27,9 +27,9 @@ namespace UnrealBuildTool.Rules
 
 					"DatasmithCore",
 
-                    "CADInterfaces",
                     "CADLibrary",
 					"CADTools",
+					"DatasmithCoreTechParametricSurfaceData",
 					"DatasmithCADTranslator", // for DatasmithMeshBuilder
                     "DatasmithTranslator",
                     "DatasmithDispatcher",
diff --git a/Engine/Plugins/Enterprise/DatasmithCADImporter/Source/DatasmithPLMXMLTranslator/Private/DatasmithPlmXmlImporter.cpp b/Engine/Plugins/Enterprise/DatasmithCADImporter/Source/DatasmithPLMXMLTranslator/Private/DatasmithPlmXmlImporter.cpp
--- a/Engine/Plugins/Enterprise/DatasmithCADImporter/Source/DatasmithPLMXMLTranslator/Private/DatasmithPlmXmlImporter.cpp
+++ b/Engine/Plugins/Enterprise/DatasmithCADImporter/Source/DatasmithPLMXMLTranslator/Private/DatasmithPlmXmlImporter.cpp
@@ -5,23 +5,30 @@
 #include "DatasmithImportOptions.h"
 #include "DatasmithPayload.h"
 #include "DatasmithSceneFactory.h"
+#include "DatasmithUtils.h"
 #include "IDatasmithSceneElements.h"
 
 #include "HAL/FileManager.h" // IFileManager
+#include "Misc/FileHelper.h"
 #include "Math/Transform.h"
 #include "Misc/Paths.h"
 #include "Misc/SecureHash.h"
-#include "XmlParser.h"
 #include "Modules/ModuleManager.h"
+#include "XmlParser.h"
 
 
+#ifdef CAD_LIBRARY
 #include "CADData.h"
 #include "CADToolsModule.h"
+#include "CoreTechParametricSurfaceExtension.h"
+#endif // CAD_LIBRARY
+#include "DatasmithAdditionalData.h"
 #include "DatasmithDispatcher.h"
 #include "DatasmithMeshBuilder.h"
 #include "DatasmithSceneGraphBuilder.h"
 #include "DatasmithUtils.h"
 
+
 DEFINE_LOG_CATEGORY_STATIC(LogDatasmithPlmXmlImport, Log, All);
 
 class IDatasmithScene;
@@ -215,17 +222,18 @@ namespace PlmXml
 		ActorElement->SetScale(Transform.GetScale3D());
 	}
 
-#ifdef CAD_INTERFACE
+#ifdef CAD_LIBRARY
 	class FPlmXmlMeshLoaderWithDatasmithDispatcher
 	{
 	public:
 		TSharedRef<IDatasmithScene> DatasmithScene;
+		FDatasmithTessellationOptions& TessellationOptions;
+		
 		TUniquePtr<DatasmithDispatcher::FDatasmithDispatcher> DatasmithDispatcher;
 
 		TUniquePtr <FDatasmithSceneGraphBuilder> SceneGraphBuilder;
 		TUniquePtr<FDatasmithMeshBuilder> MeshBuilderPtr;
 
-
 		CADLibrary::FImportParameters ImportParameters;
 		TMap<uint32, FString> CADFileToUE4GeomMap;
 		TMap<uint32, FString> CADFileToUE4FileMap;
@@ -234,8 +242,10 @@ namespace PlmXml
 
 		FString CacheDir;
 
-		FPlmXmlMeshLoaderWithDatasmithDispatcher(TSharedRef<IDatasmithScene> InDatasmithScene, FDatasmithTessellationOptions& TessellationOptions)
+		FPlmXmlMeshLoaderWithDatasmithDispatcher(TSharedRef<IDatasmithScene> InDatasmithScene, FDatasmithTessellationOptions& InTessellationOptions)
 			: DatasmithScene(InDatasmithScene)
+			, TessellationOptions(InTessellationOptions)
+		
 		{
 			FCADToolsModule& CADToolsModule = FCADToolsModule::Get();
 			
@@ -302,6 +312,9 @@ namespace PlmXml
 			if (TOptional<FMeshDescription> Mesh = MeshBuilderPtr->GetMeshDescription(MeshElement, MeshParameters))
 			{
 				OutMeshPayload.LodMeshes.Add(MoveTemp(Mesh.GetValue()));
+
+				DatasmithCoreTechParametricSurfaceData::AddCoreTechSurfaceDataForMesh(MeshElement, ImportParameters, MeshParameters, TessellationOptions, OutMeshPayload);
+				
 				return true;
 			}
 			return false;
@@ -309,7 +322,7 @@ namespace PlmXml
 
 		void SetTessellationOptions(const UDatasmithCommonTessellationOptions& TessellationOptions);
 	};
-#else
+#else // CAD_LIBRARY
 	class FPlmXmlMeshLoaderSimple
 	{
 	public:
@@ -371,7 +384,7 @@ namespace PlmXml
 
 		void SetTessellationOptions(const UDatasmithCommonTessellationOptions& TessellationOptions);
 	};
-#endif
+#endif // CAD_LIBRARY
 	struct FParsedPlmXml
 	{
 		TMap<FString, const FXmlNode*> ExternalFileNodes;
diff --git a/Engine/Plugins/Enterprise/DatasmithCADImporter/Source/DatasmithPLMXMLTranslator/Private/DatasmithPlmXmlImporter.h b/Engine/Plugins/Enterprise/DatasmithCADImporter/Source/DatasmithPLMXMLTranslator/Private/DatasmithPlmXmlImporter.h
--- a/Engine/Plugins/Enterprise/DatasmithCADImporter/Source/DatasmithPLMXMLTranslator/Private/DatasmithPlmXmlImporter.h
+++ b/Engine/Plugins/Enterprise/DatasmithCADImporter/Source/DatasmithPLMXMLTranslator/Private/DatasmithPlmXmlImporter.h
@@ -15,13 +15,13 @@ struct FDatasmithTessellationOptions;
 
 namespace PlmXml
 {
-#ifdef CAD_INTERFACE
+#ifdef CAD_LIBRARY
 	class FPlmXmlMeshLoaderWithDatasmithDispatcher;
 	using FPlmXmlMeshLoader = FPlmXmlMeshLoaderWithDatasmithDispatcher;
-#else
+#else // CAD_LIBRARY
 	class FPlmXmlMeshLoaderSimple;
 	using FPlmXmlMeshLoader = FPlmXmlMeshLoaderSimple;
-#endif
+#endif // CAD_LIBRARY
 }
 
 class FDatasmithPlmXmlImporter : FNoncopyable
diff --git a/Engine/Plugins/Enterprise/DatasmithCADImporter/Source/DatasmithWireTranslator/Private/DatasmithWireTranslator.cpp b/Engine/Plugins/Enterprise/DatasmithCADImporter/Source/DatasmithWireTranslator/Private/DatasmithWireTranslator.cpp
--- a/Engine/Plugins/Enterprise/DatasmithCADImporter/Source/DatasmithWireTranslator/Private/DatasmithWireTranslator.cpp
+++ b/Engine/Plugins/Enterprise/DatasmithCADImporter/Source/DatasmithWireTranslator/Private/DatasmithWireTranslator.cpp
@@ -2163,40 +2163,14 @@ void FDatasmithWireTranslator::UnloadScene()
 bool FDatasmithWireTranslator::LoadStaticMesh(const TSharedRef<IDatasmithMeshElement> MeshElement, FDatasmithMeshElementPayload& OutMeshPayload)
 {
 #ifdef USE_OPENMODEL
-
 	CADLibrary::FImportParameters& ImportParameters = Translator->GetImportParameters();
 	CADLibrary::FMeshParameters MeshParameters;
 	if (TOptional< FMeshDescription > Mesh = Translator->GetMeshDescription(MeshElement, MeshParameters))
 	{
 		OutMeshPayload.LodMeshes.Add(MoveTemp(Mesh.GetValue()));
-
 #ifdef CAD_LIBRARY
-		// Store CoreTech additional data if provided
-		const TCHAR* CoretechFile = MeshElement->GetFile();
-		if (FPaths::FileExists(CoretechFile))
-		{
-			TArray<uint8> ByteArray;
-			if (FFileHelper::LoadFileToArray(ByteArray, CoretechFile))
-			{
-				UCoreTechParametricSurfaceData* CoreTechData = Datasmith::MakeAdditionalData<UCoreTechParametricSurfaceData>();
-				CoreTechData->SourceFile = CoretechFile;
-				CoreTechData->RawData = MoveTemp(ByteArray);
-				CoreTechData->SceneParameters.ModelCoordSys = uint8(FDatasmithUtils::EModelCoordSystem::ZUp_RightHanded);
-				CoreTechData->SceneParameters.MetricUnit = ImportParameters.MetricUnit;
-				CoreTechData->SceneParameters.ScaleFactor = ImportParameters.ScaleFactor;
-
-				CoreTechData->MeshParameters.bNeedSwapOrientation = MeshParameters.bNeedSwapOrientation;
-				CoreTechData->MeshParameters.bIsSymmetric = MeshParameters.bIsSymmetric;
-				CoreTechData->MeshParameters.SymmetricNormal = MeshParameters.SymmetricNormal;
-				CoreTechData->MeshParameters.SymmetricOrigin = MeshParameters.SymmetricOrigin;
-
-				CoreTechData->LastTessellationOptions = GetCommonTessellationOptions();
-
-				OutMeshPayload.AdditionalData.Add(CoreTechData);
-			}
-		}
-#endif
-
+		DatasmithCoreTechParametricSurfaceData::AddCoreTechSurfaceDataForMesh(MeshElement, ImportParameters, MeshParameters, GetCommonTessellationOptions(), OutMeshPayload);
+#endif //CAD_LIBRARY
 	}
 	return OutMeshPayload.LodMeshes.Num() > 0;
 #else
