Merging CL 11278405 from //UE4/Dev-Enterprise to //UE4/Release-4.25

OpenNurbs: pivot set to geometry bbox center

diff --git a/Engine/Plugins/Enterprise/DatasmithCADImporter/Source/DatasmithOpenNurbsTranslator/Private/DatasmithOpenNurbsTranslator.cpp b/Engine/Plugins/Enterprise/DatasmithCADImporter/Source/DatasmithOpenNurbsTranslator/Private/DatasmithOpenNurbsTranslator.cpp
--- a/Engine/Plugins/Enterprise/DatasmithCADImporter/Source/DatasmithOpenNurbsTranslator/Private/DatasmithOpenNurbsTranslator.cpp
+++ b/Engine/Plugins/Enterprise/DatasmithCADImporter/Source/DatasmithOpenNurbsTranslator/Private/DatasmithOpenNurbsTranslator.cpp
@@ -265,7 +265,7 @@ namespace DatasmithOpenNurbsTranslatorUtils
 		FVector2D UVCoords[3];
 	};
 
-	bool TranslateMesh(const ON_Mesh *Mesh, FMeshDescription& MeshDescription, bool& bHasNormal, double ScalingFactor)
+	bool TranslateMesh(const ON_Mesh *Mesh, FMeshDescription& MeshDescription, bool& bHasNormal, double ScalingFactor, const ON_3dVector& Offset)
 	{
 		// Ref. GP3DMVisitorImpl::visitMesh
 		// Ref. FGPureMeshInterface::CreateMesh
@@ -296,7 +296,7 @@ namespace DatasmithOpenNurbsTranslatorUtils
 
 		for (int Index = 0; Index < VertexCount; ++Index)
 		{
-			ON_3fPoint p1 = Mesh->m_V[Index];
+			ON_3fPoint p1 = Mesh->m_V[Index] + ON_3fVector(Offset);
 			Nodes.Add(FNode(p1.x * ScalingFactor, p1.y * ScalingFactor, p1.z * ScalingFactor));
 		}
 
@@ -682,6 +682,22 @@ private:
 
 	bool TranslateBRep(ON_Brep* brep, const ON_3dmObjectAttributes& Attributes, FMeshDescription& OutMesh, const TSharedRef< IDatasmithMeshElement >& MeshElement, const FString& Name, bool& bHasNormal);
 
+	
+	bool ComputeObjectGeometryCenter(const FOpenNurbsObjectWrapper& Object, ON_3dVector& OutGeometryCenter);
+	bool ComputeGeometryCenter(ON_Geometry* Geometry, ON_3dVector& OutCenter);
+
+	// Returns vector to shift geometry to its bounding box center
+	// done before tessellating or creating MeshDescription
+	ON_3dVector GetGeometryOffset(TSharedRef<IDatasmithMeshElement> MeshElement)
+	{
+		ON_3dVector Offset = ON_3fVector::ZeroVector;
+		if (ON_3dVector* OffsetPtr = MeshElementToGeometryCenter.Find(MeshElement))
+		{
+			Offset = *OffsetPtr;
+		}
+		return -Offset;
+	}
+
 private:
 	TArray<FOpenNurbsTranslatorImpl*> ChildTranslators;
 	TSharedRef<IDatasmithScene> Scene;
@@ -769,6 +785,8 @@ private:
 	/** OpenNurbs objects to Datasmith mesh elements */
 	TMap<const FOpenNurbsObjectWrapper* , TSharedPtr< IDatasmithMeshElement > > ObjectToMeshElementMap;
 
+	TMap< TSharedPtr< IDatasmithMeshElement >, ON_3dVector > MeshElementToGeometryCenter;
+
 	TArray<FString> MissingRenderMeshes;
 };
 
@@ -1430,6 +1448,7 @@ void FOpenNurbsTranslatorImpl::TranslateInstanceDefinitionTable(const TArray<ON_
 						// Propagate data from child to parent translator for use by the "root" translator
 						MeshElementToTranslatorMap.Append(LinkedFileTranslator->MeshElementToTranslatorMap);
 						MeshElementToObjectMap.Append(LinkedFileTranslator->MeshElementToObjectMap);
+						MeshElementToGeometryCenter.Append(LinkedFileTranslator->MeshElementToGeometryCenter);
 
 						// Merge ChildScene with ParentScene
 						int32 NumActors = ChildScene->GetActorsCount();
@@ -1487,7 +1506,7 @@ void FOpenNurbsTranslatorImpl::TranslateObjectTable(const ON_ClassArray<FOpenNur
 	for (int Index = 0; Index < NumObjects; ++Index)
 	{
 		const FOpenNurbsObjectWrapper& Object = *InObjectTable.At(Index);
-		if (!Object.ObjectPtr->IsKindOf(&ON_InstanceRef::m_ON_InstanceRef_class_rtti))
+		if (!ON_InstanceRef::Cast(Object.ObjectPtr))
 		{
 			TranslateNonInstanceObject(Object);
 		}
@@ -1775,8 +1794,7 @@ bool FOpenNurbsTranslatorImpl::TranslateInstance(const FOpenNurbsObjectWrapper&
 
 	// Ref. FDatasmithCADImporter::SetWorldTransform
 	FTransform Transform(Matrix);
-	const FTransform RightHanded = FTransform(FRotator(0.0f, 0.0f, 0.0f), FVector(0.0f, 0.0f, 0.0f), FVector(-1.0f, 1.0f, 1.0f));
-	FTransform CorrectedTransform = RightHanded * Transform * RightHanded;
+	FTransform CorrectedTransform = FDatasmithUtils::ConvertTransform(FDatasmithUtils::EModelCoordSystem::ZUp_RightHanded, Transform);
 
 	ContainerElement->SetTranslation(CorrectedTransform.GetTranslation() * ScalingFactor);
 	ContainerElement->SetScale(CorrectedTransform.GetScale3D());
@@ -1886,6 +1904,14 @@ TSharedPtr<IDatasmithMeshActorElement> FOpenNurbsTranslatorImpl::GetMeshActorEle
 	ActorElement->SetLabel(*ActorLabel);
 	ActorElement->SetStaticMeshPathName(MeshElement->GetName());
 
+	ON_3dVector GeometryCenter;
+	if (ComputeObjectGeometryCenter(Object, GeometryCenter))
+	{
+		MeshElementToGeometryCenter.Add(MeshElement.ToSharedRef(), GeometryCenter);
+		FVector ActorOffset = ScalingFactor * FDatasmithUtils::ConvertVector(FDatasmithUtils::EModelCoordSystem::ZUp_RightHanded, GeometryCenter);
+		ActorElement->SetTranslation(ActorOffset);
+	}
+
 	// TODO: TBD if need to set Material Override
 
 	return ActorElement;
@@ -2992,6 +3018,8 @@ bool FOpenNurbsTranslatorImpl::TranslateBRep(ON_Brep* Brep, const ON_3dmObjectAt
 		return false;
 	}
 
+	ON_3dVector Offset = GetGeometryOffset(MeshElement);
+
 	// No tessellation if CAD library is not present...
 #ifdef CAD_LIBRARY
 	if (OpenNurbsOptions.Geometry == EDatasmithOpenNurbsBrepTessellatedSource::UseUnrealNurbsTessellation)
@@ -3004,7 +3032,7 @@ bool FOpenNurbsTranslatorImpl::TranslateBRep(ON_Brep* Brep, const ON_3dmObjectAt
 
 		LocalSession->ClearData();
 
-		Result = LocalSession->AddBRep(*Brep);
+		Result = LocalSession->AddBRep(*Brep, Offset);
 
 		FString Filename = FString::Printf(TEXT("%s.ct"), *Name);
 		FString FilePath = FPaths::Combine(OutputPath, Filename);
@@ -3048,7 +3076,7 @@ bool FOpenNurbsTranslatorImpl::TranslateBRep(ON_Brep* Brep, const ON_3dmObjectAt
 			BRepMesh.Append(AnyMeshes.Count(), AnyMeshes.Array());
 		}
 
-		if (!DatasmithOpenNurbsTranslatorUtils::TranslateMesh(&BRepMesh, OutMesh, bHasNormal, ScalingFactor))
+		if (!DatasmithOpenNurbsTranslatorUtils::TranslateMesh(&BRepMesh, OutMesh, bHasNormal, ScalingFactor, Offset))
 		{
 			return false;
 		}
@@ -3057,7 +3085,73 @@ bool FOpenNurbsTranslatorImpl::TranslateBRep(ON_Brep* Brep, const ON_3dmObjectAt
 	}
 }
 
-TOptional< FMeshDescription > FOpenNurbsTranslatorImpl::GetMeshDescription(TSharedRef< IDatasmithMeshElement > MeshElement)
+bool FOpenNurbsTranslatorImpl::ComputeGeometryCenter(ON_Geometry* Geometry, ON_3dVector& OutCenter)
+{
+	if (Geometry == nullptr)
+	{
+		return false;
+	}
+
+	double BoxMin[3];
+	double BoxMax[3];
+	if (!Geometry->GetBBox(BoxMin, BoxMax))
+	{
+		return false;
+	}
+
+	OutCenter = (ON_3dVector(BoxMin[0], BoxMin[1], BoxMin[2]) + ON_3dVector(BoxMax[0], BoxMax[1], BoxMax[2])) * 0.5f;
+
+	return true;
+}
+
+bool FOpenNurbsTranslatorImpl::ComputeObjectGeometryCenter(const FOpenNurbsObjectWrapper& Object, ON_3dVector& OutGeometryCenter)
+{
+	bool bIsValid = false;
+	ON_3dVector GeometryCenter;
+	if (Object.ObjectPtr->IsKindOf(&ON_Mesh::m_ON_Mesh_class_rtti))
+	{
+		bIsValid = ComputeGeometryCenter(ON_Mesh::Cast(Object.ObjectPtr), GeometryCenter);
+	}
+	else if (Object.ObjectPtr->IsKindOf(&ON_Brep::m_ON_Brep_class_rtti))
+	{
+		bIsValid = ComputeGeometryCenter(ON_Brep::Cast(Object.ObjectPtr), GeometryCenter);
+	}
+	else if (Object.ObjectPtr->IsKindOf(&ON_Extrusion::m_ON_Extrusion_class_rtti))
+	{
+		const ON_Extrusion* extrusion = ON_Extrusion::Cast(Object.ObjectPtr);
+		ON_Brep brep;
+		if (extrusion != nullptr && extrusion->BrepForm(&brep) != nullptr)
+		{
+			bIsValid = ComputeGeometryCenter(&brep, GeometryCenter);
+		}
+	}
+	else if (Object.ObjectPtr->IsKindOf(&ON_Hatch::m_ON_Hatch_class_rtti))
+	{
+		const ON_Hatch* hatch = ON_Hatch::Cast(Object.ObjectPtr);
+		ON_Brep brep;
+		if (hatch != nullptr && hatch->BrepForm(&brep) != nullptr)
+		{
+			bIsValid = ComputeGeometryCenter(&brep, GeometryCenter);
+		}
+	}
+	else if (Object.ObjectPtr->IsKindOf(&ON_PlaneSurface::m_ON_PlaneSurface_class_rtti))
+	{
+		// 
+	}
+	else if (Object.ObjectPtr->IsKindOf(&ON_LineCurve::m_ON_LineCurve_class_rtti))
+	{
+		// Not supported
+	}
+
+	if (bIsValid)
+	{
+		OutGeometryCenter = GeometryCenter;
+		return true;
+	}
+	return false;
+}
+
+TOptional<FMeshDescription> FOpenNurbsTranslatorImpl::GetMeshDescription(TSharedRef<IDatasmithMeshElement> MeshElement)
 {
 	// Ref. visitNonInstanceObject (mesh collection part)
 	const FOpenNurbsObjectWrapper** ObjectPtr = MeshElementToObjectMap.Find(&MeshElement.Get());
@@ -3088,7 +3182,8 @@ TOptional< FMeshDescription > FOpenNurbsTranslatorImpl::GetMeshDescription(TShar
 	bool bIsValid = false;
 	if (Object.ObjectPtr->IsKindOf(&ON_Mesh::m_ON_Mesh_class_rtti))
 	{
-		bIsValid = DatasmithOpenNurbsTranslatorUtils::TranslateMesh(ON_Mesh::Cast(Object.ObjectPtr), MeshDescription, bHasNormal, SelectedTranslator->ScalingFactor);
+		ON_3dVector Offset = GetGeometryOffset(MeshElement);
+		bIsValid = DatasmithOpenNurbsTranslatorUtils::TranslateMesh(ON_Mesh::Cast(Object.ObjectPtr), MeshDescription, bHasNormal, SelectedTranslator->ScalingFactor, Offset);
 	}
 	else if (Object.ObjectPtr->IsKindOf(&ON_Brep::m_ON_Brep_class_rtti) )
 	{
diff --git a/Engine/Plugins/Enterprise/DatasmithCADImporter/Source/DatasmithOpenNurbsTranslator/Private/RhinoCoretechWrapper.cpp b/Engine/Plugins/Enterprise/DatasmithCADImporter/Source/DatasmithOpenNurbsTranslator/Private/RhinoCoretechWrapper.cpp
--- a/Engine/Plugins/Enterprise/DatasmithCADImporter/Source/DatasmithOpenNurbsTranslator/Private/RhinoCoretechWrapper.cpp
+++ b/Engine/Plugins/Enterprise/DatasmithCADImporter/Source/DatasmithOpenNurbsTranslator/Private/RhinoCoretechWrapper.cpp
@@ -407,10 +407,17 @@ CT_IO_ERROR FRhinoCoretechWrapper::Tessellate(FMeshDescription& Mesh, CADLibrary
 	return CADLibrary::Tessellate(MainObjectId, ImportParams, Mesh, MeshParameters);
 }
 
-CT_OBJECT_ID BRepToKernelIOBodyTranslator::CreateBody()
+CT_OBJECT_ID BRepToKernelIOBodyTranslator::CreateBody(const ON_3dVector& Offset)
 {
 	BrepTrimToCoedge.SetNumZeroed(BRep.m_T.Count());
 
+	double boxmin[3];
+	double boxmax[3];
+	BRep.GetBBox(boxmin, boxmax);
+
+	BRep.Translate(Offset);
+	BRep.GetBBox(boxmin, boxmax);
+
 	// Create ct faces
 	BRep.FlipReversedSurfaces();
 	CT_LIST_IO FaceList;
@@ -420,6 +427,7 @@ CT_OBJECT_ID BRepToKernelIOBodyTranslator::CreateBody()
 		const ON_BrepFace& On_face = BRep.m_F[index];
 		CreateCTFace(On_face, FaceList);
 	}
+	BRep.Translate(-Offset);
 
 	if (FaceList.IsEmpty())
 	{
@@ -436,8 +444,7 @@ CT_OBJECT_ID BRepToKernelIOBodyTranslator::CreateBody()
 	return 0;
 }
 
-
-CADLibrary::CheckedCTError FRhinoCoretechWrapper::AddBRep(ON_Brep& Brep)
+CADLibrary::CheckedCTError FRhinoCoretechWrapper::AddBRep(ON_Brep& Brep, const ON_3dVector& Offset)
 {
 	CADLibrary::CheckedCTError Result;
 	if (!IsSessionValid())
@@ -447,7 +454,7 @@ CADLibrary::CheckedCTError FRhinoCoretechWrapper::AddBRep(ON_Brep& Brep)
 	}
 
 	BRepToKernelIOBodyTranslator BodyTranslator(Brep);
-	CT_OBJECT_ID BodyID = BodyTranslator.CreateBody();
+	CT_OBJECT_ID BodyID = BodyTranslator.CreateBody(Offset);
 
 	CT_LIST_IO Bodies;
 	if (BodyID)
diff --git a/Engine/Plugins/Enterprise/DatasmithCADImporter/Source/DatasmithOpenNurbsTranslator/Private/RhinoCoretechWrapper.h b/Engine/Plugins/Enterprise/DatasmithCADImporter/Source/DatasmithOpenNurbsTranslator/Private/RhinoCoretechWrapper.h
--- a/Engine/Plugins/Enterprise/DatasmithCADImporter/Source/DatasmithOpenNurbsTranslator/Private/RhinoCoretechWrapper.h
+++ b/Engine/Plugins/Enterprise/DatasmithCADImporter/Source/DatasmithOpenNurbsTranslator/Private/RhinoCoretechWrapper.h
@@ -11,6 +11,7 @@ class ON_Brep;
 class ON_NurbsSurface;
 class ON_BrepFace;
 class ON_BoundingBox;
+class ON_3dVector;
 
 class BRepToKernelIOBodyTranslator
 {
@@ -19,7 +20,7 @@ public:
 		: BRep(InBRep)
 	{}
 
-	CT_OBJECT_ID CreateBody();
+	CT_OBJECT_ID CreateBody(const ON_3dVector& ON_Offset);
 private:
 	CT_OBJECT_ID CreateCTSurface(ON_NurbsSurface& Surface);
 	void CreateCTFace(const ON_BrepFace& Face, CT_LIST_IO& dest);
@@ -50,7 +51,14 @@ public:
 	{
 	}
 
-	CADLibrary::CheckedCTError AddBRep(ON_Brep& brep);
+	/**
+	 * Set BRep to tessellate, offsetting it prior to tessellation(used to set mesh pivot at the center of the surface bounding box)
+	 *
+	 * @param  Brep	a brep to tessellate
+	 * @param  Offset translate brep by this value before tessellating 
+	 */
+	CADLibrary::CheckedCTError AddBRep(ON_Brep& Brep, const ON_3dVector& Offset);
+	
 	static TSharedPtr<FRhinoCoretechWrapper> GetSharedSession(double SceneUnit, double ScaleFactor);
 
 	CT_IO_ERROR Tessellate(FMeshDescription& Mesh, CADLibrary::FMeshParameters& MeshParameters);
