Merged CL 11278415 and 11278574 from Dev-Enterprise to Release-4.25

OpenNurbs - fixed UVs conversion. UV coords origin to Unreal. And UV transform for tiling/shift

diff --git a/Engine/Plugins/Enterprise/DatasmithCADImporter/Source/DatasmithOpenNurbsTranslator/Private/DatasmithOpenNurbsTranslator.cpp b/Engine/Plugins/Enterprise/DatasmithCADImporter/Source/DatasmithOpenNurbsTranslator/Private/DatasmithOpenNurbsTranslator.cpp
--- a/Engine/Plugins/Enterprise/DatasmithCADImporter/Source/DatasmithOpenNurbsTranslator/Private/DatasmithOpenNurbsTranslator.cpp
+++ b/Engine/Plugins/Enterprise/DatasmithCADImporter/Source/DatasmithOpenNurbsTranslator/Private/DatasmithOpenNurbsTranslator.cpp
@@ -371,27 +371,6 @@ namespace DatasmithOpenNurbsTranslatorUtils
 
 		if (bHasUVData)
 		{
-			// Reorient UV along V axis
-			float vMin = FLT_MAX;
-			float vMax = -FLT_MAX;
-
-			for (FVector2D& uvTextCoord : UvCoords)
-			{
-				if (uvTextCoord[1] < vMin)
-				{
-					vMin = uvTextCoord[1];
-				}
-				if (uvTextCoord[1] > vMax)
-				{
-					vMax = uvTextCoord[1];
-				}
-			}
-
-			for (FVector2D& uvTextCoord : UvCoords)
-			{
-				uvTextCoord[1] = vMin + vMax - uvTextCoord[1];
-			}
-
 			for (int Index = 0; Index < Faces.Num(); ++Index)
 			{
 				FFace& Face = Faces[Index];
@@ -435,11 +414,10 @@ namespace DatasmithOpenNurbsTranslatorUtils
 
 			// Fill the vertex array
 			FVertexID AddedVertexId = MeshDescription.CreateVertex();
-			VertexPositions[AddedVertexId] = FVector(-Pos.X, Pos.Y, Pos.Z);
+			VertexPositions[AddedVertexId] = FVector(- Pos.X, Pos.Y, Pos.Z);
 		}
 
 		int32 VertexIndices[3];
-		FBox UVBBox(FVector(MAX_FLT), FVector(-MAX_FLT));
 
 		const int32 CornerCount = 3; // only triangles
 		FVector CornerPositions[3];
@@ -495,8 +473,8 @@ namespace DatasmithOpenNurbsTranslatorUtils
 					const FVector2D& UVValues = Face.UVCoords[CornerIndex];
 					if (!UVValues.ContainsNaN())
 					{
-						UVBBox += FVector(UVValues, 0.0f);
-						VertexInstanceUVs.Set(CornerVertexInstanceIDs[CornerIndex], 0, UVValues);
+						// Convert UVs - bottom-left Rhino's texture origin to Unreal's top-left
+						VertexInstanceUVs.Set(CornerVertexInstanceIDs[CornerIndex], 0, FVector2D(UVValues[0], 1.f - UVValues[1]));
 					}
 				}
 			}
@@ -998,10 +976,25 @@ void FOpenNurbsTranslatorImpl::TranslateMaterialTable(const ON_ObjectArray<ON_Ma
 				FVector RotationAngles = Transform.GetRotation().Euler();
 
 				UVParameters.UVTiling.X = Tiling.X;
-				UVParameters.UVOffset.X = Translation.X;
 
 				UVParameters.UVTiling.Y = Tiling.Y;
-				UVParameters.UVOffset.Y = -Translation.Y; // V-coordinate is inverted in Unreal
+
+				if (!Tiling.IsNearlyZero())
+				{
+					UVParameters.UVOffset.X = Translation.X / Tiling.X;
+
+					// Recomputing offset from Rhino to Unreal, taking into account how tiling Pivot is set in SetupUVEdit: 
+					// P = pivot, C - input uv, T - tiling, O - offset
+					// Texture coordinate transformed like this: UV = T*(C+O-P)+P (ignoring mirror/rotation)
+					// P is 0.5 (for V), so  UV = T*(C+O - 0.5) + 0.5 = T*C + T*O - T*0.5 + 0.5
+					// Also, for Rhino we P` = 1 (Rhino image origin bottom-left),
+					// so tiling in Rhino is defined as  UV = T`*(C` + O` - 1) + 1
+					// T`*C + T`*O` - T`*1 + 1 = T*C + T*O - T*0.5 + 0.5
+					// T` = T, reducing further and get:
+					// O` = O + 0.5 - 0.5 / T
+					// #ueent_todo: this could be simplified if we were able to pass custom Tiling Pivot
+					UVParameters.UVOffset.Y = -(Translation.Y / Tiling.Y + 0.5f - 0.5f / Tiling.Y); // V-coordinate is inverted in Unreal
+				}
 
 				// Rotation angle is reversed because V-axis points down in Unreal while it points up in OpenNurbs
 				UVParameters.RotationAngle = -RotationAngles.Z;
@@ -3108,6 +3101,7 @@ bool FOpenNurbsTranslatorImpl::ComputeObjectGeometryCenter(const FOpenNurbsObjec
 {
 	bool bIsValid = false;
 	ON_3dVector GeometryCenter;
+	GeometryCenter.Zero();
 	if (Object.ObjectPtr->IsKindOf(&ON_Mesh::m_ON_Mesh_class_rtti))
 	{
 		bIsValid = ComputeGeometryCenter(ON_Mesh::Cast(Object.ObjectPtr), GeometryCenter);
