Merging CL 11302594 from //UE4/Dev-Enterprise to //UE4/Release-4.25

PLMXML - added xml extension and wording fixes

diff --git a/Engine/Plugins/Enterprise/DatasmithCADImporter/Source/DatasmithPLMXMLTranslator/Private/DatasmithPlmXmlTranslator.cpp b/Engine/Plugins/Enterprise/DatasmithCADImporter/Source/DatasmithPLMXMLTranslator/Private/DatasmithPlmXmlTranslator.cpp
--- a/Engine/Plugins/Enterprise/DatasmithCADImporter/Source/DatasmithPLMXMLTranslator/Private/DatasmithPlmXmlTranslator.cpp
+++ b/Engine/Plugins/Enterprise/DatasmithCADImporter/Source/DatasmithPLMXMLTranslator/Private/DatasmithPlmXmlTranslator.cpp
@@ -15,7 +15,8 @@ void FDatasmithPlmXmlTranslator::Initialize(FDatasmithTranslatorCapabilities& Ou
 	OutCapabilities.bParallelLoadStaticMeshSupported = true;
 
 	TArray<FFileFormatInfo>& Formats = OutCapabilities.SupportedFileFormats;
-    Formats.Emplace(TEXT("plmxml"), TEXT("PML(Product Lifecycle Management) XML"));
+    Formats.Emplace(TEXT("plmxml"), TEXT("PLMXML"));
+    Formats.Emplace(TEXT("xml"), TEXT("PLMXML"));
 }
 
 bool FDatasmithPlmXmlTranslator::LoadScene(TSharedRef<IDatasmithScene> OutScene)
