Merging 11284724 //UE4/Dev-Enterprise to Release-4.25 (//UE4/Release-4.25)
CoreTechFileParser improvement : converting as soon as possible void* into TArray<Vector>, ...

diff --git a/Engine/Plugins/Enterprise/DatasmithCADImporter/Source/CADInterfaces/Private/CoreTechFileParser.cpp b/Engine/Plugins/Enterprise/DatasmithCADImporter/Source/CADInterfaces/Private/CoreTechFileParser.cpp
--- a/Engine/Plugins/Enterprise/DatasmithCADImporter/Source/CADInterfaces/Private/CoreTechFileParser.cpp
+++ b/Engine/Plugins/Enterprise/DatasmithCADImporter/Source/CADInterfaces/Private/CoreTechFileParser.cpp
@@ -20,24 +20,19 @@ namespace CADLibrary
 {
 
 namespace {
-	double Distance(CT_COORDINATE Point1, CT_COORDINATE Point2)
-	{
-		return sqrt((Point2.xyz[0] - Point1.xyz[0]) * (Point2.xyz[0] - Point1.xyz[0]) + (Point2.xyz[1] - Point1.xyz[1]) * (Point2.xyz[1] - Point1.xyz[1]) + (Point2.xyz[2] - Point1.xyz[2]) * (Point2.xyz[2] - Point1.xyz[2]));
-	}
 
-	template<typename UVType>
-	void ScaleUV(CT_OBJECT_ID FaceID, void* TexCoordArray, int32 VertexCount, UVType Scale)
+	void ScaleUV(CT_OBJECT_ID FaceID, TArray<FVector2D>& TexCoordArray, float Scale)
 	{
-		UVType VMin, VMax, UMin, UMax;
+		float VMin, VMax, UMin, UMax;
 		VMin = UMin = HUGE_VALF;
 		VMax = UMax = -HUGE_VALF;
-		UVType* UVSet = (UVType*) TexCoordArray;
-		for (int32 Index = 0, UVCoord = 0; Index < VertexCount; ++Index, UVCoord += 2)
+
+		for (const FVector2D& TexCoord : TexCoordArray)
 		{
-			UMin = FMath::Min(UVSet[UVCoord + 0], UMin);
-			UMax = FMath::Max(UVSet[UVCoord + 0], UMax);
-			VMin = FMath::Min(UVSet[UVCoord + 1], VMin);
-			VMax = FMath::Max(UVSet[UVCoord + 1], VMax);
+			UMin = FMath::Min(TexCoord[0], UMin);
+			UMax = FMath::Max(TexCoord[0], UMax);
+			VMin = FMath::Min(TexCoord[1], VMin);
+			VMax = FMath::Max(TexCoord[1], VMax);
 		}
 
 		double PuMin, PuMax, PvMin, PvMax;
@@ -57,9 +52,9 @@ namespace {
 		CT_OBJECT_TYPE SurfaceType;
 		CT_SURFACE_IO::AskType(SurfaceID, SurfaceType);
 
-		UVType DeltaU = (PuMax - PuMin) / (NbIsoCurves - 1);
-		UVType DeltaV = (PvMax - PvMin) / (NbIsoCurves - 1);
-		UVType U = PuMin, V = PvMin;
+		float DeltaU = (PuMax - PuMin) / (NbIsoCurves - 1);
+		float DeltaV = (PvMax - PvMin) / (NbIsoCurves - 1);
+		float U = PuMin, V = PvMin;
 
 		CT_COORDINATE NodeMatrix[121];
 
@@ -75,10 +70,10 @@ namespace {
 		}
 
 		// Compute length of 7 iso V line
-		UVType LengthU[NbIsoCurves];
-		UVType LengthUMin = HUGE_VAL;
-		UVType LengthUMax = 0;
-		UVType LengthUMed = 0;
+		float LengthU[NbIsoCurves];
+		float LengthUMin = HUGE_VAL;
+		float LengthUMax = 0;
+		float LengthUMed = 0;
 
 		for (int32 IndexJ = 0; IndexJ < NbIsoCurves; IndexJ++)
 		{
@@ -95,10 +90,10 @@ namespace {
 		LengthUMed = LengthUMed * 2 / 3 + LengthUMax / 3;
 
 		// Compute length of 7 iso U line
-		UVType LengthV[NbIsoCurves];
-		UVType LengthVMin = HUGE_VAL;
-		UVType LengthVMax = 0;
-		UVType LengthVMed = 0;
+		float LengthV[NbIsoCurves];
+		float LengthVMin = HUGE_VAL;
+		float LengthVMax = 0;
+		float LengthVMed = 0;
 
 		for (int32 IndexI = 0; IndexI < NbIsoCurves; IndexI++)
 		{
@@ -140,13 +135,13 @@ namespace {
 
 		// scale the UV map
 		// 0.1 define UV in cm and not in mm
-		UVType VScale = Scale * LengthVMed * 1 / (VMax - VMin) / 100;
-		UVType UScale = Scale * LengthUMed * 1 / (UMax - UMin) / 100;
+		float VScale = Scale * LengthVMed * 1 / (VMax - VMin) / 100;
+		float UScale = Scale * LengthUMed * 1 / (UMax - UMin) / 100;
 
-		for (int32 Index = 0, UVCoord = 0; Index < VertexCount; ++Index, UVCoord += 2)
+		for (FVector2D& TexCoord : TexCoordArray)
 		{
-			UVSet[UVCoord + 0] *= UScale;
-			UVSet[UVCoord + 1] *= VScale;
+			TexCoord[0] *= UScale;
+			TexCoord[1] *= VScale;
 		}
 	}
 }
@@ -183,6 +178,36 @@ uint32 GetGeomFileHash(const uint32 InSGHash, const FImportParameters& ImportPar
 	return FileHash;
 }
 
+template<typename ValueType>
+void FillArrayOfVector(int32 ElementCount, void* InCTValueArray, FVector* OutValueArray)
+{
+	ValueType* Values = (ValueType*)InCTValueArray;
+	for (int Indice = 0; Indice < ElementCount; ++Indice)
+	{
+		OutValueArray[Indice].Set((float)Values[Indice * 3], (float)Values[Indice * 3 + 1], (float)Values[Indice * 3 + 2]);
+	}
+}
+
+template<typename ValueType>
+void FillArrayOfVector2D(int32 ElementCount, void* InCTValueArray, FVector2D* OutValueArray)
+{
+	ValueType* Values = (ValueType*)InCTValueArray;
+	for (int Indice = 0; Indice < ElementCount; ++Indice)
+	{
+		OutValueArray[Indice].Set((float)Values[Indice * 2], (float)Values[Indice * 2 + 1]);
+	}
+}
+
+template<typename ValueType>
+void FillArrayOfInt(int32 ElementCount, void* InCTValueArray, int32* OutValueArray)
+{
+	ValueType* Values = (ValueType*)InCTValueArray;
+	for (int Indice = 0; Indice < ElementCount; ++Indice)
+	{
+		OutValueArray[Indice] = (int32)Values[Indice];
+	}
+}
+
 uint32 GetFaceTessellation(CT_OBJECT_ID FaceID, TArray<FTessellationData>& FaceTessellationSet, const FImportParameters& ImportParams)
 {
 	CT_IO_ERROR Error = IO_OK;
@@ -215,38 +240,81 @@ uint32 GetFaceTessellation(CT_OBJECT_ID FaceID, TArray<FTessellationData>& FaceT
 	}
 
 	FTessellationData& Tessellation = FaceTessellationSet.Emplace_GetRef();
-	if (ImportParams.bScaleUVMap && TexCoordArray != nullptr)
+	Tessellation.IndexArray.SetNum(IndexCount);
+
+	switch (IndexType)
+	{
+	case CT_TESS_UBYTE:
+		FillArrayOfInt<uint8>(IndexCount, IndexArray, Tessellation.IndexArray.GetData());
+		break;
+	case CT_TESS_USHORT:
+		FillArrayOfInt<uint16>(IndexCount, IndexArray, Tessellation.IndexArray.GetData());
+		break;
+	case CT_TESS_UINT:
+		FillArrayOfInt<uint32>(IndexCount, IndexArray, Tessellation.IndexArray.GetData());
+		break;
+	}
+
+	Tessellation.VertexArray.SetNum(VertexCount);
+	switch (VertexType)
 	{
+	case CT_TESS_FLOAT:
+		FillArrayOfVector<float>(VertexCount, VertexArray, Tessellation.VertexArray.GetData());
+		break;
+	case CT_TESS_DOUBLE:
+		FillArrayOfVector<double>(VertexCount, VertexArray, Tessellation.VertexArray.GetData());
+		break;
+	}
+
+	Tessellation.NormalArray.SetNum(NormalCount);
+	switch (NormalType)
+	{
+	case CT_TESS_BYTE:
+		Tessellation.NormalArray.SetNumZeroed(NormalCount);
+		break;
+	case CT_TESS_SHORT:
+	{
+		int8* InCTValueArray = (int8*)NormalArray;
+		for (CT_UINT32 Indice = 0; Indice < NormalCount; ++Indice)
+		{
+			Tessellation.NormalArray[Indice].Set(((float)InCTValueArray[Indice]) / 255.f, ((float)InCTValueArray[Indice + 1]) / 255.f, ((float)InCTValueArray[Indice + 2]) / 255.f);
+		}
+		break;
+	}
+	case CT_TESS_FLOAT:
+		FillArrayOfVector<float>(NormalCount, NormalArray, Tessellation.NormalArray.GetData());
+		break;
+	}
+
+	if (TexCoordArray)
+	{
+		Tessellation.TexCoordArray.SetNum(VertexCount);
 		switch (TexCoordType)
 		{
+		case CT_TESS_SHORT:
+		{
+			int8* InCTValueArray = (int8*)TexCoordArray;
+			for (CT_UINT32 Indice = 0; Indice < VertexCount; ++Indice)
+			{
+				Tessellation.TexCoordArray[Indice].Set(((float)InCTValueArray[Indice]) / 255.f, ((float)InCTValueArray[Indice + 1]) / 255.f);
+			}
+			break;
+		}
 		case CT_TESS_FLOAT:
-			ScaleUV<float>(FaceID, TexCoordArray, VertexCount, (float) ImportParams.ScaleFactor);
+			FillArrayOfVector2D<float>(VertexCount, TexCoordArray, Tessellation.TexCoordArray.GetData());
 			break;
 		case CT_TESS_DOUBLE:
-			ScaleUV<double>(FaceID, TexCoordArray, VertexCount, ImportParams.ScaleFactor);
+			FillArrayOfVector2D<double>(VertexCount, TexCoordArray, Tessellation.TexCoordArray.GetData());
 			break;
 		}
 	}
 
-	Tessellation.VertexCount = VertexCount;
-	Tessellation.NormalCount = NormalCount;
-	Tessellation.IndexCount = IndexCount;
-	Tessellation.TexCoordCount = TexCoordArray ? VertexCount : 0;
-	Tessellation.SizeOfVertexType = GetSize(VertexType);
-	Tessellation.SizeOfTexCoordType = GetSize(TexCoordType);
-	Tessellation.SizeOfNormalType = GetSize(NormalType);
-	Tessellation.SizeOfIndexType = GetSize(IndexType);
-
-	Tessellation.VertexArray.Append((uint8*) VertexArray, 3 * Tessellation.VertexCount * Tessellation.SizeOfVertexType);
-	Tessellation.NormalArray.Append((uint8*)NormalArray, 3 * Tessellation.NormalCount * Tessellation.SizeOfNormalType);
-	Tessellation.IndexArray.Append((uint8*)IndexArray, Tessellation.IndexCount * Tessellation.SizeOfIndexType);
-
-	if (TexCoordArray)
+	if (ImportParams.bScaleUVMap && Tessellation.TexCoordArray.Num() != 0)
 	{
-		Tessellation.TexCoordArray.Append((uint8*)TexCoordArray, 2 * Tessellation.TexCoordCount * Tessellation.SizeOfTexCoordType);
+		ScaleUV(FaceID, Tessellation.TexCoordArray, (float) ImportParams.ScaleFactor);
 	}
 
-	return (uint32)Tessellation.IndexCount / 3;
+	return Tessellation.IndexArray.Num() / 3;
 }
 
 
diff --git a/Engine/Plugins/Enterprise/DatasmithCADImporter/Source/CADInterfaces/Public/CoreTechFileParser.h b/Engine/Plugins/Enterprise/DatasmithCADImporter/Source/CADInterfaces/Public/CoreTechFileParser.h
--- a/Engine/Plugins/Enterprise/DatasmithCADImporter/Source/CADInterfaces/Public/CoreTechFileParser.h
+++ b/Engine/Plugins/Enterprise/DatasmithCADImporter/Source/CADInterfaces/Public/CoreTechFileParser.h
@@ -34,6 +34,16 @@ enum class ECoretechParsingResult
 	FileNotFound,
 };
 
+FORCEINLINE double FaceArea(const FVector& TriangleA, const FVector& TriangleB, const FVector& TriangleC)
+{
+	return ((TriangleB - TriangleA) ^ (TriangleC - TriangleA)).Size() * 0.5f;;
+};
+
+FORCEINLINE double Distance(const CT_COORDINATE& Point1, const CT_COORDINATE& Point2)
+{
+	return sqrt((Point2.xyz[0] - Point1.xyz[0]) * (Point2.xyz[0] - Point1.xyz[0]) + (Point2.xyz[1] - Point1.xyz[1]) * (Point2.xyz[1] - Point1.xyz[1]) + (Point2.xyz[2] - Point1.xyz[2]) * (Point2.xyz[2] - Point1.xyz[2]));
+};
+
 #ifdef CAD_INTERFACE
 
 CADINTERFACES_API bool GetColor(uint32 ColorHash, FColor& OutColor);
diff --git a/Engine/Plugins/Enterprise/DatasmithCADImporter/Source/CADLibrary/Private/CoreTechHelper.cpp b/Engine/Plugins/Enterprise/DatasmithCADImporter/Source/CADLibrary/Private/CoreTechHelper.cpp
--- a/Engine/Plugins/Enterprise/DatasmithCADImporter/Source/CADLibrary/Private/CoreTechHelper.cpp
+++ b/Engine/Plugins/Enterprise/DatasmithCADImporter/Source/CADLibrary/Private/CoreTechHelper.cpp
@@ -73,7 +73,7 @@ void FillVertexPosition(const FImportParameters& ImportParams, const FMeshParame
 	int32 TriangleCount = Body.TriangleCount;
 	TArray<FTessellationData>& FaceTessellationSet = Body.Faces;
 
-	// Add offset on the bbox to avoid to remove good vertex
+	// Add offset to the bounding box to avoid to remove good vertex
 	FVector Size = Body.BBox.GetSize();
 	FBox BBox = Body.BBox.ExpandBy(Size.Size());
 	BBox.IsValid = Body.BBox.IsValid;
@@ -91,11 +91,10 @@ void FillVertexPosition(const FImportParameters& ImportParams, const FMeshParame
 	for (FTessellationData& CTTessellation : FaceTessellationSet)
 	{
 		CTTessellation.StartVertexIndex = GlobalVertexCount;
-		for (uint32 VertexIndex = 0; VertexIndex < CTTessellation.VertexCount; ++VertexIndex, ++GlobalVertexCount)
+		for (const FVector& Vertex : CTTessellation.VertexArray)
 		{
-			CopyValue(CTTessellation.VertexArray.GetData(), VertexIndex * 3, CTTessellation.SizeOfVertexType, true, Position);
-			Position *= ImportParams.ScaleFactor;  // convert Position unit into cm according to scaleFactor
-			VertexDataSet.Emplace(GlobalVertexCount, Position);
+			VertexDataSet.Emplace(GlobalVertexCount, Vertex * ImportParams.ScaleFactor);
+			++GlobalVertexCount;
 		}
 	}
 	VertexDataSet.SetNum(GlobalVertexCount);
@@ -212,8 +211,8 @@ void FillVertexPosition(const FImportParameters& ImportParams, const FMeshParame
 	GlobalVertexIndex = 0;
 	for (FTessellationData& CTTessellation : FaceTessellationSet)
 	{
-		CTTessellation.VertexIdSet.SetNum(CTTessellation.VertexCount);
-		for (uint32 VertexIndex = 0; VertexIndex < CTTessellation.VertexCount; ++VertexIndex, ++GlobalVertexIndex)
+		CTTessellation.VertexIdSet.SetNum(CTTessellation.VertexArray.Num());
+		for (int32 VertexIndex = 0; VertexIndex < CTTessellation.VertexArray.Num(); ++VertexIndex, ++GlobalVertexIndex)
 		{
 			int32 NewIndex = NewIndexOf[IndexOfCoincidentNode[GlobalVertexIndex]];
 			CTTessellation.VertexIdSet[VertexIndex] = VertexDataSet[NewIndex].VertexID.GetValue();
@@ -226,8 +225,8 @@ void FillVertexPosition(const FImportParameters& ImportParams, const FMeshParame
 		GlobalVertexIndex = 0;
 		for (FTessellationData& CTTessellation : FaceTessellationSet)
 		{
-			CTTessellation.SymVertexIdSet.SetNum(CTTessellation.VertexCount);
-			for (uint32 VertexIndex = 0; VertexIndex < CTTessellation.VertexCount; ++VertexIndex, ++GlobalVertexIndex)
+			CTTessellation.SymVertexIdSet.SetNum(CTTessellation.VertexArray.Num());
+			for (int32 VertexIndex = 0; VertexIndex < CTTessellation.VertexArray.Num(); ++VertexIndex, ++GlobalVertexIndex)
 			{
 				uint32 NewIndex = NewIndexOf[IndexOfCoincidentNode[GlobalVertexIndex]];
 				CTTessellation.SymVertexIdSet[VertexIndex] = VertexDataSet[NewIndex].SymVertexID.GetValue();
@@ -330,21 +329,19 @@ bool FillMesh(const FMeshParameters& MeshParameters, const FImportParameters& Im
 			int32 VertexIDs[3];
 			FVector Temp3D = { 0, 0, 0 };
 			FVector2D TexCoord2D = { 0, 0 };
-			int32_t IndicesVertex[3];
 
-			MeshVertexInstanceIDs.SetNum(Tessellation.IndexCount);
-			CTFaceIndex.Reserve(Tessellation.IndexCount);
+			MeshVertexInstanceIDs.SetNum(Tessellation.IndexArray.Num());
+			CTFaceIndex.Reserve(Tessellation.IndexArray.Num());
 			CTFaceIndex.SetNum(0);
 
 			TArray<int32>& VertexIdSet = (Step == 0) ? Tessellation.VertexIdSet : Tessellation.SymVertexIdSet;
 
 			// build each valid face i.e. 3 different indexes
-			for (uint32 Index = 0, NewIndex = 0; Index < Tessellation.IndexCount; Index += 3)
+			for (int32 Index = 0, NewIndex = 0; Index < Tessellation.IndexArray.Num(); Index += 3)
 			{
-				CopyValue(Tessellation.IndexArray.GetData(), Index, Tessellation.SizeOfIndexType, IndicesVertex);
-				VertexIDs[Orientation[0]] = VertexIdSet[IndicesVertex[0]];
-				VertexIDs[Orientation[1]] = VertexIdSet[IndicesVertex[1]];
-				VertexIDs[Orientation[2]] = VertexIdSet[IndicesVertex[2]];
+				VertexIDs[Orientation[0]] = VertexIdSet[Tessellation.IndexArray[Index + 0]];
+				VertexIDs[Orientation[1]] = VertexIdSet[Tessellation.IndexArray[Index + 1]];
+				VertexIDs[Orientation[2]] = VertexIdSet[Tessellation.IndexArray[Index + 2]];
 
 				if (VertexIDs[0] == InvalidID || VertexIDs[1] == InvalidID || VertexIDs[2] == InvalidID)
 				{
@@ -357,9 +354,9 @@ bool FillMesh(const FMeshParameters& MeshParameters, const FImportParameters& Im
 					continue;
 				}
 
-				CTFaceIndex.Add(IndicesVertex[0]);
-				CTFaceIndex.Add(IndicesVertex[1]);
-				CTFaceIndex.Add(IndicesVertex[2]);
+				CTFaceIndex.Add(Tessellation.IndexArray[Index + 0]);
+				CTFaceIndex.Add(Tessellation.IndexArray[Index + 1]);
+				CTFaceIndex.Add(Tessellation.IndexArray[Index + 2]);
 
 				TriangleVertexInstanceIDs[0] = MeshVertexInstanceIDs[NewIndex++] = MeshDescription.CreateVertexInstance((FVertexID) VertexIDs[0]);
 				TriangleVertexInstanceIDs[1] = MeshVertexInstanceIDs[NewIndex++] = MeshDescription.CreateVertexInstance((FVertexID) VertexIDs[1]);
@@ -391,16 +388,15 @@ bool FillMesh(const FMeshParameters& MeshParameters, const FImportParameters& Im
 					for (int32 Index = 0; Index < TriangleCount; Index++)
 					{
 						FVertexInstanceID VertexInstanceID = MeshVertexInstanceIDs[IndexFace + Orientation[Index]];
-						CopyValue(Tessellation.TexCoordArray.GetData(), CTFaceIndex[IndexFace + Index] * 2, Tessellation.SizeOfTexCoordType, false, Temp3D);
-						VertexInstanceUVs.Set(VertexInstanceID, UVChannel, FVector2D(Temp3D));
+						VertexInstanceUVs.Set(VertexInstanceID, UVChannel, Tessellation.TexCoordArray[CTFaceIndex[IndexFace + Index]]);
 					}
 				}
 			}
 
-			if (Tessellation.NormalCount == 1)
+			FDatasmithUtils::ConvertVectorArray((FDatasmithUtils::EModelCoordSystem) ImportParams.ModelCoordSys, Tessellation.NormalArray);
+			if (Tessellation.NormalArray.Num() == 1)
 			{
-				CopyValue(Tessellation.NormalArray.GetData(), 0, Tessellation.SizeOfNormalType, true, Temp3D);
-				Temp3D = FDatasmithUtils::ConvertVector((FDatasmithUtils::EModelCoordSystem) ImportParams.ModelCoordSys, Temp3D).GetSafeNormal();
+				Temp3D = Tessellation.NormalArray[0].GetSafeNormal();
 				for (int32 Index = 0; Index < CTFaceIndex.Num(); Index++)
 				{
 					FVertexInstanceID VertexInstanceID = MeshVertexInstanceIDs[Index];
@@ -409,13 +405,17 @@ bool FillMesh(const FMeshParameters& MeshParameters, const FImportParameters& Im
 			}
 			else
 			{
+				for (FVector& Normal : Tessellation.NormalArray)
+				{
+					Normal = Normal.GetSafeNormal();
+				}
+
 				for (int32 IndexFace = 0; IndexFace < CTFaceIndex.Num(); IndexFace += 3)
 				{
 					for (int32 Index = 0; Index < 3; Index++)
 					{
 						FVertexInstanceID VertexInstanceID = MeshVertexInstanceIDs[IndexFace + Orientation[Index]];
-						CopyValue(Tessellation.NormalArray.GetData(), CTFaceIndex[IndexFace + Index] * 3, Tessellation.SizeOfNormalType, true, Temp3D);
-						VertexInstanceNormals[VertexInstanceID] = FDatasmithUtils::ConvertVector((FDatasmithUtils::EModelCoordSystem) ImportParams.ModelCoordSys, Temp3D).GetSafeNormal();
+						VertexInstanceNormals[VertexInstanceID] = Tessellation.NormalArray[CTFaceIndex[IndexFace + Index]];
 					}
 				}
 			}
@@ -471,11 +471,6 @@ bool ConvertCTBodySetToMeshDescription(const FImportParameters& ImportParams, co
 	return MeshDescription.Polygons().Num() > 0;
 }
 
-double Distance(CT_COORDINATE Point1, CT_COORDINATE Point2)
-{
-	return sqrt((Point2.xyz[0] - Point1.xyz[0]) * (Point2.xyz[0] - Point1.xyz[0]) + (Point2.xyz[1] - Point1.xyz[1]) * (Point2.xyz[1] - Point1.xyz[1]) + (Point2.xyz[2] - Point1.xyz[2]) * (Point2.xyz[2] - Point1.xyz[2]));
-}
-
 template<typename UVType>
 void ScaleUV(CT_OBJECT_ID FaceID, FTessellationData& Tessellation, UVType Scale)
 {
diff --git a/Engine/Plugins/Enterprise/DatasmithCADImporter/Source/CADTools/Private/CADData.cpp b/Engine/Plugins/Enterprise/DatasmithCADImporter/Source/CADTools/Private/CADData.cpp
--- a/Engine/Plugins/Enterprise/DatasmithCADImporter/Source/CADTools/Private/CADData.cpp
+++ b/Engine/Plugins/Enterprise/DatasmithCADImporter/Source/CADTools/Private/CADData.cpp
@@ -76,18 +76,9 @@ FArchive& operator<<(FArchive& Ar, FTessellationData& TessellationData)
 	Ar << TessellationData.NormalArray;
 	Ar << TessellationData.IndexArray;
 	Ar << TessellationData.TexCoordArray;
-	Ar << TessellationData.VertexCount;
-	Ar << TessellationData.NormalCount;
-	Ar << TessellationData.IndexCount;
-	Ar << TessellationData.TexCoordCount;
 
 	Ar << TessellationData.StartVertexIndex;
 
-	Ar << TessellationData.SizeOfVertexType;
-	Ar << TessellationData.SizeOfTexCoordType;
-	Ar << TessellationData.SizeOfNormalType;
-	Ar << TessellationData.SizeOfIndexType;
-
 	Ar << TessellationData.ColorName;
 	Ar << TessellationData.MaterialName;
 
diff --git a/Engine/Plugins/Enterprise/DatasmithCADImporter/Source/CADTools/Private/CADToolsModule.cpp b/Engine/Plugins/Enterprise/DatasmithCADImporter/Source/CADTools/Private/CADToolsModule.cpp
--- a/Engine/Plugins/Enterprise/DatasmithCADImporter/Source/CADTools/Private/CADToolsModule.cpp
+++ b/Engine/Plugins/Enterprise/DatasmithCADImporter/Source/CADTools/Private/CADToolsModule.cpp
@@ -18,7 +18,7 @@ bool FCADToolsModule::IsAvailable()
 
 uint32 FCADToolsModule::GetCacheVersion()
 {
-	return 4;
+	return 5;
 }
 
 IMPLEMENT_MODULE(FCADToolsModule, CADTools);
diff --git a/Engine/Plugins/Enterprise/DatasmithCADImporter/Source/CADTools/Public/CADData.h b/Engine/Plugins/Enterprise/DatasmithCADImporter/Source/CADTools/Public/CADData.h
--- a/Engine/Plugins/Enterprise/DatasmithCADImporter/Source/CADTools/Public/CADData.h
+++ b/Engine/Plugins/Enterprise/DatasmithCADImporter/Source/CADTools/Public/CADData.h
@@ -89,22 +89,13 @@ struct CADTOOLS_API FTessellationData
 {
 	friend CADTOOLS_API FArchive& operator<<(FArchive& Ar, FTessellationData& Tessellation);
 
-	TArray<uint8> VertexArray;
-	TArray<uint8> NormalArray;
-	TArray<uint8> IndexArray;
-	TArray<uint8> TexCoordArray;
-	uint32        VertexCount = 0;
-	uint32        NormalCount = 0;
-	uint32        IndexCount = 0;
-	uint32        TexCoordCount = 0;
+	TArray<FVector> VertexArray;
+	TArray<FVector> NormalArray;
+	TArray<int32> IndexArray;
+	TArray<FVector2D> TexCoordArray;
 
 	uint32    StartVertexIndex = 0;
 
-	uint8 SizeOfVertexType = 0;
-	uint8 SizeOfTexCoordType = 0;
-	uint8 SizeOfNormalType = 0;
-	uint8 SizeOfIndexType = 0;
-
 	CADUUID ColorName = 0;
 	CADUUID MaterialName = 0;
 
diff --git a/Engine/Source/Runtime/Datasmith/DatasmithCore/Public/DatasmithUtils.h b/Engine/Source/Runtime/Datasmith/DatasmithCore/Public/DatasmithUtils.h
--- a/Engine/Source/Runtime/Datasmith/DatasmithCore/Public/DatasmithUtils.h
+++ b/Engine/Source/Runtime/Datasmith/DatasmithCore/Public/DatasmithUtils.h
@@ -60,6 +60,37 @@ public:
 		YUp_RightHanded,
 	};
 
+	static void ConvertVectorArray(EModelCoordSystem ModelCoordSys, TArray<FVector>& Array)
+	{
+		switch (ModelCoordSys)
+		{
+		case EModelCoordSystem::YUp_LeftHanded:
+			for (FVector& Vector : Array)
+			{
+				Vector.Set(Vector[2], Vector[0], Vector[1]);
+			}
+			break;
+
+		case EModelCoordSystem::YUp_RightHanded:
+			for (FVector& Vector : Array)
+			{
+				Vector.Set(-Vector[2], Vector[0], Vector[1]);
+			}
+			break;
+
+		case EModelCoordSystem::ZUp_RightHanded:
+			for (FVector& Vector : Array)
+			{
+				Vector.Set(-Vector[0], Vector[1], Vector[2]);
+			}
+			break;
+
+		case EModelCoordSystem::ZUp_LeftHanded:
+		default:
+			break;
+		}
+	}
+
 	template<typename VecType>
 	static FVector ConvertVector(EModelCoordSystem ModelCoordSys, const VecType& V)
 	{
