Merging 11287461 from //UE4/Release-4.25/Engine/... to //UE4/Release-4.25/Engine/...
Remove unused code

diff --git a/Engine/Plugins/Enterprise/DatasmithCADImporter/Source/CADLibrary/Private/CoreTechHelper.cpp b/Engine/Plugins/Enterprise/DatasmithCADImporter/Source/CADLibrary/Private/CoreTechHelper.cpp
--- a/Engine/Plugins/Enterprise/DatasmithCADImporter/Source/CADLibrary/Private/CoreTechHelper.cpp
+++ b/Engine/Plugins/Enterprise/DatasmithCADImporter/Source/CADLibrary/Private/CoreTechHelper.cpp
@@ -471,131 +471,6 @@ bool ConvertCTBodySetToMeshDescription(const FImportParameters& ImportParams, co
 	return MeshDescription.Polygons().Num() > 0;
 }
 
-template<typename UVType>
-void ScaleUV(CT_OBJECT_ID FaceID, FTessellationData& Tessellation, UVType Scale)
-{
-	UVType VMin, VMax, UMin, UMax;
-	VMin = UMin = HUGE_VALF;
-	VMax = UMax = -HUGE_VALF;
-	UVType* UVSet = (UVType*)Tessellation.TexCoordArray;
-	for (uint32 Index = 0, UVCoord = 0; Index < Tessellation.VertexCount; ++Index, UVCoord += 2)
-	{
-		UMin = FMath::Min(UVSet[UVCoord + 0], UMin);
-		UMax = FMath::Max(UVSet[UVCoord + 0], UMax);
-		VMin = FMath::Min(UVSet[UVCoord + 1], VMin);
-		VMax = FMath::Max(UVSet[UVCoord + 1], VMax);
-	}
-
-	double PuMin, PuMax, PvMin, PvMax;
-	PuMin = PvMin = HUGE_VALF;
-	PuMax = PvMax = -HUGE_VALF;
-
-	// fast UV min max
-	CT_FACE_IO::AskUVminmax(FaceID, PuMin, PuMax, PvMin, PvMax);
-
-	const uint32 NbIsoCurves = 7;
-
-	// Compute Point grid on the restricted surface defined by [PuMin, PuMax], [PvMin, PvMax]
-	CT_OBJECT_ID SurfaceID;
-	CT_ORIENTATION Orientation;
-	CT_FACE_IO::AskSurface(FaceID, SurfaceID, Orientation);
-
-	CT_OBJECT_TYPE SurfaceType;
-	CT_SURFACE_IO::AskType(SurfaceID, SurfaceType);
-
-	UVType DeltaU = (PuMax - PuMin) / (NbIsoCurves - 1);
-	UVType DeltaV = (PvMax - PvMin) / (NbIsoCurves - 1);
-	UVType U = PuMin, V = PvMin;
-
-	CT_COORDINATE NodeMatrix[121];
-
-	for (int32 IndexI = 0; IndexI < NbIsoCurves; IndexI++)
-	{
-		for (int32 IndexJ = 0; IndexJ < NbIsoCurves; IndexJ++)
-		{
-			CT_SURFACE_IO::Evaluate(SurfaceID, U, V, NodeMatrix[IndexI*NbIsoCurves +IndexJ]);
-			V += DeltaV;
-		}
-		U += DeltaU;
-		V = PvMin;
-	}
-
-	// Compute length of 7 iso V line
-	UVType LengthU[NbIsoCurves];
-	UVType LengthUMin = HUGE_VAL;
-	UVType LengthUMax = 0;
-	UVType LengthUMed = 0;
-
-	for (int32 IndexJ = 0; IndexJ < NbIsoCurves; IndexJ++)
-	{
-		LengthU[IndexJ] = 0;
-		for (int32 IndexI = 0; IndexI < (NbIsoCurves - 1); IndexI++)
-		{
-			LengthU[IndexJ] += Distance(NodeMatrix[IndexI * NbIsoCurves + IndexJ], NodeMatrix[(IndexI + 1) * NbIsoCurves + IndexJ]);
-		}
-		LengthUMed += LengthU[IndexJ];
-		LengthUMin = FMath::Min(LengthU[IndexJ], LengthUMin);
-		LengthUMax = FMath::Max(LengthU[IndexJ], LengthUMax);
-	}
-	LengthUMed /= NbIsoCurves;
-	LengthUMed = LengthUMed * 2 / 3 + LengthUMax / 3;
-
-	// Compute length of 7 iso U line
-	UVType LengthV[NbIsoCurves];
-	UVType LengthVMin = HUGE_VAL;
-	UVType LengthVMax = 0;
-	UVType LengthVMed = 0;
-
-	for (int32 IndexI = 0; IndexI < NbIsoCurves; IndexI++)
-	{
-		LengthV[IndexI] = 0;
-		for (int32 IndexJ = 0; IndexJ < (NbIsoCurves - 1); IndexJ++)
-		{
-			LengthV[IndexI] += Distance(NodeMatrix[IndexI * NbIsoCurves + IndexJ], NodeMatrix[IndexI * NbIsoCurves + IndexJ + 1]);
-		}
-		LengthVMed += LengthV[IndexI];
-		LengthVMin = FMath::Min(LengthV[IndexI], LengthVMin);
-		LengthVMax = FMath::Max(LengthV[IndexI], LengthVMax);
-	}
-	LengthVMed /= NbIsoCurves;
-	LengthVMed = LengthVMed * 2 / 3 + LengthVMax / 3;
-
-	switch (SurfaceType)
-	{
-	case CT_CONE_TYPE:
-	case CT_CYLINDER_TYPE:
-	case CT_SPHERE_TYPE:
-		Swap(LengthUMed, LengthVMed);
-		break;
-	case CT_S_REVOL_TYPE:
-	case CT_TORUS_TYPE:
-		// Need swap ?
-		// Swap(LengthUMed, LengthVMed);
-		break;
-	case CT_S_NURBS_TYPE:
-	case CT_PLANE_TYPE:
-	case CT_S_OFFSET_TYPE:
-	case CT_S_RULED_TYPE:
-	case CT_TABULATED_RULED_TYPE:
-	case CT_S_LINEARTRANSFO_TYPE:
-	case CT_S_NONLINEARTRANSFO_TYPE:
-	case CT_S_BLEND_TYPE:
-	default:
-		break;
-	}
-
-	// scale the UV map
-	// 0.1 define UV in cm and not in mm
-	UVType VScale = Scale * LengthVMed * 1 / (VMax - VMin) / 100;
-	UVType UScale = Scale * LengthUMed * 1 / (UMax - UMin) / 100;
-
-	for (uint32 Index = 0, UVCoord = 0; Index < Tessellation.VertexCount; ++Index, UVCoord += 2)
-	{
-		UVSet[UVCoord + 0] *= UScale;
-		UVSet[UVCoord + 1] *= VScale;
-	}
-}
-
 TSharedPtr<IDatasmithUEPbrMaterialElement> CreateDefaultUEPbrMaterial()
 {
 	// Take the Material diffuse color and connect it to the BaseColor of a UEPbrMaterial
