The problem appear with alias symetric part.
Due to optimization of CL11284724, the change of axis coordinate to Zup Left hand was applied twice on normales of the symmetric of the body.

diff --git a/Engine/Plugins/Enterprise/DatasmithCADImporter/Source/CADLibrary/Private/CoreTechHelper.cpp b/Engine/Plugins/Enterprise/DatasmithCADImporter/Source/CADLibrary/Private/CoreTechHelper.cpp
--- a/Engine/Plugins/Enterprise/DatasmithCADImporter/Source/CADLibrary/Private/CoreTechHelper.cpp
+++ b/Engine/Plugins/Enterprise/DatasmithCADImporter/Source/CADLibrary/Private/CoreTechHelper.cpp
@@ -393,10 +393,18 @@ bool FillMesh(const FMeshParameters& MeshParameters, const FImportParameters& Im
 				}
 			}
 
-			FDatasmithUtils::ConvertVectorArray((FDatasmithUtils::EModelCoordSystem) ImportParams.ModelCoordSys, Tessellation.NormalArray);
+			if (!Step)
+			{
+				FDatasmithUtils::ConvertVectorArray((FDatasmithUtils::EModelCoordSystem) ImportParams.ModelCoordSys, Tessellation.NormalArray);
+				for (FVector& Normal : Tessellation.NormalArray)
+				{
+					Normal = Normal.GetSafeNormal();
+				}
+			}
+
 			if (Tessellation.NormalArray.Num() == 1)
 			{
-				Temp3D = Tessellation.NormalArray[0].GetSafeNormal();
+				Temp3D = Tessellation.NormalArray[0];
 				for (int32 Index = 0; Index < CTFaceIndex.Num(); Index++)
 				{
 					FVertexInstanceID VertexInstanceID = MeshVertexInstanceIDs[Index];
