[Build] UE4 Release-4.25 - CL 11502113 - Build Engine Localization

diff --git a/Engine/Plugins/Enterprise/DatasmithCADImporter/Source/CADInterfaces/Private/CoreTechFileParser.cpp b/Engine/Plugins/Enterprise/DatasmithCADImporter/Source/CADInterfaces/Private/CoreTechFileParser.cpp
--- a/Engine/Plugins/Enterprise/DatasmithCADImporter/Source/CADInterfaces/Private/CoreTechFileParser.cpp
+++ b/Engine/Plugins/Enterprise/DatasmithCADImporter/Source/CADInterfaces/Private/CoreTechFileParser.cpp
@@ -21,6 +21,11 @@ namespace CADLibrary
 
 namespace {
 
+	double Distance(const CT_COORDINATE& Point1, const CT_COORDINATE& Point2)
+	{
+		return sqrt((Point2.xyz[0] - Point1.xyz[0]) * (Point2.xyz[0] - Point1.xyz[0]) + (Point2.xyz[1] - Point1.xyz[1]) * (Point2.xyz[1] - Point1.xyz[1]) + (Point2.xyz[2] - Point1.xyz[2]) * (Point2.xyz[2] - Point1.xyz[2]));
+	};
+
 	void ScaleUV(CT_OBJECT_ID FaceID, TArray<FVector2D>& TexCoordArray, float Scale)
 	{
 		float VMin, VMax, UMin, UMax;
diff --git a/Engine/Plugins/Enterprise/DatasmithCADImporter/Source/CADInterfaces/Public/CoreTechFileParser.h b/Engine/Plugins/Enterprise/DatasmithCADImporter/Source/CADInterfaces/Public/CoreTechFileParser.h
--- a/Engine/Plugins/Enterprise/DatasmithCADImporter/Source/CADInterfaces/Public/CoreTechFileParser.h
+++ b/Engine/Plugins/Enterprise/DatasmithCADImporter/Source/CADInterfaces/Public/CoreTechFileParser.h
@@ -34,15 +34,6 @@ enum class ECoretechParsingResult
 	FileNotFound,
 };
 
-FORCEINLINE double FaceArea(const FVector& TriangleA, const FVector& TriangleB, const FVector& TriangleC)
-{
-	return ((TriangleB - TriangleA) ^ (TriangleC - TriangleA)).Size() * 0.5f;;
-};
-
-FORCEINLINE double Distance(const CT_COORDINATE& Point1, const CT_COORDINATE& Point2)
-{
-	return sqrt((Point2.xyz[0] - Point1.xyz[0]) * (Point2.xyz[0] - Point1.xyz[0]) + (Point2.xyz[1] - Point1.xyz[1]) * (Point2.xyz[1] - Point1.xyz[1]) + (Point2.xyz[2] - Point1.xyz[2]) * (Point2.xyz[2] - Point1.xyz[2]));
-};
 
 #ifdef CAD_INTERFACE
 
