Unshelved from pending changelist '11412703':

PLMXML: fixed actor labels -  they are named after ProductRevisionView(or other view nodes/occurrence), without artificial prefixes(same as JT2Go displays them except sanitized)

- also fixed crash when rootRef is a ProductRevisionView

diff --git a/Engine/Plugins/Enterprise/DatasmithCADImporter/Source/DatasmithPLMXMLTranslator/Private/DatasmithPlmXmlImporter.cpp b/Engine/Plugins/Enterprise/DatasmithCADImporter/Source/DatasmithPLMXMLTranslator/Private/DatasmithPlmXmlImporter.cpp
--- a/Engine/Plugins/Enterprise/DatasmithCADImporter/Source/DatasmithPLMXMLTranslator/Private/DatasmithPlmXmlImporter.cpp
+++ b/Engine/Plugins/Enterprise/DatasmithCADImporter/Source/DatasmithPLMXMLTranslator/Private/DatasmithPlmXmlImporter.cpp
@@ -115,11 +115,7 @@ namespace PlmXml
 	FString GetLabel(const FXmlNode* Node)
 	{
 		FString Name = Node->GetAttribute(TEXT("name"));
-		if (!Name.IsEmpty())
-		{
-			return Name;
-		}
-		return GetAttributeId(Node);
+		return FDatasmithUtils::SanitizeObjectName(Name.IsEmpty() ? GetAttributeId(Node) : Name);
 	}
 
 	FString GetReferenceIdFromUri(FString Uri)
@@ -536,12 +532,12 @@ namespace PlmXml
 	public:
 		virtual TSharedPtr<IDatasmithActorElement> GetActorElement() override
 		{
-			return Parent->GetActorElement();
+			return ActorElement ? ActorElement : Parent->GetActorElement();
 		}
 
-		virtual void AddChildActorElement(TSharedRef<IDatasmithActorElement> ActorElement) override
+		virtual void AddChildActorElement(TSharedRef<IDatasmithActorElement> InActorElement) override
 		{
-			Parent->AddChildActorElement(ActorElement);
+			Parent->AddChildActorElement(InActorElement);
 		}
 		
 		void AddChildProductInstance(FString Id, TSharedRef<FImportedInstanceTreeNode> InImportedInstanceTreeNode) override
@@ -549,7 +545,13 @@ namespace PlmXml
 			Parent->AddChildProductInstance(Id, InImportedInstanceTreeNode);
 		}
 
+		void SetActorElement(const TSharedRef<IDatasmithActorElement>& InActorElement)
+		{
+			ActorElement = InActorElement;
+			Parent->AddChildActorElement(InActorElement);
+		}
 
+		TSharedPtr<IDatasmithActorElement> ActorElement;
 		FProductRevisionView ProductRevisionView;
 	};
 
@@ -768,6 +770,20 @@ namespace PlmXml
 	{
 		FProductRevisionView ProductRevisionView = Context.ProductRevisionView;
 
+		if (!Context.GetActorElement()) // in case no ActorElement was created by parent(this happens when PRV is among rootRefs)
+		{
+			FString InstanceId = GetAttributeId(ProductRevisionView.Node);
+
+			FString DatasmithName = CreateDatasmithName(Context.CurrentDatasmithName, InstanceId);
+
+			TSharedRef<IDatasmithActorElement> ActorElement = FDatasmithSceneFactory::CreateActor(*DatasmithName);
+			Context.SetActorElement(ActorElement);
+		}
+		TSharedRef<IDatasmithActorElement> ActorElement = Context.GetActorElement().ToSharedRef();
+		// Label actor after ProductRevisionView
+		FString DatasmithLabel = GetLabel(ProductRevisionView.Node);
+		ActorElement->SetLabel(*DatasmithLabel);
+
 		// type: The type of PRV. This is one of assembly, minimal, wire, solid, sheet, general. Optional, though.
 		FString ProductRevisionViewType = ProductRevisionView.Node->GetAttribute(TEXT("type"));
 		FString ProductRevisionViewId = PlmXml::GetAttributeId(ProductRevisionView.Node);
@@ -803,7 +819,7 @@ namespace PlmXml
 			}
 		}
 
-		ParseUserDataToDatasmithMetadata(ProductRevisionView.Node, Context.GetActorElement().ToSharedRef(), Context.ImportContext.MainImportContext.DatasmithScene);
+		ParseUserDataToDatasmithMetadata(ProductRevisionView.Node, ActorElement, Context.ImportContext.MainImportContext.DatasmithScene);
 
 		TArray<FIdRef> InstanceRefs = GetAttributeIDREFS(ProductRevisionView.Node, TEXT("instanceRefs")); // xsd:IDREFS, unlike in ProductInstance, ProductRevisionView uses IDREFS, not Uri
 
@@ -862,7 +878,7 @@ namespace PlmXml
 		FString InstanceId = GetAttributeId(ProductInstance.Node);
 		
 		FString DatasmithName = CreateDatasmithName(ProductInstanceContext.CurrentDatasmithName, InstanceId);
-		FString DatasmithLabel = TEXT("ProductInstance_") + FDatasmithUtils::SanitizeObjectName(ProductInstance.Node->GetAttribute(TEXT("name")));
+		FString DatasmithLabel = GetLabel(ProductInstance.Node);
 		
 		TSharedRef<IDatasmithActorElement> ActorElement = FDatasmithSceneFactory::CreateActor(*DatasmithName);
 		ActorElement->SetLabel(*DatasmithLabel);
@@ -918,7 +934,7 @@ namespace PlmXml
 	{
 		PlmXml::FParsedPlmXml& ParsedPlmXml = ImportContext.ParsedPlmXml;
 		FString ProductRevisionUUID = GetAttributeId(ProductRevisionNode);
-		FString ProductRevisionLabel = TEXT("ProductRevision_") + GetLabel(ProductRevisionNode);
+		FString ProductRevisionLabel = GetLabel(ProductRevisionNode);
 		TSharedRef<IDatasmithActorElement> ProductRevisionActorElement = FDatasmithSceneFactory::CreateActor(*CreateDatasmithName(ParentUUID, ProductRevisionUUID));
 		ProductRevisionActorElement->SetLabel(*ProductRevisionLabel);
 		RootActorElement->AddChild(ProductRevisionActorElement, EDatasmithActorAttachmentRule::KeepRelativeTransform);
@@ -930,7 +946,7 @@ namespace PlmXml
 				const FXmlNode* AssociatedDataSetNode = ProductRevisionChildNode;
 
 				TSharedRef<IDatasmithActorElement> AssociatedDataSetActorElement = FDatasmithSceneFactory::CreateActor(*CreateDatasmithName(ProductRevisionActorElement->GetName(), GetAttributeId(AssociatedDataSetNode)));
-				AssociatedDataSetActorElement->SetLabel(*(TEXT("AssociatedDataSet_") + GetLabel(AssociatedDataSetNode)));
+				AssociatedDataSetActorElement->SetLabel(*GetLabel(AssociatedDataSetNode));
 				ProductRevisionActorElement->AddChild(AssociatedDataSetActorElement, EDatasmithActorAttachmentRule::KeepRelativeTransform);
 
 				TArray<FIdRef> DataSetIds = PlmXml::GetAttributeUriReferenceList(AssociatedDataSetNode, TEXT("dataSetRef"));
@@ -1023,13 +1039,12 @@ namespace PlmXml
 		// A variant per ProductView
 		for (const FXmlNode* ProductViewNode : ImportContext.ParsedPlmXml.ProductViewNodes)
 		{
-			const TSharedRef<IDatasmithActorElement> ProductViewActorElement = CreateActor(RootActorElementForProductViews, *GetAttributeId(ProductViewNode), *(TEXT("ProductView_") + GetLabel(ProductViewNode)));
+			const TSharedRef<IDatasmithActorElement> ProductViewActorElement = CreateActor(RootActorElementForProductViews, *GetAttributeId(ProductViewNode), *GetLabel(ProductViewNode));
 			
 			// Collect all actors used in this ProductView via Occurrence node
 			TSet<TSharedRef<IDatasmithActorElement>> ActorOptions;
 
 			// TODO: rootRefs
-			
 			for (const FXmlNode* ProductViewChildNode : ProductViewNode->GetChildrenNodes())
 			{
 				if (ProductViewChildNode->GetTag() == TEXT("Occurrence"))
@@ -1037,7 +1052,7 @@ namespace PlmXml
 					
 					// TODO: rootRefs - if present use it to collect occurrence tree, else take all occurrences
 					const FXmlNode* OccurrenceNode = ProductViewChildNode;
-					const TSharedRef<IDatasmithActorElement> OccurrenceActorElement = CreateActor(ProductViewActorElement, *GetAttributeId(OccurrenceNode), *(TEXT("Occurrence_") + GetLabel(OccurrenceNode)));
+					const TSharedRef<IDatasmithActorElement> OccurrenceActorElement = CreateActor(ProductViewActorElement, *GetAttributeId(OccurrenceNode), *GetLabel(OccurrenceNode));
 
 					ParseUserDataToDatasmithMetadata(OccurrenceNode, OccurrenceActorElement, ImportContext.DatasmithScene);
 					
