OpenNurbsTranslator: extrusion render meshes

diff --git a/Engine/Plugins/Enterprise/DatasmithCADImporter/Source/DatasmithOpenNurbsTranslator/Private/DatasmithOpenNurbsTranslator.cpp b/Engine/Plugins/Enterprise/DatasmithCADImporter/Source/DatasmithOpenNurbsTranslator/Private/DatasmithOpenNurbsTranslator.cpp
--- a/Engine/Plugins/Enterprise/DatasmithCADImporter/Source/DatasmithOpenNurbsTranslator/Private/DatasmithOpenNurbsTranslator.cpp
+++ b/Engine/Plugins/Enterprise/DatasmithCADImporter/Source/DatasmithOpenNurbsTranslator/Private/DatasmithOpenNurbsTranslator.cpp
@@ -3046,29 +3046,20 @@ bool FOpenNurbsTranslatorImpl::TranslateBRep(ON_Brep* Brep, const ON_3dmObjectAt
 #endif // CAD_LIBRARY
 	{
 		// .. Trying to load the mesh tessellated by Rhino
+		ON_Mesh BRepMesh;
+
 		ON_SimpleArray<const ON_Mesh*> RenderMeshes;
-		ON_SimpleArray<const ON_Mesh*> AnyMeshes;
 		Brep->GetMesh(ON::mesh_type::render_mesh, RenderMeshes);
-		Brep->GetMesh(ON::mesh_type::any_mesh, AnyMeshes);
 
 		// Aborting because there is no mesh associated with the BRep
-		if (RenderMeshes.Count() == 0 && AnyMeshes.Count() == 0)
+		if (RenderMeshes.Count() == 0)
 		{
 			MissingRenderMeshes.Add(MeshElement->GetLabel());
 			return false;
 		}
 
-		ON_Mesh BRepMesh;
-
-		if (RenderMeshes.Count() == AnyMeshes.Count())
-		{
-			BRepMesh.Append(RenderMeshes.Count(), RenderMeshes.Array());
-		}
-		else
-		{
-			BRepMesh.Append(AnyMeshes.Count(), AnyMeshes.Array());
-		}
-
+		BRepMesh.Append(RenderMeshes.Count(), RenderMeshes.Array());
+		
 		if (!DatasmithOpenNurbsTranslatorUtils::TranslateMesh(&BRepMesh, OutMesh, bHasNormal, ScalingFactor, Offset))
 		{
 			return false;
@@ -3183,11 +3174,26 @@ TOptional<FMeshDescription> FOpenNurbsTranslatorImpl::GetMeshDescription(TShared
 	{
 		bIsValid = SelectedTranslator->TranslateBRep(ON_Brep::Cast(Object.ObjectPtr), Object.Attributes, MeshDescription, MeshElement, UUID, bHasNormal);
 	}
-	else if (Object.ObjectPtr->IsKindOf(&ON_Extrusion::m_ON_Extrusion_class_rtti) )
+	else if (const ON_Extrusion* extrusion = ON_Extrusion::Cast(Object.ObjectPtr))
 	{
-		const ON_Extrusion *extrusion = ON_Extrusion::Cast(Object.ObjectPtr);
+		if (OpenNurbsOptions.Geometry == EDatasmithOpenNurbsBrepTessellatedSource::UseRenderMeshes)
+		{
+			ON_3dVector Offset = GetGeometryOffset(MeshElement);
+			if (const ON_Mesh* Mesh = extrusion->m_mesh_cache.Mesh(ON::mesh_type::render_mesh))
+			{
+				if (DatasmithOpenNurbsTranslatorUtils::TranslateMesh(Mesh, MeshDescription, bHasNormal, ScalingFactor, Offset))
+				{
+					return MeshDescription;
+				}
+			}
+			else
+			{
+				MissingRenderMeshes.Add(MeshElement->GetLabel());
+			}
+		}
+
 		ON_Brep brep;
-		if (extrusion != nullptr && extrusion->BrepForm(&brep) != nullptr)
+		if (extrusion->BrepForm(&brep) != nullptr)
 		{
 			bIsValid = SelectedTranslator->TranslateBRep(&brep, Object.Attributes, MeshDescription, MeshElement, UUID, bHasNormal);
 		}
