OpenNurbsTranslator: fixed misuse of mapping channel id

diff --git a/Engine/Plugins/Enterprise/DatasmithCADImporter/Source/DatasmithOpenNurbsTranslator/Private/DatasmithOpenNurbsTranslator.cpp b/Engine/Plugins/Enterprise/DatasmithCADImporter/Source/DatasmithOpenNurbsTranslator/Private/DatasmithOpenNurbsTranslator.cpp
--- a/Engine/Plugins/Enterprise/DatasmithCADImporter/Source/DatasmithOpenNurbsTranslator/Private/DatasmithOpenNurbsTranslator.cpp
+++ b/Engine/Plugins/Enterprise/DatasmithCADImporter/Source/DatasmithOpenNurbsTranslator/Private/DatasmithOpenNurbsTranslator.cpp
@@ -956,13 +956,12 @@ void FOpenNurbsTranslatorImpl::TranslateMaterialTable(const ON_ObjectArray<ON_Ma
 				// Extract texture mapping info
 				DatasmithMaterialsUtils::FUVEditParameters UVParameters;
 
-				uint8 ChannelIndex = 0;
-				if (!ON_Texture::IsBuiltInMappingChannel(Texture->m_mapping_channel_id))
-				{
-					// Non built-in channels start at 2 and use 1-based indexing
-					ChannelIndex = Texture->m_mapping_channel_id - 1;
-				}
-				UVParameters.ChannelIndex = ChannelIndex;
+				// Use cached texture coordinates(channel 0)
+				UVParameters.ChannelIndex = 0;
+				// Note - Texture->m_mapping_channel_id may used to get TextureMapping in order to generate texture coordinates from it.
+				// Specifically, from object's attributes get MappingRef - attributes.m_rendering_attributes.m_mappings
+				// Next, find channel within MappingRef's m_mapping_channels which m_mapping_channel_id is equal to Texture's
+				// Then m_mapping_id of that channel will be UUID in the UUIDToTextureMapping 
 
 				// Extract the UV tiling, offset and rotation angle from the UV transform matrix
 				FMatrix Matrix;
