The bug is due to a modification of the topology during the retesselation step.
To avoid it, sew option of retesselation do not any more split part into connected bodies (as at the end the splitted bodies are merged togeter). If sew splits body and a body does not have any material, the first material is used.

diff --git a/Engine/Plugins/Enterprise/DatasmithCADImporter/Source/CADInterfaces/Private/CoreTechFileParser.cpp b/Engine/Plugins/Enterprise/DatasmithCADImporter/Source/CADInterfaces/Private/CoreTechFileParser.cpp
--- a/Engine/Plugins/Enterprise/DatasmithCADImporter/Source/CADInterfaces/Private/CoreTechFileParser.cpp
+++ b/Engine/Plugins/Enterprise/DatasmithCADImporter/Source/CADInterfaces/Private/CoreTechFileParser.cpp
@@ -569,7 +569,7 @@ void FCoreTechFileParser::ReadMaterials()
 FCoreTechFileParser::FCoreTechFileParser(const FImportParameters& ImportParams, const FString& EnginePluginsPath, const FString& InCachePath)
 	: CachePath(InCachePath)
 	, ImportParameters(ImportParams)
-	, bSewRequired(ImportParameters.StitchingTechnique == StitchingSew)
+	, bFixTopoRequired(ImportParameters.StitchingTechnique != StitchingNone)
 {
 	CTKIO_InitializeKernel(ImportParameters.MetricUnit, *EnginePluginsPath);
 }
@@ -821,11 +821,11 @@ CT_FLAGS FCoreTechFileParser::SetCoreTechImportOption(const FString& MainFileExt
 	// Ask Kernel IO to complete or create missing topology
 	if (MainFileExt == TEXT("igs") || MainFileExt == TEXT("iges"))
 	{
+		bFixTopoRequired = false;
 		Flags |= CT_LOAD_FLAG_COMPLETE_TOPOLOGY;
 		if (ImportParameters.StitchingTechnique == StitchingSew)
 		{
 			Flags |= CT_LOAD_FLAG_SEARCH_NEW_TOPOLOGY;
-			bSewRequired = false;
 		}
 	}
 
@@ -969,9 +969,19 @@ bool FCoreTechFileParser::ReadComponent(CT_OBJECT_ID ComponentId, uint32 Default
 	TArray<CT_OBJECT_ID> Instances, Bodies;
 	GetInstancesAndBodies(ComponentId, Instances, Bodies);
 
-	if (!Instances.Num() && (Bodies.Num() > 1) && bSewRequired)
+	if (ImportParameters.StitchingTechnique != StitchingNone && bFixTopoRequired)
 	{
-		Repair(ComponentId, StitchingSew);
+		if (!Instances.Num() && Bodies.Num() > 1 && ImportParameters.StitchingTechnique == StitchingSew)
+		{
+			Repair(ComponentId, StitchingSew);
+		}
+		else
+		{
+			for (CT_OBJECT_ID Body : Bodies)
+			{
+				Repair(Body, StitchingHeal);
+			}
+		}
 		GetInstancesAndBodies(ComponentId, Instances, Bodies);
 	}
 
diff --git a/Engine/Plugins/Enterprise/DatasmithCADImporter/Source/CADInterfaces/Private/CoreTechTypes.cpp b/Engine/Plugins/Enterprise/DatasmithCADImporter/Source/CADInterfaces/Private/CoreTechTypes.cpp
--- a/Engine/Plugins/Enterprise/DatasmithCADImporter/Source/CADInterfaces/Private/CoreTechTypes.cpp
+++ b/Engine/Plugins/Enterprise/DatasmithCADImporter/Source/CADInterfaces/Private/CoreTechTypes.cpp
@@ -91,6 +91,59 @@ namespace CADLibrary
 		return CT_KERNEL_IO::AskTesselationParameters(max_face_sag, max_face_length, max_face_angle, max_curve_sag, max_curve_length, max_curve_angle, high_quality, vertex_type, normal_type, texture_type);
 	}
 
+	void GetAllBodies(CT_OBJECT_ID NodeId, TArray<CT_OBJECT_ID>& OutBodies)
+	{
+		CT_OBJECT_TYPE type;
+		CT_IO_ERROR Result = CT_OBJECT_IO::AskType(NodeId, type);
+		if (Result != IO_OK)
+		{
+			return;
+		}
+
+		switch (type)
+		{
+		case CT_ASSEMBLY_TYPE:
+		case CT_PART_TYPE:
+		case CT_COMPONENT_TYPE:
+		{
+			CT_LIST_IO Children;
+			Result = CT_COMPONENT_IO::AskChildren(NodeId, Children);
+			if (Result != IO_OK)
+			{
+				return;
+			}
+
+			// Iterate over the instances and call some function on each
+			Children.IteratorInitialize();
+			CT_OBJECT_ID Child;
+			while ((Child = Children.IteratorIter()) != 0)
+			{
+				GetAllBodies(Child, OutBodies);
+			}
+			break;
+		}
+
+		case CT_INSTANCE_TYPE:
+		{
+			CT_OBJECT_ID ChildId;
+			Result = CT_INSTANCE_IO::AskChild(NodeId, ChildId);
+			if (Result != IO_OK)
+			{
+				return;
+			}
+
+			GetAllBodies(ChildId, OutBodies);
+			break;
+		}
+
+		case CT_BODY_TYPE:
+		{
+			OutBodies.Add(NodeId);
+			break;
+		}
+		}
+	}
+
 	CT_IO_ERROR Repair(CT_OBJECT_ID MainObjectID, EStitchingTechnique StitchingTechnique, CT_DOUBLE SewingToleranceFactor)
 	{
 		switch (StitchingTechnique)
@@ -101,14 +154,25 @@ namespace CADLibrary
 			return CT_REPAIR_IO::Sew(MainObjectID, SewingToleranceFactor, CT_SEW_CREATE_BODIES_BY_TOPOLOGY);
 		}
 
+		// sew disconnected faces back together
+		// to be sure that there is no topology modification, this function has to be called body by body.
+		// CT_REPAIR_IO::Sew is prefered to CT_REPAIR_IO::heal because CT_REPAIR_IO::heal is too restricive
 		case EStitchingTechnique::StitchingHeal:
 		{
-			// Heal disconnected faces back together
-			CT_UINT32 StitchingToleranceFactor = 0; // topo only if deformation=0, factor of 'tolerance' value
-			CT_LOGICAL   OnlyVisible = CT_TRUE;
-			CT_LOGICAL   MergeBeforeHealing = CT_FALSE;
-			CT_LOGICAL   RemoveTinyObjects = CT_TRUE;
-			return CT_REPAIR_IO::Heal(MainObjectID, StitchingToleranceFactor, OnlyVisible, MergeBeforeHealing, RemoveTinyObjects);
+			uint32 BodiesNum;
+			CT_IO_ERROR Result = CT_KERNEL_IO::AskNbObjectsType(BodiesNum, CT_BODY_TYPE);
+			if (Result != IO_OK)
+			{
+				return Result;
+			}
+
+			TArray<CT_OBJECT_ID> OutBodies;
+			OutBodies.Reserve(BodiesNum);
+			GetAllBodies(MainObjectID, OutBodies);
+			for (CT_OBJECT_ID BodyId : OutBodies)
+			{
+				CT_REPAIR_IO::Sew(BodyId, SewingToleranceFactor, CT_SEW_CREATE_BODIES_BY_LAYER);
+			}
 		}
 		}
 		return CT_IO_ERROR::IO_OK;
diff --git a/Engine/Plugins/Enterprise/DatasmithCADImporter/Source/CADInterfaces/Public/CoreTechFileParser.h b/Engine/Plugins/Enterprise/DatasmithCADImporter/Source/CADInterfaces/Public/CoreTechFileParser.h
--- a/Engine/Plugins/Enterprise/DatasmithCADImporter/Source/CADInterfaces/Public/CoreTechFileParser.h
+++ b/Engine/Plugins/Enterprise/DatasmithCADImporter/Source/CADInterfaces/Public/CoreTechFileParser.h
@@ -133,7 +133,7 @@ protected:
 	bool bNeedSaveCTFile = false;
 
 	const FImportParameters& ImportParameters;
-	bool bSewRequired;
+	bool bFixTopoRequired;
 };
 
 #endif // CAD_INTERFACE
diff --git a/Engine/Plugins/Enterprise/DatasmithCADImporter/Source/CADLibrary/Private/CoreTechMeshLoader.cpp b/Engine/Plugins/Enterprise/DatasmithCADImporter/Source/CADLibrary/Private/CoreTechMeshLoader.cpp
--- a/Engine/Plugins/Enterprise/DatasmithCADImporter/Source/CADLibrary/Private/CoreTechMeshLoader.cpp
+++ b/Engine/Plugins/Enterprise/DatasmithCADImporter/Source/CADLibrary/Private/CoreTechMeshLoader.cpp
@@ -27,16 +27,22 @@ bool CoreTechMeshLoader::LoadFile(const FString& FileName, FMeshDescription& Mes
 		return false;
 	}
 
-	Result = Repair(MainObjectID, ImportParameters.StitchingTechnique);
+	if (ImportParameters.StitchingTechnique != StitchingNone)
+	{
+		Result = Repair(MainObjectID, StitchingHeal);
+	}
 	Result = SetCoreTechTessellationState(ImportParameters);
 
-	bool bNeedSwapOrientation = ImportParameters.StitchingTechnique != EStitchingTechnique::StitchingSew ? MeshParameters.bNeedSwapOrientation : false;
+	bool bNeedSwapOrientation = ImportParameters.StitchingTechnique != EStitchingTechnique::StitchingNone ? MeshParameters.bNeedSwapOrientation : false;
 
-	// Parse CoreTech model's hierarchy to collect body objects.
-	ExtractComponent(MainObjectID);
+	uint32 BodiesNum;
+	Result = CT_KERNEL_IO::AskNbObjectsType(BodiesNum, CT_BODY_TYPE);
+	if (Result != IO_OK)
+	{
+		return false;
+	}
 
-	// Nothing useful in the file, abort
-	if (BodySet.Num() == 0)
+	if (BodiesNum == 0)
 	{
 		return false;
 	}
diff --git a/Engine/Plugins/Enterprise/DatasmithCADImporter/Source/DatasmithCoreTechExtension/Private/CoreTechRetessellateAction.cpp b/Engine/Plugins/Enterprise/DatasmithCADImporter/Source/DatasmithCoreTechExtension/Private/CoreTechRetessellateAction.cpp
--- a/Engine/Plugins/Enterprise/DatasmithCADImporter/Source/DatasmithCoreTechExtension/Private/CoreTechRetessellateAction.cpp
+++ b/Engine/Plugins/Enterprise/DatasmithCADImporter/Source/DatasmithCoreTechExtension/Private/CoreTechRetessellateAction.cpp
@@ -203,12 +203,16 @@ bool FCoreTechRetessellate_Impl::ApplyOnOneAsset(UStaticMesh& StaticMesh, UCoreT
 			{
 				TPolygonGroupAttributesConstRef<FName> MaterialSlotNames = MeshDescriptionAttributes.GetPolygonGroupMaterialSlotNames();
 				FMeshSectionInfoMap& SectionInfoMap = StaticMesh.GetSectionInfoMap();
-				ensure(MeshDescription.PolygonGroups().Num() == SectionInfoMap.GetSectionNumber(0));
 
 				for (FPolygonGroupID PolygonGroupID : MeshDescription.PolygonGroups().GetElementIDs())
 				{
 					FMeshSectionInfo Section = SectionInfoMap.Get(0, PolygonGroupID.GetValue());
-					Section.MaterialIndex = StaticMesh.GetMaterialIndex(MaterialSlotNames[PolygonGroupID]);
+					int32 MaterialIndex = StaticMesh.GetMaterialIndex(MaterialSlotNames[PolygonGroupID]);
+					if (MaterialIndex < 0)
+					{
+						MaterialIndex = 0;
+					}
+					Section.MaterialIndex = MaterialIndex;
 					SectionInfoMap.Set(0, PolygonGroupID.GetValue(), Section);
 				}
 			}
