OpenNurbsTranslator: support multimaterial(per-brepface material channels) for imported render meshes

diff --git a/Engine/Plugins/Enterprise/DatasmithCADImporter/Source/DatasmithOpenNurbsTranslator/Private/DatasmithOpenNurbsTranslator.cpp b/Engine/Plugins/Enterprise/DatasmithCADImporter/Source/DatasmithOpenNurbsTranslator/Private/DatasmithOpenNurbsTranslator.cpp
--- a/Engine/Plugins/Enterprise/DatasmithCADImporter/Source/DatasmithOpenNurbsTranslator/Private/DatasmithOpenNurbsTranslator.cpp
+++ b/Engine/Plugins/Enterprise/DatasmithCADImporter/Source/DatasmithOpenNurbsTranslator/Private/DatasmithOpenNurbsTranslator.cpp
@@ -265,15 +265,13 @@ namespace DatasmithOpenNurbsTranslatorUtils
 		FVector2D UVCoords[3];
 	};
 
-	bool TranslateMesh(const ON_Mesh *Mesh, FMeshDescription& MeshDescription, bool& bHasNormal, double ScalingFactor, const ON_3dVector& Offset)
+	bool TranslateMesh(const ON_Mesh** Meshes, int MeshCount, FMeshDescription& MeshDescription, bool& bHasNormal, double ScalingFactor, const ON_3dVector& Offset, bool bHasFaceMaterialChannel, int32* FaceMaterialChannel)
 	{
-		// Ref. GP3DMVisitorImpl::visitMesh
-		// Ref. FGPureMeshInterface::CreateMesh
-		if (Mesh == nullptr || Mesh->VertexCount() == 0 || Mesh->FaceCount() == 0)
+		if (!Meshes || !MeshCount)
 		{
 			return false;
 		}
-
+		
 		MeshDescription.Empty();
 
 		TVertexAttributesRef<FVector> VertexPositions = MeshDescription.VertexAttributes().GetAttributesRef<FVector>(MeshAttribute::Vertex::Position);
@@ -287,199 +285,224 @@ namespace DatasmithOpenNurbsTranslatorUtils
 			return false;
 		}
 
-		bool bHasPackedTexCoords = HasPackedTextureRegion(*Mesh);
-		bool bHasUVData = Mesh->HasTextureCoordinates() || bHasPackedTexCoords;
-
-		int VertexCount = Mesh->VertexCount();
-		TArray<FNode> Nodes;
-		Nodes.Reserve(VertexCount);
-
-		for (int Index = 0; Index < VertexCount; ++Index)
+		FPolygonGroupID PolyGroupId;
+		if (bHasFaceMaterialChannel)
 		{
-			ON_3fPoint p1 = Mesh->m_V[Index] + ON_3fVector(Offset);
-			Nodes.Add(FNode(p1.x * ScalingFactor, p1.y * ScalingFactor, p1.z * ScalingFactor));
+			MeshDescription.ReserveNewPolygonGroups(MeshCount);
 		}
-
-		int faceCount = Mesh->FaceCount();
-		TArray<FVector2D> UvCoords;
-		TArray<FFace> Faces;
-		Faces.Reserve(faceCount);
-
-		for (int Index = 0; Index < faceCount; ++Index)
+		else
 		{
-			const ON_MeshFace& meshFace = Mesh->m_F[Index];
+			MeshDescription.ReserveNewPolygonGroups(1);
+			PolyGroupId = MeshDescription.CreatePolygonGroup();
+			PolygonGroupImportedMaterialSlotNames[PolyGroupId] = DatasmithMeshHelper::DefaultSlotName(0);
+		}
+		// At least one UV set must exist.
+		VertexInstanceUVs.SetNumIndices(1);
 
-			FNode& n1 = Nodes[meshFace.vi[0]];
-			FNode& n2 = Nodes[meshFace.vi[1]];
-			FNode& n3 = Nodes[meshFace.vi[2]];
-			Faces.Push(FFace(n1, n2, n3));
+		for (int32 MeshIndex = 0; MeshIndex < MeshCount; ++MeshIndex)
+		{
+			const ON_Mesh* Mesh = Meshes[MeshIndex];
 
-			if (Mesh->HasFaceNormals())
+			if (Mesh == nullptr || Mesh->VertexCount() == 0 || Mesh->FaceCount() == 0)
 			{
-				const ON_3fVector& NormalN1 = Mesh->m_FN[Index];
-
-				n1.SetNormal(NormalN1);
-				n2.SetNormal(NormalN1);
-				n3.SetNormal(NormalN1);
-
-				bHasNormal = true;
+				return false;
 			}
-			else if (Mesh->HasVertexNormals())
-			{
-				const ON_3fVector& NormalN1 = Mesh->m_N[meshFace.vi[0]];
-				const ON_3fVector& NormalN2 = Mesh->m_N[meshFace.vi[1]];
-				const ON_3fVector& NormalN3 = Mesh->m_N[meshFace.vi[2]];
 
-				n1.SetNormal(NormalN1);
-				n2.SetNormal(NormalN2);
-				n3.SetNormal(NormalN3);
+			bool bHasPackedTexCoords = HasPackedTextureRegion(*Mesh);
+			bool bHasUVData = Mesh->HasTextureCoordinates() || bHasPackedTexCoords;
 
-				bHasNormal = true;
-			}
+			int VertexCount = Mesh->VertexCount();
+			TArray<FNode> Nodes;
+			Nodes.Reserve(VertexCount);
 
-			if (bHasUVData)
+			for (int Index = 0; Index < VertexCount; ++Index)
 			{
-				UvCoords.Add(GetMeshTexCoords(Mesh, VertexCount, meshFace.vi[0], bHasPackedTexCoords));
-				UvCoords.Add(GetMeshTexCoords(Mesh, VertexCount, meshFace.vi[1], bHasPackedTexCoords));
-				UvCoords.Add(GetMeshTexCoords(Mesh, VertexCount, meshFace.vi[2], bHasPackedTexCoords));
+				ON_3fPoint p1 = Mesh->m_V[Index] + ON_3fVector(Offset);
+				Nodes.Add(FNode(p1.x * ScalingFactor, p1.y * ScalingFactor, p1.z * ScalingFactor));
 			}
 
-			if (meshFace.IsQuad())
+			int faceCount = Mesh->FaceCount();
+			TArray<FVector2D> UvCoords;
+			TArray<FFace> Faces;
+			Faces.Reserve(faceCount);
+
+			for (int Index = 0; Index < faceCount; ++Index)
 			{
-				FNode& n4 = Nodes[meshFace.vi[3]];
-				Faces.Push(FFace(n1, n3, n4));
+				const ON_MeshFace& meshFace = Mesh->m_F[Index];
+
+				FNode& n1 = Nodes[meshFace.vi[0]];
+				FNode& n2 = Nodes[meshFace.vi[1]];
+				FNode& n3 = Nodes[meshFace.vi[2]];
+				Faces.Push(FFace(n1, n2, n3));
 
 				if (Mesh->HasFaceNormals())
 				{
-					const ON_3fVector& NormalN4 = Mesh->m_FN[Index];
-					n4.SetNormal(NormalN4);
+					const ON_3fVector& NormalN1 = Mesh->m_FN[Index];
+
+					n1.SetNormal(NormalN1);
+					n2.SetNormal(NormalN1);
+					n3.SetNormal(NormalN1);
+
+					bHasNormal = true;
 				}
 				else if (Mesh->HasVertexNormals())
 				{
-					const ON_3fVector& NormalN4 = Mesh->m_N[meshFace.vi[3]];
-					n4.SetNormal(NormalN4);
+					const ON_3fVector& NormalN1 = Mesh->m_N[meshFace.vi[0]];
+					const ON_3fVector& NormalN2 = Mesh->m_N[meshFace.vi[1]];
+					const ON_3fVector& NormalN3 = Mesh->m_N[meshFace.vi[2]];
+
+					n1.SetNormal(NormalN1);
+					n2.SetNormal(NormalN2);
+					n3.SetNormal(NormalN3);
+
+					bHasNormal = true;
 				}
 
 				if (bHasUVData)
 				{
 					UvCoords.Add(GetMeshTexCoords(Mesh, VertexCount, meshFace.vi[0], bHasPackedTexCoords));
+					UvCoords.Add(GetMeshTexCoords(Mesh, VertexCount, meshFace.vi[1], bHasPackedTexCoords));
 					UvCoords.Add(GetMeshTexCoords(Mesh, VertexCount, meshFace.vi[2], bHasPackedTexCoords));
-					UvCoords.Add(GetMeshTexCoords(Mesh, VertexCount, meshFace.vi[3], bHasPackedTexCoords));
 				}
-			}
-		}
 
-		if (bHasUVData)
-		{
-			for (int Index = 0; Index < Faces.Num(); ++Index)
-			{
-				FFace& Face = Faces[Index];
-				for (int iTriangleNode = 0; iTriangleNode < 3; ++iTriangleNode)
+				if (meshFace.IsQuad())
 				{
-					Face.UVCoords[iTriangleNode] = FVector2D(UvCoords[iTriangleNode + 3 * Index][0], UvCoords[iTriangleNode + 3 * Index][1]);
-				}
-			}
-		}
-
-		// Fill out the MeshDescription with the processed data from the ON_Mesh
-		// Ref. FDatasmithMeshUtils::ToMeshDescription
-
-		// Reserve space for attributes
-		// At this point, all the faces are triangles
-		const int32 TriangleCount = Faces.Num();
-		const int32 VertexInstanceCount = 3 * TriangleCount;
-
-		MeshDescription.ReserveNewVertices(VertexCount);
-		MeshDescription.ReserveNewVertexInstances(VertexInstanceCount);
-		MeshDescription.ReserveNewEdges(VertexInstanceCount);
-		MeshDescription.ReserveNewPolygons(TriangleCount);
-
-		// Assume one material per mesh, no partitioning
-		MeshDescription.ReserveNewPolygonGroups(1);
-
-		FPolygonGroupID PolyGroupId = MeshDescription.CreatePolygonGroup();
-		FName ImportedSlotName = *FString::FromInt(0); // No access to DatasmithMeshHelper::DefaultSlotName
-		PolygonGroupImportedMaterialSlotNames[PolyGroupId] = ImportedSlotName;
+					FNode& n4 = Nodes[meshFace.vi[3]];
+					Faces.Push(FFace(n1, n3, n4));
 
-		// At least one UV set must exist.
-		VertexInstanceUVs.SetNumIndices(1);
-
-		// Set vertex positions
-		TMap<const FNode*, int> NodeToIndex;
-		for (int Index = 0; Index < VertexCount; ++Index)
-		{
-			const FNode& Node = Nodes[Index];
-			const FVector& Pos = Node.Vertex;
-			NodeToIndex.Add(&Node, Index);
-
-			// Fill the vertex array
-			FVertexID AddedVertexId = MeshDescription.CreateVertex();
-			VertexPositions[AddedVertexId] = FVector(- Pos.X, Pos.Y, Pos.Z);
-		}
-
-		int32 VertexIndices[3];
+					if (Mesh->HasFaceNormals())
+					{
+						const ON_3fVector& NormalN4 = Mesh->m_FN[Index];
+						n4.SetNormal(NormalN4);
+					}
+					else if (Mesh->HasVertexNormals())
+					{
+						const ON_3fVector& NormalN4 = Mesh->m_N[meshFace.vi[3]];
+						n4.SetNormal(NormalN4);
+					}
 
-		const int32 CornerCount = 3; // only triangles
-		FVector CornerPositions[3];
-		TArray<FVertexInstanceID> CornerVertexInstanceIDs;
-		CornerVertexInstanceIDs.SetNum(3);
-		FVertexID CornerVertexIDs[3];
+					if (bHasUVData)
+					{
+						UvCoords.Add(GetMeshTexCoords(Mesh, VertexCount, meshFace.vi[0], bHasPackedTexCoords));
+						UvCoords.Add(GetMeshTexCoords(Mesh, VertexCount, meshFace.vi[2], bHasPackedTexCoords));
+						UvCoords.Add(GetMeshTexCoords(Mesh, VertexCount, meshFace.vi[3], bHasPackedTexCoords));
+					}
+				}
+			}
 
-		// Get per-triangle data: indices, normals and uvs
-		//int WedgeIndex = 0;
-		for (int FaceIndex = 0; FaceIndex < TriangleCount; ++FaceIndex)
-		{
-			const FFace& Face = Faces[FaceIndex];
-			for (int32 CornerIndex = 0; CornerIndex < CornerCount; ++CornerIndex)
+			if (bHasUVData)
 			{
-				const FNode& FaceNode = *Face.Nodes[CornerIndex];
-				VertexIndices[CornerIndex] = NodeToIndex[&FaceNode];
-
-				CornerVertexIDs[CornerIndex] = FVertexID(VertexIndices[CornerIndex]);
-				CornerPositions[CornerIndex] = VertexPositions[CornerVertexIDs[CornerIndex]];
+				for (int Index = 0; Index < Faces.Num(); ++Index)
+				{
+					FFace& Face = Faces[Index];
+					for (int iTriangleNode = 0; iTriangleNode < 3; ++iTriangleNode)
+					{
+						Face.UVCoords[iTriangleNode] = FVector2D(UvCoords[iTriangleNode + 3 * Index][0], UvCoords[iTriangleNode + 3 * Index][1]);
+					}
+				}
 			}
 
-			// Skip degenerated polygons
-			FVector RawNormal = ((CornerPositions[1] - CornerPositions[2]) ^ (CornerPositions[0] - CornerPositions[2]));
-			if (RawNormal.SizeSquared() < SMALL_NUMBER)
+			// Fill out the MeshDescription with the processed data from the ON_Mesh
+			// Ref. FDatasmithMeshUtils::ToMeshDescription
+
+			// Reserve space for attributes
+			// At this point, all the faces are triangles
+			const int32 TriangleCount = Faces.Num();
+			const int32 VertexInstanceCount = 3 * TriangleCount;
+
+			int32 VertexIndexBase = MeshDescription.Vertices().Num();
+			
+			MeshDescription.ReserveNewVertices(VertexCount);
+			
+			// Do not reserve vertex instances -  UUVGenerationFlattenMapping::GetOverlappingCornersRemapping is crashing because it assumes that vertex instances has no gaps
+			// but they appear when we skip triangles(degenerates) between meshes(if there are few meshes combine mesh description)
+			// MeshDescription.ReserveNewVertexInstances(VertexInstanceCount);
+			MeshDescription.ReserveNewEdges(VertexInstanceCount);
+			MeshDescription.ReserveNewPolygons(TriangleCount);
+
+			if (bHasFaceMaterialChannel)
 			{
-				continue; // this will leave holes...
+				PolyGroupId = MeshDescription.CreatePolygonGroup();
+				PolygonGroupImportedMaterialSlotNames[PolyGroupId] = DatasmithMeshHelper::DefaultSlotName(FaceMaterialChannel[MeshIndex]);
 			}
 
-			// Create Vertex instances and set their attributes
-			for (int32 CornerIndex = 0; CornerIndex < CornerCount; ++CornerIndex)
+			// Set vertex positions
+			TMap<const FNode*, int> NodeToIndex;
+			for (int Index = 0; Index < VertexCount; ++Index)
 			{
-				CornerVertexInstanceIDs[CornerIndex] = MeshDescription.CreateVertexInstance(CornerVertexIDs[CornerIndex]);
+				const FNode& Node = Nodes[Index];
+				const FVector& Pos = Node.Vertex;
+				NodeToIndex.Add(&Node, VertexIndexBase + Index);
+
+				// Fill the vertex array
+				FVertexID AddedVertexId = MeshDescription.CreateVertex();
+				VertexPositions[AddedVertexId] = FVector(-Pos.X, Pos.Y, Pos.Z);
+			}
 
-				const FNode& FaceNode = *Face.Nodes[CornerIndex];
+			int32 VertexIndices[3];
 
-				// Set the normal
-				FVector UENormal = FDatasmithUtils::ConvertVector(FDatasmithUtils::EModelCoordSystem::ZUp_RightHanded, FaceNode.Normal);
-				UENormal = UENormal.GetSafeNormal();
+			const int32 CornerCount = 3; // only triangles
+			FVector CornerPositions[3];
+			TArray<FVertexInstanceID> CornerVertexInstanceIDs;
+			CornerVertexInstanceIDs.SetNum(3);
+			FVertexID CornerVertexIDs[3];
 
-				// Check to see if normal is correct. If not replace by face's normal
-				if (UENormal.IsNormalized())
+			// Get per-triangle data: indices, normals and uvs
+			//int WedgeIndex = 0;
+			for (int FaceIndex = 0; FaceIndex < TriangleCount; ++FaceIndex)
+			{
+				const FFace& Face = Faces[FaceIndex];
+				for (int32 CornerIndex = 0; CornerIndex < CornerCount; ++CornerIndex)
 				{
-					VertexInstanceNormals[CornerVertexInstanceIDs[CornerIndex]] = UENormal;
+					const FNode& FaceNode = *Face.Nodes[CornerIndex];
+					VertexIndices[CornerIndex] = NodeToIndex[&FaceNode];
+
+					CornerVertexIDs[CornerIndex] = FVertexID(VertexIndices[CornerIndex]);
+					CornerPositions[CornerIndex] = VertexPositions[CornerVertexIDs[CornerIndex]];
 				}
-				else
+
+				// Skip degenerated polygons
+				FVector RawNormal = ((CornerPositions[1] - CornerPositions[2]) ^ (CornerPositions[0] - CornerPositions[2]));
+				if (RawNormal.SizeSquared() < SMALL_NUMBER)
 				{
-					// TODO: Check if this case could happen in Rhino
+					continue; // this will leave holes...
 				}
 
-				// Set the UV
-				if (bHasUVData)
+				// Create Vertex instances and set their attributes
+				for (int32 CornerIndex = 0; CornerIndex < CornerCount; ++CornerIndex)
 				{
-					const FVector2D& UVValues = Face.UVCoords[CornerIndex];
-					if (!UVValues.ContainsNaN())
+					CornerVertexInstanceIDs[CornerIndex] = MeshDescription.CreateVertexInstance(CornerVertexIDs[CornerIndex]);
+
+					const FNode& FaceNode = *Face.Nodes[CornerIndex];
+
+					// Set the normal
+					FVector UENormal = FDatasmithUtils::ConvertVector(FDatasmithUtils::EModelCoordSystem::ZUp_RightHanded, FaceNode.Normal);
+					UENormal = UENormal.GetSafeNormal();
+
+					// Check to see if normal is correct. If not replace by face's normal
+					if (UENormal.IsNormalized())
 					{
-						// Convert UVs - bottom-left Rhino's texture origin to Unreal's top-left
-						VertexInstanceUVs.Set(CornerVertexInstanceIDs[CornerIndex], 0, FVector2D(UVValues[0], 1.f - UVValues[1]));
+						VertexInstanceNormals[CornerVertexInstanceIDs[CornerIndex]] = UENormal;
+					}
+					else
+					{
+						// TODO: Check if this case could happen in Rhino
+					}
+
+					// Set the UV
+					if (bHasUVData)
+					{
+						const FVector2D& UVValues = Face.UVCoords[CornerIndex];
+						if (!UVValues.ContainsNaN())
+						{
+							// Convert UVs - bottom-left Rhino's texture origin to Unreal's top-left
+							VertexInstanceUVs.Set(CornerVertexInstanceIDs[CornerIndex], 0, FVector2D(UVValues[0], 1.f - UVValues[1]));
+						}
 					}
 				}
-			}
 
-			const FPolygonID NewPolygonID = MeshDescription.CreatePolygon(PolyGroupId, CornerVertexInstanceIDs);
+				const FPolygonID NewPolygonID = MeshDescription.CreatePolygon(PolyGroupId, CornerVertexInstanceIDs);
+			}
 		}
 
 		// Build edge meta data
@@ -487,6 +510,11 @@ namespace DatasmithOpenNurbsTranslatorUtils
 
 		return true;
 	}
+	
+	bool TranslateMesh(const ON_Mesh* Mesh, FMeshDescription& MeshDescription, bool& bHasNormal, double ScalingFactor, const ON_3dVector& Offset)
+	{
+		return TranslateMesh(&Mesh, 1, MeshDescription, bHasNormal, ScalingFactor, Offset, false, nullptr);
+	}
 
 	void PropagateLayers(const TSharedPtr<IDatasmithActorElement>& ActorElement, const FString& LayerNames)
 	{
@@ -648,8 +676,13 @@ private:
 	TSharedPtr<IDatasmithActorElement> GetPointActorElement(const FOpenNurbsObjectWrapper& Object);
 	TSharedPtr<IDatasmithMeshElement> GetMeshElement(const FOpenNurbsObjectWrapper& Object, const FString& Uuid, const FString& Label);
 
-	TSharedPtr<IDatasmithBaseMaterialElement> FindMaterial(const FOpenNurbsObjectWrapper& Object);
-	TSharedPtr<IDatasmithBaseMaterialElement> GetObjectMaterial(const FOpenNurbsObjectWrapper& Object);
+	struct FMaterial
+	{
+		TSharedPtr<IDatasmithBaseMaterialElement> MaterialElement;
+		const ON_Material* OpenNurbsMaterialPtr = nullptr;
+	};
+	FMaterial GetObjectMaterial(const FOpenNurbsObjectWrapper& Object);
+	ON_Material* GetOpenNurbsMaterial(int MaterialIndex);
 	TSharedPtr<IDatasmithBaseMaterialElement> GetMaterial(int MaterialIndex);
 	TSharedPtr<IDatasmithBaseMaterialElement> GetDefaultMaterial();
 
@@ -657,6 +690,8 @@ private:
 	FString GetLayerName(const TSharedPtr<IDatasmithActorElement>& LayerElement);
 	void SetLayers(const TSharedPtr<IDatasmithActorElement>& ActorElement, const FOpenNurbsObjectWrapper& Object);
 	void SetTags(const TSharedPtr<IDatasmithActorElement>& ActorElement, const FOpenNurbsObjectWrapper& Object);
+	
+	void SetMaterialToMeshElement(TSharedPtr<IDatasmithMeshElement> MeshElement, TSharedPtr<IDatasmithBaseMaterialElement> MaterialElement, int32 SlotId);
 
 	bool TranslateBRep(ON_Brep* brep, const ON_3dmObjectAttributes& Attributes, FMeshDescription& OutMesh, const TSharedRef< IDatasmithMeshElement >& MeshElement, const FString& Name, bool& bHasNormal);
 
@@ -736,13 +771,15 @@ private:
 	// Materials
 	TMap<FMD5Hash, TSharedPtr<IDatasmithBaseMaterialElement>> HashToMaterial;
 	TMap<int, TSharedPtr<IDatasmithBaseMaterialElement>> MaterialIndexToMaterial;
+	TMap<int, int> MaterialIndexToMaterialTableIndex;
+	TMap<ON_UUID, int> MaterialIdToIndex;
 	TSet<TSharedPtr<IDatasmithBaseMaterialElement>> UsedMaterials;
 	TSharedPtr<IDatasmithBaseMaterialElement> DefaultMaterial;
 
 	// Layers
 	TMap<ON_UUID, TSharedPtr<IDatasmithActorElement>> LayerUUIDToContainer;
 	TMap<int, TSharedPtr<IDatasmithActorElement>> LayerIndexToContainer;
-	TMap<int, TSharedPtr<IDatasmithBaseMaterialElement>> LayerIndexToMaterial;
+	TMap<int, int> LayerIndexToMaterialIndex;
 	TMap<TSharedPtr<IDatasmithActorElement>, FString> LayerNames;
 	TSet<int> HiddenLayersIndices;
 
@@ -814,6 +851,10 @@ void FOpenNurbsTranslatorImpl::TranslateMaterialTable(const ON_ObjectArray<ON_Ma
 		FMD5Hash Hash = DatasmithOpenNurbsTranslatorUtils::ComputeMaterialHash(OpenNurbsMaterial);
 		TSharedPtr<IDatasmithBaseMaterialElement>* MaterialPtr = HashToMaterial.Find(Hash);
 
+		MaterialIndexToMaterialTableIndex.Add(OpenNurbsMaterial.Index(), Index);
+
+		MaterialIdToIndex.Add(OpenNurbsMaterial.Id(), OpenNurbsMaterial.Index());
+
 		if (MaterialPtr)
 		{
 			MaterialIndexToMaterial.Add(OpenNurbsMaterial.Index(), *MaterialPtr);
@@ -1146,9 +1187,7 @@ void FOpenNurbsTranslatorImpl::TranslateLayerTable(const ON_ObjectArray<ON_Layer
 		// Use the layer's render material, no fallback on the display color
 		if (CurrentLayer.RenderMaterialIndex() != -1)
 		{
-			// The layer's render material index should match a material previously translated from the material table
-			TSharedPtr<IDatasmithBaseMaterialElement> MaterialElement = GetMaterial(CurrentLayer.RenderMaterialIndex());
-			LayerIndexToMaterial.Add(CurrentLayer.Index(), MaterialElement);
+			LayerIndexToMaterialIndex.Add(CurrentLayer.Index(), CurrentLayer.RenderMaterialIndex());
 		}
 	}
 }
@@ -1955,6 +1994,21 @@ TSharedPtr<IDatasmithActorElement> FOpenNurbsTranslatorImpl::GetPointActorElemen
 	return ActorElement;
 }
 
+void FOpenNurbsTranslatorImpl::SetMaterialToMeshElement(TSharedPtr<IDatasmithMeshElement> MeshElement, TSharedPtr<IDatasmithBaseMaterialElement> MaterialElement, int32 SlotId)
+{
+	if (MaterialElement.IsValid())
+	{
+		MeshElement->SetMaterial(MaterialElement->GetName(), SlotId);
+
+		// And add it to the Datasmith Scene as needed
+		if (!UsedMaterials.Contains(MaterialElement))
+		{
+			UsedMaterials.Add(MaterialElement);
+			Scene->AddMaterial(MaterialElement);
+		}
+	}
+}
+
 TSharedPtr<IDatasmithMeshElement> FOpenNurbsTranslatorImpl::GetMeshElement(const FOpenNurbsObjectWrapper& Object, const FString& Uuid, const FString& Label)
 {
 	// Ref. FDatasmithCADImporter::FindOrAddMeshElement
@@ -1972,10 +2026,28 @@ TSharedPtr<IDatasmithMeshElement> FOpenNurbsTranslatorImpl::GetMeshElement(const
 	MeshElement->SetLabel(*Label);
 	MeshElement->SetLightmapSourceUV(-1);
 
-	TSharedPtr<IDatasmithBaseMaterialElement> Material = FindMaterial(Object);
-	if (Material.IsValid() && BaseOptions.bIncludeMaterial)
+	if (BaseOptions.bIncludeMaterial)
 	{
-		MeshElement->SetMaterial(Material->GetName(), 0);
+		FMaterial Material = GetObjectMaterial(Object);
+
+		SetMaterialToMeshElement(MeshElement, Material.MaterialElement, 0);
+
+		// Material slots(OpenNurbs 'channels')
+		// https://developer.rhino3d.com/guides/opennurbs/reading-per-face-render-materials/
+		if (Material.OpenNurbsMaterialPtr)
+		{
+			const ON_Material& OpenNurbsMaterial = *Material.OpenNurbsMaterialPtr;
+			for (int MaterialChannel = 0; MaterialChannel < OpenNurbsMaterial.m_material_channel.Count(); ++MaterialChannel)
+			{
+				const ON_UuidIndex& MaterialIdx = OpenNurbsMaterial.m_material_channel[MaterialChannel];
+				if (int* IndexPtr = MaterialIdToIndex.Find(MaterialIdx.m_id))
+				{
+					// In OpenNurbs, m_material_channel is indexed by "m_face_material_channel-1"(where m_face_material_channel==0 means 'parent' material)
+					// We put "channel material" into slot "channel+1" so 'parent' material can go into slot 0 and can index slots by unmodified m_face_material_channel
+					SetMaterialToMeshElement(MeshElement, GetMaterial(*IndexPtr), MaterialChannel + 1);
+				}
+			}
+		}
 	}
 
 	Scene->AddMesh(MeshElement);
@@ -2005,22 +2077,7 @@ TSharedPtr<IDatasmithMeshElement> FOpenNurbsTranslatorImpl::GetMeshElement(const
 	return MeshElement;
 }
 
-TSharedPtr<IDatasmithBaseMaterialElement> FOpenNurbsTranslatorImpl::FindMaterial(const FOpenNurbsObjectWrapper& Object)
-{
-	// Find a previously translated material for the Object
-	TSharedPtr<IDatasmithBaseMaterialElement> Material = GetObjectMaterial(Object);
-
-	// And add it to the Datasmith Scene as needed
-	if (!UsedMaterials.Contains(Material))
-	{
-		UsedMaterials.Add(Material);
-		Scene->AddMaterial(Material);
-	}
-
-	return Material;
-}
-
-TSharedPtr<IDatasmithBaseMaterialElement> FOpenNurbsTranslatorImpl::GetObjectMaterial(const FOpenNurbsObjectWrapper& Object)
+FOpenNurbsTranslatorImpl::FMaterial FOpenNurbsTranslatorImpl::GetObjectMaterial(const FOpenNurbsObjectWrapper& Object)
 {
 	// Ref: getObjectMaterialID(const ONX_Model_Object& object)
 	ON::object_material_source MaterialSource = Object.Attributes.MaterialSource();
@@ -2035,23 +2092,33 @@ TSharedPtr<IDatasmithBaseMaterialElement> FOpenNurbsTranslatorImpl::GetObjectMat
 		{
 			if (Object.Attributes.m_material_index != -1)
 			{
-				// Get material from Material table
-				return GetMaterial(Object.Attributes.m_material_index);
+				return { GetMaterial(Object.Attributes.m_material_index) , GetOpenNurbsMaterial(Object.Attributes.m_material_index) };
 			}
 			break;
 		}
 		case ON::material_from_layer:
 		{
-			TSharedPtr<IDatasmithBaseMaterialElement>* Material = LayerIndexToMaterial.Find(Object.Attributes.m_layer_index);
-			if (Material)
+			if (int* MaterialIndexPtr = LayerIndexToMaterialIndex.Find(Object.Attributes.m_layer_index))
 			{
-				return *Material;
+				return { GetMaterial(*MaterialIndexPtr), GetOpenNurbsMaterial(*MaterialIndexPtr) };
 			}
 			break;
 		}
 	}
 
-	return GetDefaultMaterial();
+	return { GetDefaultMaterial() };
+}
+
+ON_Material* FOpenNurbsTranslatorImpl::GetOpenNurbsMaterial(int MaterialIndex)
+{
+	int* MaterialTableIndexPtr = MaterialIndexToMaterialTableIndex.Find(MaterialIndex);
+
+	if (!MaterialTableIndexPtr)
+	{
+		return nullptr;
+	}
+	
+	return MaterialTable.At(*MaterialTableIndexPtr);
 }
 
 TSharedPtr<IDatasmithBaseMaterialElement> FOpenNurbsTranslatorImpl::GetMaterial(int MaterialIndex)
@@ -3057,12 +3124,44 @@ bool FOpenNurbsTranslatorImpl::TranslateBRep(ON_Brep* Brep, const ON_3dmObjectAt
 			return false;
 		}
 
-		BRepMesh.Append(RenderMeshes.Count(), RenderMeshes.Array());
-		
-		if (!DatasmithOpenNurbsTranslatorUtils::TranslateMesh(&BRepMesh, OutMesh, bHasNormal, ScalingFactor, Offset))
+		TArray<int32> FaceMaterialChannel;
+		bool bHasFaceMaterialChannel = false;
+		// Use BRep's face material channel for each of the render meshes
+		// Assuming that meshes correspond to brep faces
+		if (ensure(RenderMeshes.Count() == Brep->m_F.Count()))
 		{
-			return false;
+			FaceMaterialChannel.Reserve(Brep->m_F.Count());
+			for (int FaceIndex = 0; FaceIndex < Brep->m_F.Count(); ++FaceIndex)
+			{
+				const ON_BrepFace& Face = Brep->m_F[FaceIndex];
+				int Channel = Face.m_face_material_channel;
+
+				if (Channel != 0 && (Channel < MeshElement->GetMaterialSlotCount()))
+				{
+					bHasFaceMaterialChannel = true;
+				}
+
+				FaceMaterialChannel.Add(Channel);
+			}
+		}
+		if (bHasFaceMaterialChannel)
+		{
+			if (!DatasmithOpenNurbsTranslatorUtils::TranslateMesh(RenderMeshes, RenderMeshes.Count(), OutMesh, bHasNormal, ScalingFactor, Offset, bHasFaceMaterialChannel, FaceMaterialChannel.GetData()))
+			{
+				return false;
+			}
+		}
+		else
+		{
+			ON_Mesh Mesh;
+			Mesh.Append(RenderMeshes.Count(), RenderMeshes);
+			
+			if (!DatasmithOpenNurbsTranslatorUtils::TranslateMesh(&Mesh, OutMesh, bHasNormal, ScalingFactor, Offset))
+			{
+				return false;
+			}
 		}
+		
 
 		return true;
 	}
