Integrating CL# 11892719 from Dev-Enterprise
Fix lightmap density preview in Dataprep Viewport

JeanLuc.Corenthin

diff --git a/Engine/Plugins/Enterprise/DataprepEditor/Source/DataprepCore/Private/DataprepCorePrivateUtils.cpp b/Engine/Plugins/Enterprise/DataprepEditor/Source/DataprepCore/Private/DataprepCorePrivateUtils.cpp
--- a/Engine/Plugins/Enterprise/DataprepEditor/Source/DataprepCore/Private/DataprepCorePrivateUtils.cpp
+++ b/Engine/Plugins/Enterprise/DataprepEditor/Source/DataprepCore/Private/DataprepCorePrivateUtils.cpp
@@ -118,7 +118,8 @@ void DataprepCorePrivateUtils::BuildStaticMeshes(TSet<UStaticMesh*>& StaticMeshe
 					if(SourceModel.BuildSettings.DstLightmapIndex != -1)
 					{
 						TVertexInstanceAttributesConstRef<FVector2D> VertexInstanceUVs = Attributes.GetVertexInstanceUVs();
-						SourceModel.BuildSettings.bGenerateLightmapUVs = VertexInstanceUVs.IsValid() && VertexInstanceUVs.GetNumIndices() > SourceModel.BuildSettings.DstLightmapIndex;
+						//If the importer have enabled lightmap generation, disabling it may interfere with the build process, so we only allow enabling it.
+						SourceModel.BuildSettings.bGenerateLightmapUVs |= VertexInstanceUVs.IsValid() && VertexInstanceUVs.GetNumIndices() > SourceModel.BuildSettings.SrcLightmapIndex;
 					}
 					else
 					{
diff --git a/Engine/Plugins/Enterprise/DataprepEditor/Source/DataprepEditor/Private/Widgets/SDataprepEditorViewport.cpp b/Engine/Plugins/Enterprise/DataprepEditor/Source/DataprepEditor/Private/Widgets/SDataprepEditorViewport.cpp
--- a/Engine/Plugins/Enterprise/DataprepEditor/Source/DataprepEditor/Private/Widgets/SDataprepEditorViewport.cpp
+++ b/Engine/Plugins/Enterprise/DataprepEditor/Source/DataprepEditor/Private/Widgets/SDataprepEditorViewport.cpp
@@ -386,7 +386,13 @@ void SDataprepEditorViewport::UpdateScene()
 						{
 							PreviewMeshComponent->SetMobility(EComponentMobility::Static);
 						}
+						else
+						{
+							PreviewMeshComponent->SetMobility( SceneMeshComponent->Mobility );
+						}
 
+						PreviewMeshComponent->bOverrideLightMapRes = SceneMeshComponent->bOverrideLightMapRes;
+						PreviewMeshComponent->OverriddenLightMapRes = SceneMeshComponent->OverriddenLightMapRes;
 						PreviewMeshComponent->bForceWireframe = bWireframeRenderingMode;
 						PreviewMeshComponent->MeshColorIndex = PerMeshColorIndex % PerMeshColorsCount;
 						++PerMeshColorIndex;
diff --git a/Engine/Plugins/Enterprise/DatasmithCADImporter/Source/DatasmithCoreTechExtension/Private/CoreTechBlueprintLibrary.cpp b/Engine/Plugins/Enterprise/DatasmithCADImporter/Source/DatasmithCoreTechExtension/Private/CoreTechBlueprintLibrary.cpp
--- a/Engine/Plugins/Enterprise/DatasmithCADImporter/Source/DatasmithCoreTechExtension/Private/CoreTechBlueprintLibrary.cpp
+++ b/Engine/Plugins/Enterprise/DatasmithCADImporter/Source/DatasmithCoreTechExtension/Private/CoreTechBlueprintLibrary.cpp
@@ -10,6 +10,7 @@
 #include "AssetData.h"
 #include "Engine/StaticMesh.h"
 #include "IStaticMeshEditor.h"
+#include "StaticMeshAttributes.h"
 #include "Toolkits/ToolkitManager.h"
 
 
@@ -46,8 +47,15 @@ bool UCoreTechBlueprintLibrary::RetessellateStaticMeshWithNotification(UStaticMe
 				StaticMesh->PreEditChange( nullptr );
 			}
 
+			const int32 OldNumberOfUVChannels = DestinationMeshDescription->VertexInstanceAttributes().GetAttributeIndexCount<FVector2D>(MeshAttribute::VertexInstance::TextureCoordinate);
 			if (FCoreTechRetessellate_Impl::ApplyOnOneAsset( *StaticMesh, *CoreTechData, TessellationSettings ))
 			{
+				const int32 NumberOfUVChannels = DestinationMeshDescription->VertexInstanceAttributes().GetAttributeIndexCount<FVector2D>(MeshAttribute::VertexInstance::TextureCoordinate);
+				if (NumberOfUVChannels < OldNumberOfUVChannels)
+				{
+					FailureReason = FText::Format(NSLOCTEXT("BlueprintRetessellation", "UVChannelsDestroyed", "Tessellation operation on Static Mesh {0} is destroying all UV channels above channel #{1}"), FText::FromString(StaticMesh->GetName()), NumberOfUVChannels - 1);
+				}
+
 				// Post static mesh has changed
 				if (bApplyChanges)
 				{
diff --git a/Engine/Plugins/Enterprise/DatasmithCADImporter/Source/DatasmithCoreTechExtension/Private/DataprepTessellationOperation.cpp b/Engine/Plugins/Enterprise/DatasmithCADImporter/Source/DatasmithCoreTechExtension/Private/DataprepTessellationOperation.cpp
--- a/Engine/Plugins/Enterprise/DatasmithCADImporter/Source/DatasmithCoreTechExtension/Private/DataprepTessellationOperation.cpp
+++ b/Engine/Plugins/Enterprise/DatasmithCADImporter/Source/DatasmithCoreTechExtension/Private/DataprepTessellationOperation.cpp
@@ -91,6 +91,11 @@ void UDataprepTessellationOperation::OnExecution_Implementation(const FDataprepC
 				if( UCoreTechBlueprintLibrary::RetessellateStaticMeshWithNotification( StaticMesh, TessellationSettings, false, OutReason ) )
 				{
 					ModifiedStaticMeshes.Add( StaticMesh );
+					if (!OutReason.IsEmpty())
+					{
+						FText WarningMsg = FText::Format(LOCTEXT("DataprepTessellationOperation_TessellationCompletedWithWarning", "{0}"), OutReason);
+						LogInfo(WarningMsg);
+					}
 				}
 				else
 				{
diff --git a/Engine/Plugins/Enterprise/DatasmithCADImporter/Source/DatasmithCoreTechExtension/Public/CoreTechBlueprintLibrary.h b/Engine/Plugins/Enterprise/DatasmithCADImporter/Source/DatasmithCoreTechExtension/Public/CoreTechBlueprintLibrary.h
--- a/Engine/Plugins/Enterprise/DatasmithCADImporter/Source/DatasmithCoreTechExtension/Public/CoreTechBlueprintLibrary.h
+++ b/Engine/Plugins/Enterprise/DatasmithCADImporter/Source/DatasmithCoreTechExtension/Public/CoreTechBlueprintLibrary.h
@@ -36,7 +36,7 @@ public:
 	 * @param	StaticMesh				Static mesh to re-tessellate.
 	 * @param	TessellationSettings	Tessellation settings to use.
 	 * @param	bApplyChanges			Indicates if change must be notified.
-	 * @param	FailureReason			Text describing the reason of failure.
+	 * @param	FailureReason			Text describing the reason of failure, or a warning if the operation was successful.
 	 * @return True if successful, false otherwise
 	 */
 	UFUNCTION(Category = "Datasmith | Surface Operations" )
