[Datasmith] import option: serialization improvement and cleanup
- remove errorprone duplicated state in ImportContext
- add UDatasmithOptionBase class -> forces the serialization of all properties regrardless of CDO values
- all datasmith option classes now inherit from UDatasmithOptionBase
- updated translator API to reflect that change
- (de)serialize translator options in the scene asset -> correct values on reimport

diff --git a/Engine/Plugins/Enterprise/DatasmithC4DImporter/Source/DatasmithC4DTranslator/Private/DatasmithC4DImportOptions.h b/Engine/Plugins/Enterprise/DatasmithC4DImporter/Source/DatasmithC4DTranslator/Private/DatasmithC4DImportOptions.h
--- a/Engine/Plugins/Enterprise/DatasmithC4DImporter/Source/DatasmithC4DTranslator/Private/DatasmithC4DImportOptions.h
+++ b/Engine/Plugins/Enterprise/DatasmithC4DImporter/Source/DatasmithC4DTranslator/Private/DatasmithC4DImportOptions.h
@@ -21,7 +21,7 @@ enum class EC4DUVGeneration : uint8
 };*/
 
 UCLASS(config = EditorPerProjectUserSettings, HideCategories = (DebugProperty))
-class UDatasmithC4DImportOptions : public UObject
+class UDatasmithC4DImportOptions : public UDatasmithOptionsBase
 {
 	GENERATED_UCLASS_BODY()
 
diff --git a/Engine/Plugins/Enterprise/DatasmithC4DImporter/Source/DatasmithC4DTranslator/Private/DatasmithC4DTranslator.cpp b/Engine/Plugins/Enterprise/DatasmithC4DImporter/Source/DatasmithC4DTranslator/Private/DatasmithC4DTranslator.cpp
--- a/Engine/Plugins/Enterprise/DatasmithC4DImporter/Source/DatasmithC4DTranslator/Private/DatasmithC4DTranslator.cpp
+++ b/Engine/Plugins/Enterprise/DatasmithC4DImporter/Source/DatasmithC4DTranslator/Private/DatasmithC4DTranslator.cpp
@@ -69,17 +69,16 @@ bool FDatasmithC4DTranslator::LoadLevelSequence(const TSharedRef<IDatasmithLevel
 	return false;
 }
 
-void FDatasmithC4DTranslator::GetSceneImportOptions(TArray<TStrongObjectPtr<UObject>>& Options)
+void FDatasmithC4DTranslator::GetSceneImportOptions(TArray<TStrongObjectPtr<UDatasmithOptionsBase>>& Options)
 {
 	Options.Add(GetOrCreateC4DImportOptions());
 }
 
-void FDatasmithC4DTranslator::SetSceneImportOptions(TArray<TStrongObjectPtr<UObject>>& Options)
+void FDatasmithC4DTranslator::SetSceneImportOptions(TArray<TStrongObjectPtr<UDatasmithOptionsBase>>& Options)
 {
-	for (TStrongObjectPtr<UObject>& OptionPtr : Options)
+	for (const TStrongObjectPtr<UDatasmithOptionsBase>& OptionPtr : Options)
 	{
-		UObject* Option = OptionPtr.Get();
-		if (UDatasmithC4DImportOptions* InImportOptions = Cast<UDatasmithC4DImportOptions>(Option))
+		if (UDatasmithC4DImportOptions* InImportOptions = Cast<UDatasmithC4DImportOptions>(OptionPtr.Get()))
 		{
 			ImportOptions.Reset(InImportOptions);
 		}
diff --git a/Engine/Plugins/Enterprise/DatasmithC4DImporter/Source/DatasmithC4DTranslator/Private/DatasmithC4DTranslator.h b/Engine/Plugins/Enterprise/DatasmithC4DImporter/Source/DatasmithC4DTranslator/Private/DatasmithC4DTranslator.h
--- a/Engine/Plugins/Enterprise/DatasmithC4DImporter/Source/DatasmithC4DTranslator/Private/DatasmithC4DTranslator.h
+++ b/Engine/Plugins/Enterprise/DatasmithC4DImporter/Source/DatasmithC4DTranslator/Private/DatasmithC4DTranslator.h
@@ -24,8 +24,8 @@ public:
 	virtual bool LoadStaticMesh(const TSharedRef<IDatasmithMeshElement> MeshElement, FDatasmithMeshElementPayload& OutMeshPayload) override;
 	virtual bool LoadLevelSequence(const TSharedRef<IDatasmithLevelSequenceElement> LevelSequenceElement, FDatasmithLevelSequencePayload& OutLevelSequencePayload) override;
 
-	virtual void GetSceneImportOptions(TArray<TStrongObjectPtr<UObject>>& Options) override;
-	virtual void SetSceneImportOptions(TArray<TStrongObjectPtr<UObject>>& Options) override;
+	virtual void GetSceneImportOptions(TArray<TStrongObjectPtr<UDatasmithOptionsBase>>& Options) override;
+	virtual void SetSceneImportOptions(TArray<TStrongObjectPtr<UDatasmithOptionsBase>>& Options) override;
 
 private:
 	TStrongObjectPtr<UDatasmithC4DImportOptions>& GetOrCreateC4DImportOptions();
diff --git a/Engine/Plugins/Enterprise/DatasmithCADImporter/Source/DatasmithCADTranslator/Private/DatasmithCADTranslator.cpp b/Engine/Plugins/Enterprise/DatasmithCADImporter/Source/DatasmithCADTranslator/Private/DatasmithCADTranslator.cpp
--- a/Engine/Plugins/Enterprise/DatasmithCADImporter/Source/DatasmithCADTranslator/Private/DatasmithCADTranslator.cpp
+++ b/Engine/Plugins/Enterprise/DatasmithCADImporter/Source/DatasmithCADTranslator/Private/DatasmithCADTranslator.cpp
@@ -43,9 +43,9 @@ void FDatasmithCADTranslator::Initialize(FDatasmithTranslatorCapabilities& OutCa
 	OutCapabilities.SupportedFileFormats.Add(FFileFormatInfo{ TEXT("igs"), TEXT("IGES files") });
 
 	OutCapabilities.SupportedFileFormats.Add(FFileFormatInfo{ TEXT("jt"), TEXT("JT Open files") });
-	
+
 	OutCapabilities.SupportedFileFormats.Add(FFileFormatInfo{ TEXT("sat"), TEXT("3D ACIS model files") });
-	
+
 	OutCapabilities.SupportedFileFormats.Add(FFileFormatInfo{ TEXT("SLDASM"), TEXT("SolidWorks Product files") });
 	OutCapabilities.SupportedFileFormats.Add(FFileFormatInfo{ TEXT("SLDPRT"), TEXT("SolidWorks Part files") });
 
@@ -74,8 +74,8 @@ bool FDatasmithCADTranslator::LoadScene(TSharedRef<IDatasmithScene> DatasmithSce
 	ImportParameters.MaxNormalAngle = TesselationOptions.NormalTolerance;
 	ImportParameters.StitchingTechnique = (CADLibrary::EStitchingTechnique) TesselationOptions.StitchingTechnique;
 
-	CADLibrary::FFileDescription FileDescription(*FPaths::ConvertRelativePathToFull(GetSource().GetSourceFile()), 
-		TEXT(""), 
+	CADLibrary::FFileDescription FileDescription(*FPaths::ConvertRelativePathToFull(GetSource().GetSourceFile()),
+		TEXT(""),
 		*FPaths::GetPath(FPaths::ConvertRelativePathToFull(GetSource().GetSourceFile())) );
 
 	if (FileDescription.Extension == TEXT("jt"))
@@ -108,7 +108,7 @@ bool FDatasmithCADTranslator::LoadScene(TSharedRef<IDatasmithScene> DatasmithSce
 #ifdef CAD_TRANSLATOR_DEBUG
 		bWithProcessor = false;
 #endif //CAD_TRANSLATOR_DEBUG
-		
+
 		Dispatcher.Process(bWithProcessor);
 	}
 
@@ -145,7 +145,7 @@ bool FDatasmithCADTranslator::LoadStaticMesh(const TSharedRef<IDatasmithMeshElem
 	return OutMeshPayload.LodMeshes.Num() > 0;
 }
 
-void FDatasmithCADTranslator::SetSceneImportOptions(TArray<TStrongObjectPtr<UObject>>& Options)
+void FDatasmithCADTranslator::SetSceneImportOptions(TArray<TStrongObjectPtr<UDatasmithOptionsBase>>& Options)
 {
 	FDatasmithCoreTechTranslator::SetSceneImportOptions(Options);
 }
diff --git a/Engine/Plugins/Enterprise/DatasmithCADImporter/Source/DatasmithCADTranslator/Private/DatasmithCADTranslator.h b/Engine/Plugins/Enterprise/DatasmithCADImporter/Source/DatasmithCADTranslator/Private/DatasmithCADTranslator.h
--- a/Engine/Plugins/Enterprise/DatasmithCADImporter/Source/DatasmithCADTranslator/Private/DatasmithCADTranslator.h
+++ b/Engine/Plugins/Enterprise/DatasmithCADImporter/Source/DatasmithCADTranslator/Private/DatasmithCADTranslator.h
@@ -33,7 +33,7 @@ public:
 
 	virtual bool LoadStaticMesh(const TSharedRef<IDatasmithMeshElement> MeshElement, FDatasmithMeshElementPayload& OutMeshPayload) override;
 
-	virtual void SetSceneImportOptions(TArray<TStrongObjectPtr<UObject>>& Options) override;
+	virtual void SetSceneImportOptions(TArray<TStrongObjectPtr<UDatasmithOptionsBase>>& Options) override;
 
 private:
 	TMap<uint32, FString> CADFileToUE4GeomMap;
diff --git a/Engine/Plugins/Enterprise/DatasmithCADImporter/Source/DatasmithCoreTechParametricSurfaceData/Private/DatasmithCoreTechTranslator.cpp b/Engine/Plugins/Enterprise/DatasmithCADImporter/Source/DatasmithCoreTechParametricSurfaceData/Private/DatasmithCoreTechTranslator.cpp
--- a/Engine/Plugins/Enterprise/DatasmithCADImporter/Source/DatasmithCoreTechParametricSurfaceData/Private/DatasmithCoreTechTranslator.cpp
+++ b/Engine/Plugins/Enterprise/DatasmithCADImporter/Source/DatasmithCoreTechParametricSurfaceData/Private/DatasmithCoreTechTranslator.cpp
@@ -4,28 +4,25 @@
 
 #include "DatasmithImportOptions.h"
 
-void FDatasmithCoreTechTranslator::GetSceneImportOptions(TArray<TStrongObjectPtr<UObject>>& Options)
+void FDatasmithCoreTechTranslator::GetSceneImportOptions(TArray<TStrongObjectPtr<UDatasmithOptionsBase>>& Options)
 {
 	Options.Add(GetCommonTessellationOptionsPtr());
 }
 
-void FDatasmithCoreTechTranslator::SetSceneImportOptions(TArray<TStrongObjectPtr<UObject>>& Options)
+void FDatasmithCoreTechTranslator::SetSceneImportOptions(TArray<TStrongObjectPtr<UDatasmithOptionsBase>>& Options)
 {
-	for (const TStrongObjectPtr<UObject>& OptionPtr : Options)
+	for (const TStrongObjectPtr<UDatasmithOptionsBase>& OptionPtr : Options)
 	{
-		if (UObject* Option = OptionPtr.Get())
+		if (UDatasmithCommonTessellationOptions* TessellationOptionsObject = Cast<UDatasmithCommonTessellationOptions>(OptionPtr.Get()))
 		{
-			if (UDatasmithCommonTessellationOptions* TessellationOptionsObject = Cast<UDatasmithCommonTessellationOptions>(Option))
-			{
-				GetCommonTessellationOptionsPtr().Reset(TessellationOptionsObject);
-			}
+			CommonTessellationOptionsPtr.Reset(TessellationOptionsObject);
 		}
 	}
 }
 
 TStrongObjectPtr<UDatasmithCommonTessellationOptions>& FDatasmithCoreTechTranslator::GetCommonTessellationOptionsPtr()
 {
-	if(!CommonTessellationOptionsPtr.IsValid())
+	if (!CommonTessellationOptionsPtr.IsValid())
 	{
 		CommonTessellationOptionsPtr = Datasmith::MakeOptions<UDatasmithCommonTessellationOptions>();
 		check(CommonTessellationOptionsPtr.IsValid());
diff --git a/Engine/Plugins/Enterprise/DatasmithCADImporter/Source/DatasmithCoreTechParametricSurfaceData/Public/DatasmithCoreTechTranslator.h b/Engine/Plugins/Enterprise/DatasmithCADImporter/Source/DatasmithCoreTechParametricSurfaceData/Public/DatasmithCoreTechTranslator.h
--- a/Engine/Plugins/Enterprise/DatasmithCADImporter/Source/DatasmithCoreTechParametricSurfaceData/Public/DatasmithCoreTechTranslator.h
+++ b/Engine/Plugins/Enterprise/DatasmithCADImporter/Source/DatasmithCoreTechParametricSurfaceData/Public/DatasmithCoreTechTranslator.h
@@ -17,8 +17,8 @@ public:
 	virtual ~FDatasmithCoreTechTranslator() {}
 
 	// Begin IDatasmithTranslator overrides
-	virtual void GetSceneImportOptions(TArray<TStrongObjectPtr<UObject>>& Options) override;
-	virtual void SetSceneImportOptions(TArray<TStrongObjectPtr<UObject>>& Options) override;
+	virtual void GetSceneImportOptions(TArray<TStrongObjectPtr<UDatasmithOptionsBase>>& Options) override;
+	virtual void SetSceneImportOptions(TArray<TStrongObjectPtr<UDatasmithOptionsBase>>& Options) override;
 	// End IDatasmithTranslator overrides
 
 protected:
diff --git a/Engine/Plugins/Enterprise/DatasmithCADImporter/Source/DatasmithOpenNurbsTranslator/Private/DatasmithOpenNurbsImportOptions.h b/Engine/Plugins/Enterprise/DatasmithCADImporter/Source/DatasmithOpenNurbsTranslator/Private/DatasmithOpenNurbsImportOptions.h
--- a/Engine/Plugins/Enterprise/DatasmithCADImporter/Source/DatasmithOpenNurbsTranslator/Private/DatasmithOpenNurbsImportOptions.h
+++ b/Engine/Plugins/Enterprise/DatasmithCADImporter/Source/DatasmithOpenNurbsTranslator/Private/DatasmithOpenNurbsImportOptions.h
@@ -41,13 +41,13 @@ struct FDatasmithOpenNurbsOptions : public FDatasmithTessellationOptions
 
 	uint32 GetHash() const
 	{
-		return HashCombine(FDatasmithTessellationOptions::GetHash(), GetTypeHash(Geometry));;
+		return HashCombine(FDatasmithTessellationOptions::GetHash(), GetTypeHash(Geometry));
 	}
 };
 
 
 UCLASS(BlueprintType, config = EditorPerProjectUserSettings)
-class UDatasmithOpenNurbsImportOptions: public UObject
+class UDatasmithOpenNurbsImportOptions: public UDatasmithOptionsBase
 {
 	GENERATED_BODY()
 
diff --git a/Engine/Plugins/Enterprise/DatasmithCADImporter/Source/DatasmithOpenNurbsTranslator/Private/DatasmithOpenNurbsTranslator.cpp b/Engine/Plugins/Enterprise/DatasmithCADImporter/Source/DatasmithOpenNurbsTranslator/Private/DatasmithOpenNurbsTranslator.cpp
--- a/Engine/Plugins/Enterprise/DatasmithCADImporter/Source/DatasmithOpenNurbsTranslator/Private/DatasmithOpenNurbsTranslator.cpp
+++ b/Engine/Plugins/Enterprise/DatasmithCADImporter/Source/DatasmithOpenNurbsTranslator/Private/DatasmithOpenNurbsTranslator.cpp
@@ -271,7 +271,7 @@ namespace DatasmithOpenNurbsTranslatorUtils
 		{
 			return false;
 		}
-		
+
 		MeshDescription.Empty();
 
 		TVertexAttributesRef<FVector> VertexPositions = MeshDescription.VertexAttributes().GetAttributesRef<FVector>(MeshAttribute::Vertex::Position);
@@ -411,9 +411,9 @@ namespace DatasmithOpenNurbsTranslatorUtils
 			const int32 VertexInstanceCount = 3 * TriangleCount;
 
 			int32 VertexIndexBase = MeshDescription.Vertices().Num();
-			
+
 			MeshDescription.ReserveNewVertices(VertexCount);
-			
+
 			// Do not reserve vertex instances -  UUVGenerationFlattenMapping::GetOverlappingCornersRemapping is crashing because it assumes that vertex instances has no gaps
 			// but they appear when we skip triangles(degenerates) between meshes(if there are few meshes combine mesh description)
 			// MeshDescription.ReserveNewVertexInstances(VertexInstanceCount);
@@ -510,7 +510,7 @@ namespace DatasmithOpenNurbsTranslatorUtils
 
 		return true;
 	}
-	
+
 	bool TranslateMesh(const ON_Mesh* Mesh, FMeshDescription& MeshDescription, bool& bHasNormal, double ScalingFactor, const ON_3dVector& Offset)
 	{
 		return TranslateMesh(&Mesh, 1, MeshDescription, bHasNormal, ScalingFactor, Offset, false, nullptr);
@@ -690,12 +690,12 @@ private:
 	FString GetLayerName(const TSharedPtr<IDatasmithActorElement>& LayerElement);
 	void SetLayers(const TSharedPtr<IDatasmithActorElement>& ActorElement, const FOpenNurbsObjectWrapper& Object);
 	void SetTags(const TSharedPtr<IDatasmithActorElement>& ActorElement, const FOpenNurbsObjectWrapper& Object);
-	
+
 	void SetMaterialToMeshElement(TSharedPtr<IDatasmithMeshElement> MeshElement, TSharedPtr<IDatasmithBaseMaterialElement> MaterialElement, int32 SlotId);
 
 	bool TranslateBRep(ON_Brep* brep, const ON_3dmObjectAttributes& Attributes, FMeshDescription& OutMesh, const TSharedRef< IDatasmithMeshElement >& MeshElement, const FString& Name, bool& bHasNormal);
 
-	
+
 	bool ComputeObjectGeometryCenter(const FOpenNurbsObjectWrapper& Object, ON_3dVector& OutGeometryCenter);
 	bool ComputeGeometryCenter(ON_Geometry* Geometry, ON_3dVector& OutCenter);
 
@@ -1002,7 +1002,7 @@ void FOpenNurbsTranslatorImpl::TranslateMaterialTable(const ON_ObjectArray<ON_Ma
 				// Note - Texture->m_mapping_channel_id may used to get TextureMapping in order to generate texture coordinates from it.
 				// Specifically, from object's attributes get MappingRef - attributes.m_rendering_attributes.m_mappings
 				// Next, find channel within MappingRef's m_mapping_channels which m_mapping_channel_id is equal to Texture's
-				// Then m_mapping_id of that channel will be UUID in the UUIDToTextureMapping 
+				// Then m_mapping_id of that channel will be UUID in the UUIDToTextureMapping
 
 				// Extract the UV tiling, offset and rotation angle from the UV transform matrix
 				FMatrix Matrix;
@@ -1023,7 +1023,7 @@ void FOpenNurbsTranslatorImpl::TranslateMaterialTable(const ON_ObjectArray<ON_Ma
 				{
 					UVParameters.UVOffset.X = Translation.X / Tiling.X;
 
-					// Recomputing offset from Rhino to Unreal, taking into account how tiling Pivot is set in SetupUVEdit: 
+					// Recomputing offset from Rhino to Unreal, taking into account how tiling Pivot is set in SetupUVEdit:
 					// P = pivot, C - input uv, T - tiling, O - offset
 					// Texture coordinate transformed like this: UV = T*(C+O-P)+P (ignoring mirror/rotation)
 					// P is 0.5 (for V), so  UV = T*(C+O - 0.5) + 0.5 = T*C + T*O - T*0.5 + 0.5
@@ -2117,7 +2117,7 @@ ON_Material* FOpenNurbsTranslatorImpl::GetOpenNurbsMaterial(int MaterialIndex)
 	{
 		return nullptr;
 	}
-	
+
 	return MaterialTable.At(*MaterialTableIndexPtr);
 }
 
@@ -3155,13 +3155,13 @@ bool FOpenNurbsTranslatorImpl::TranslateBRep(ON_Brep* Brep, const ON_3dmObjectAt
 		{
 			ON_Mesh Mesh;
 			Mesh.Append(RenderMeshes.Count(), RenderMeshes);
-			
+
 			if (!DatasmithOpenNurbsTranslatorUtils::TranslateMesh(&Mesh, OutMesh, bHasNormal, ScalingFactor, Offset))
 			{
 				return false;
 			}
 		}
-		
+
 
 		return true;
 	}
@@ -3219,7 +3219,7 @@ bool FOpenNurbsTranslatorImpl::ComputeObjectGeometryCenter(const FOpenNurbsObjec
 	}
 	else if (Object.ObjectPtr->IsKindOf(&ON_PlaneSurface::m_ON_PlaneSurface_class_rtti))
 	{
-		// 
+		//
 	}
 	else if (Object.ObjectPtr->IsKindOf(&ON_LineCurve::m_ON_LineCurve_class_rtti))
 	{
@@ -3416,15 +3416,15 @@ bool FDatasmithOpenNurbsTranslator::LoadStaticMesh(const TSharedRef<IDatasmithMe
 	return OutMeshPayload.LodMeshes.Num() > 0;
 }
 
-void FDatasmithOpenNurbsTranslator::SetSceneImportOptions(TArray<TStrongObjectPtr<UObject>>& Options)
+void FDatasmithOpenNurbsTranslator::SetSceneImportOptions(TArray<TStrongObjectPtr<UDatasmithOptionsBase>>& Options)
 {
-	for (TStrongObjectPtr<UObject>& Option : Options)
+	for (const TStrongObjectPtr<UDatasmithOptionsBase>& Option : Options)
 	{
 		if (UDatasmithImportOptions* DatasmithOptions = Cast<UDatasmithImportOptions>(Option.Get()))
 		{
 			BaseOptions = DatasmithOptions->BaseOptions;
 		}
-		if (UDatasmithOpenNurbsImportOptions* OpenNurbsOptionsObj = Cast<UDatasmithOpenNurbsImportOptions>(Option.Get()))
+		else if (UDatasmithOpenNurbsImportOptions* OpenNurbsOptionsObj = Cast<UDatasmithOpenNurbsImportOptions>(Option.Get()))
 		{
 			OpenNurbsOptions = OpenNurbsOptionsObj->Options;
 		}
@@ -3437,7 +3437,7 @@ void FDatasmithOpenNurbsTranslator::SetSceneImportOptions(TArray<TStrongObjectPt
 	}
 }
 
-void FDatasmithOpenNurbsTranslator::GetSceneImportOptions(TArray<TStrongObjectPtr<UObject>>& Options)
+void FDatasmithOpenNurbsTranslator::GetSceneImportOptions(TArray<TStrongObjectPtr<UDatasmithOptionsBase>>& Options)
 {
 	Options.Add(Datasmith::MakeOptions<UDatasmithOpenNurbsImportOptions>());
 }
diff --git a/Engine/Plugins/Enterprise/DatasmithCADImporter/Source/DatasmithOpenNurbsTranslator/Private/DatasmithOpenNurbsTranslator.h b/Engine/Plugins/Enterprise/DatasmithCADImporter/Source/DatasmithOpenNurbsTranslator/Private/DatasmithOpenNurbsTranslator.h
--- a/Engine/Plugins/Enterprise/DatasmithCADImporter/Source/DatasmithOpenNurbsTranslator/Private/DatasmithOpenNurbsTranslator.h
+++ b/Engine/Plugins/Enterprise/DatasmithCADImporter/Source/DatasmithOpenNurbsTranslator/Private/DatasmithOpenNurbsTranslator.h
@@ -28,9 +28,9 @@ public:
 
 	virtual bool LoadStaticMesh(const TSharedRef<IDatasmithMeshElement> MeshElement, FDatasmithMeshElementPayload& OutMeshPayload) override;
 
-	virtual void SetSceneImportOptions(TArray<TStrongObjectPtr<UObject>>& Options) override;
+	virtual void SetSceneImportOptions(TArray<TStrongObjectPtr<UDatasmithOptionsBase>>& Options) override;
 
-	virtual void GetSceneImportOptions(TArray<TStrongObjectPtr<UObject>>& Options) override;
+	virtual void GetSceneImportOptions(TArray<TStrongObjectPtr<UDatasmithOptionsBase>>& Options) override;
 
 private:
 	TSharedPtr<FOpenNurbsTranslatorImpl> Translator;
diff --git a/Engine/Plugins/Enterprise/DatasmithCADImporter/Source/DatasmithPLMXMLTranslator/Private/DatasmithPlmXmlTranslator.cpp b/Engine/Plugins/Enterprise/DatasmithCADImporter/Source/DatasmithPLMXMLTranslator/Private/DatasmithPlmXmlTranslator.cpp
--- a/Engine/Plugins/Enterprise/DatasmithCADImporter/Source/DatasmithPLMXMLTranslator/Private/DatasmithPlmXmlTranslator.cpp
+++ b/Engine/Plugins/Enterprise/DatasmithCADImporter/Source/DatasmithPLMXMLTranslator/Private/DatasmithPlmXmlTranslator.cpp
@@ -51,7 +51,7 @@ bool FDatasmithPlmXmlTranslator::LoadStaticMesh(const TSharedRef<IDatasmithMeshE
 	return false;
 }
 
-void FDatasmithPlmXmlTranslator::GetSceneImportOptions(TArray<TStrongObjectPtr<UObject>>& Options)
+void FDatasmithPlmXmlTranslator::GetSceneImportOptions(TArray<TStrongObjectPtr<UDatasmithOptionsBase>>& Options)
 {
 	if (!CommonTessellationOptionsPtr.IsValid())
 	{
@@ -60,9 +60,9 @@ void FDatasmithPlmXmlTranslator::GetSceneImportOptions(TArray<TStrongObjectPtr<U
 	Options.Add(CommonTessellationOptionsPtr);
 }
 
-void FDatasmithPlmXmlTranslator::SetSceneImportOptions(TArray<TStrongObjectPtr<UObject>>& Options)
+void FDatasmithPlmXmlTranslator::SetSceneImportOptions(TArray<TStrongObjectPtr<UDatasmithOptionsBase>>& Options)
 {
-	for (const TStrongObjectPtr<UObject>& OptionPtr : Options)
+	for (const TStrongObjectPtr<UDatasmithOptionsBase>& OptionPtr : Options)
 	{
 		if (UDatasmithCommonTessellationOptions* TessellationOptionsObject = Cast<UDatasmithCommonTessellationOptions>(OptionPtr.Get()))
 		{
diff --git a/Engine/Plugins/Enterprise/DatasmithCADImporter/Source/DatasmithPLMXMLTranslator/Private/DatasmithPlmXmlTranslator.h b/Engine/Plugins/Enterprise/DatasmithCADImporter/Source/DatasmithPLMXMLTranslator/Private/DatasmithPlmXmlTranslator.h
--- a/Engine/Plugins/Enterprise/DatasmithCADImporter/Source/DatasmithPLMXMLTranslator/Private/DatasmithPlmXmlTranslator.h
+++ b/Engine/Plugins/Enterprise/DatasmithCADImporter/Source/DatasmithPLMXMLTranslator/Private/DatasmithPlmXmlTranslator.h
@@ -24,8 +24,8 @@ public:
 
 	virtual bool LoadStaticMesh(const TSharedRef<IDatasmithMeshElement> MeshElement, FDatasmithMeshElementPayload& OutMeshPayload) override;
 
-	virtual void GetSceneImportOptions(TArray<TStrongObjectPtr<UObject>>& Options) override;
-	virtual void SetSceneImportOptions(TArray<TStrongObjectPtr<UObject>>& Options) override;
+	virtual void GetSceneImportOptions(TArray<TStrongObjectPtr<UDatasmithOptionsBase>>& Options) override;
+	virtual void SetSceneImportOptions(TArray<TStrongObjectPtr<UDatasmithOptionsBase>>& Options) override;
 	//~ End IDatasmithTranslator interface
 
 private:
diff --git a/Engine/Plugins/Enterprise/DatasmithCADImporter/Source/DatasmithWireTranslator/Private/DatasmithWireTranslator.cpp b/Engine/Plugins/Enterprise/DatasmithCADImporter/Source/DatasmithWireTranslator/Private/DatasmithWireTranslator.cpp
--- a/Engine/Plugins/Enterprise/DatasmithCADImporter/Source/DatasmithWireTranslator/Private/DatasmithWireTranslator.cpp
+++ b/Engine/Plugins/Enterprise/DatasmithCADImporter/Source/DatasmithWireTranslator/Private/DatasmithWireTranslator.cpp
@@ -2178,7 +2178,7 @@ bool FDatasmithWireTranslator::LoadStaticMesh(const TSharedRef<IDatasmithMeshEle
 #endif
 }
 
-void FDatasmithWireTranslator::SetSceneImportOptions(TArray<TStrongObjectPtr<UObject>>& Options)
+void FDatasmithWireTranslator::SetSceneImportOptions(TArray<TStrongObjectPtr<UDatasmithOptionsBase>>& Options)
 {
 #ifdef USE_OPENMODEL
 	FDatasmithCoreTechTranslator::SetSceneImportOptions(Options);
diff --git a/Engine/Plugins/Enterprise/DatasmithCADImporter/Source/DatasmithWireTranslator/Private/DatasmithWireTranslator.h b/Engine/Plugins/Enterprise/DatasmithCADImporter/Source/DatasmithWireTranslator/Private/DatasmithWireTranslator.h
--- a/Engine/Plugins/Enterprise/DatasmithCADImporter/Source/DatasmithWireTranslator/Private/DatasmithWireTranslator.h
+++ b/Engine/Plugins/Enterprise/DatasmithCADImporter/Source/DatasmithWireTranslator/Private/DatasmithWireTranslator.h
@@ -34,7 +34,7 @@ public:
 	// End IDatasmithTranslator overrides
 
 	// Begin ADatasmithCoreTechTranslator overrides
-	virtual void SetSceneImportOptions(TArray<TStrongObjectPtr<UObject>>& Options) override;
+	virtual void SetSceneImportOptions(TArray<TStrongObjectPtr<UDatasmithOptionsBase>>& Options) override;
 
 protected:
 	virtual void InitCommonTessellationOptions(FDatasmithTessellationOptions& TessellationOptions) override
diff --git a/Engine/Plugins/Enterprise/DatasmithContent/Source/DatasmithContent/Private/DatasmithImportOptions.cpp b/Engine/Plugins/Enterprise/DatasmithContent/Source/DatasmithContent/Private/DatasmithImportOptions.cpp
--- a/Engine/Plugins/Enterprise/DatasmithContent/Source/DatasmithContent/Private/DatasmithImportOptions.cpp
+++ b/Engine/Plugins/Enterprise/DatasmithContent/Source/DatasmithContent/Private/DatasmithImportOptions.cpp
@@ -200,3 +200,9 @@ int32 FDatasmithStaticMeshImportOptions::ConvertLightmapEnumToValue( EDatasmithI
 
 	return MaxLightmapSize;
 	}
+
+void UDatasmithOptionsBase::Serialize(FStructuredArchive::FRecord Record)
+{
+	FGuardValue_Bitfield(Record.GetUnderlyingArchive().ArNoDelta, true);
+	Super::Serialize(Record);
+}
diff --git a/Engine/Plugins/Enterprise/DatasmithContent/Source/DatasmithContent/Public/DatasmithAssetImportData.h b/Engine/Plugins/Enterprise/DatasmithContent/Source/DatasmithContent/Public/DatasmithAssetImportData.h
--- a/Engine/Plugins/Enterprise/DatasmithContent/Source/DatasmithContent/Public/DatasmithAssetImportData.h
+++ b/Engine/Plugins/Enterprise/DatasmithContent/Source/DatasmithContent/Public/DatasmithAssetImportData.h
@@ -111,8 +111,8 @@ class DATASMITHCONTENT_API UDatasmithTranslatedSceneImportData : public UDatasmi
 
 #if WITH_EDITORONLY_DATA
 public:
-	UPROPERTY(EditAnywhere, Category = "Translation", meta = (ShowOnlyInnerProperties))
-	TSubclassOf<UObject> OriginFactory;
+	UPROPERTY(EditAnywhere, Category = "Options", meta = (ShowOnlyInnerProperties))
+	TArray<UDatasmithOptionsBase*> AdditionalOptions;
 #endif // WITH_EDITORONLY_DATA
 };
 
diff --git a/Engine/Plugins/Enterprise/DatasmithContent/Source/DatasmithContent/Public/DatasmithImportOptions.h b/Engine/Plugins/Enterprise/DatasmithContent/Source/DatasmithContent/Public/DatasmithImportOptions.h
--- a/Engine/Plugins/Enterprise/DatasmithContent/Source/DatasmithContent/Public/DatasmithImportOptions.h
+++ b/Engine/Plugins/Enterprise/DatasmithContent/Source/DatasmithContent/Public/DatasmithImportOptions.h
@@ -146,7 +146,7 @@ public:
 
 	bool operator == (const FDatasmithStaticMeshImportOptions& Other) const
 	{
-		return MinLightmapResolution == Other.MinLightmapResolution 
+		return MinLightmapResolution == Other.MinLightmapResolution
 			&& MaxLightmapResolution == Other.MaxLightmapResolution
 			&& bGenerateLightmapUVs == Other.bGenerateLightmapUVs
 			&& bRemoveDegenerates == Other.bRemoveDegenerates;
@@ -267,17 +267,12 @@ struct DATASMITHCONTENT_API FDatasmithTessellationOptions
 public:
 	bool operator == (const FDatasmithTessellationOptions& Other) const
 	{
-		return (FMath::IsNearlyEqual(ChordTolerance, Other.ChordTolerance) && FMath::IsNearlyEqual(MaxEdgeLength, Other.MaxEdgeLength) && FMath::IsNearlyEqual(NormalTolerance, Other.NormalTolerance) && StitchingTechnique == Other.StitchingTechnique);
+		return FMath::IsNearlyEqual(ChordTolerance, Other.ChordTolerance)
+			&& FMath::IsNearlyEqual(MaxEdgeLength, Other.MaxEdgeLength)
+			&& FMath::IsNearlyEqual(NormalTolerance, Other.NormalTolerance)
+			&& StitchingTechnique == Other.StitchingTechnique;
 	}
 
-	//void operator = (const FDatasmithTessellationOptions& Other)
-	//{
-	//	ChordTolerance = Other.ChordTolerance;
-	//	MaxEdgeLength = Other.MaxEdgeLength;
-	//	NormalTolerance = Other.NormalTolerance;
-	//	StitchingTechnique = Other.StitchingTechnique;
-	//}
-
 	uint32 GetHash() const
 	{
 		uint32 Hash = uint32(StitchingTechnique);
@@ -289,8 +284,25 @@ public:
 	}
 };
 
+/**
+ * Base class for all import options in datasmith.
+ *
+ * Notable feature: forces a full serialization of its properties (by opposition
+ * to the standard delta serialization which stores only the diff wrt the CDO)
+ * The intent is to store the exact options used in a previous import.
+ */
+UCLASS()
+class DATASMITHCONTENT_API UDatasmithOptionsBase : public UObject
+{
+	GENERATED_BODY()
+
+public:
+	void Serialize(FStructuredArchive::FRecord Record) override;
+};
+
+
 UCLASS(BlueprintType, config = EditorPerProjectUserSettings, Transient)
-class DATASMITHCONTENT_API UDatasmithCommonTessellationOptions : public UObject
+class DATASMITHCONTENT_API UDatasmithCommonTessellationOptions : public UDatasmithOptionsBase
 {
 	GENERATED_BODY()
 
@@ -300,7 +312,7 @@ public:
 };
 
 UCLASS(config = EditorPerProjectUserSettings, HideCategories = ("NotVisible"))
-class DATASMITHCONTENT_API UDatasmithImportOptions : public UObject
+class DATASMITHCONTENT_API UDatasmithImportOptions : public UDatasmithOptionsBase
 {
 	GENERATED_UCLASS_BODY()
 
diff --git a/Engine/Plugins/Enterprise/DatasmithFBXImporter/Source/DatasmithDeltaGenTranslator/Private/DatasmithDeltaGenTranslator.cpp b/Engine/Plugins/Enterprise/DatasmithFBXImporter/Source/DatasmithDeltaGenTranslator/Private/DatasmithDeltaGenTranslator.cpp
--- a/Engine/Plugins/Enterprise/DatasmithFBXImporter/Source/DatasmithDeltaGenTranslator/Private/DatasmithDeltaGenTranslator.cpp
+++ b/Engine/Plugins/Enterprise/DatasmithFBXImporter/Source/DatasmithDeltaGenTranslator/Private/DatasmithDeltaGenTranslator.cpp
@@ -136,7 +136,7 @@ bool FDatasmithDeltaGenTranslator::LoadLevelSequence(const TSharedRef<IDatasmith
 	return false;
 }
 
-void FDatasmithDeltaGenTranslator::GetSceneImportOptions(TArray<TStrongObjectPtr<UObject>>& Options)
+void FDatasmithDeltaGenTranslator::GetSceneImportOptions(TArray<TStrongObjectPtr<UDatasmithOptionsBase>>& Options)
 {
 	if (!ImportOptions.IsValid())
 	{
@@ -149,12 +149,11 @@ void FDatasmithDeltaGenTranslator::GetSceneImportOptions(TArray<TStrongObjectPtr
 	Options.Add(ImportOptions);
 }
 
-void FDatasmithDeltaGenTranslator::SetSceneImportOptions(TArray<TStrongObjectPtr<UObject>>& Options)
+void FDatasmithDeltaGenTranslator::SetSceneImportOptions(TArray<TStrongObjectPtr<UDatasmithOptionsBase>>& Options)
 {
-	for (TStrongObjectPtr<UObject>& OptionPtr : Options)
+	for (const TStrongObjectPtr<UDatasmithOptionsBase>& OptionPtr : Options)
 	{
-		UObject* Option = OptionPtr.Get();
-		if (UDatasmithDeltaGenImportOptions* InImportOptions = Cast<UDatasmithDeltaGenImportOptions>(Option))
+		if (UDatasmithDeltaGenImportOptions* InImportOptions = Cast<UDatasmithDeltaGenImportOptions>(OptionPtr.Get()))
 		{
 			ImportOptions.Reset(InImportOptions);
 		}
diff --git a/Engine/Plugins/Enterprise/DatasmithFBXImporter/Source/DatasmithDeltaGenTranslator/Private/DatasmithDeltaGenTranslator.h b/Engine/Plugins/Enterprise/DatasmithFBXImporter/Source/DatasmithDeltaGenTranslator/Private/DatasmithDeltaGenTranslator.h
--- a/Engine/Plugins/Enterprise/DatasmithFBXImporter/Source/DatasmithDeltaGenTranslator/Private/DatasmithDeltaGenTranslator.h
+++ b/Engine/Plugins/Enterprise/DatasmithFBXImporter/Source/DatasmithDeltaGenTranslator/Private/DatasmithDeltaGenTranslator.h
@@ -21,8 +21,8 @@ public:
 	virtual bool LoadStaticMesh(const TSharedRef<IDatasmithMeshElement> MeshElement, FDatasmithMeshElementPayload& OutMeshPayload) override;
 	virtual bool LoadLevelSequence(const TSharedRef<IDatasmithLevelSequenceElement> LevelSequenceElement, FDatasmithLevelSequencePayload& OutLevelSequencePayload) override;
 
-	virtual void GetSceneImportOptions(TArray<TStrongObjectPtr<UObject>>& Options) override;
-	virtual void SetSceneImportOptions(TArray<TStrongObjectPtr<UObject>>& Options) override;
+	virtual void GetSceneImportOptions(TArray<TStrongObjectPtr<UDatasmithOptionsBase>>& Options) override;
+	virtual void SetSceneImportOptions(TArray<TStrongObjectPtr<UDatasmithOptionsBase>>& Options) override;
 
 private:
     TStrongObjectPtr<UDatasmithDeltaGenImportOptions> ImportOptions;
diff --git a/Engine/Plugins/Enterprise/DatasmithFBXImporter/Source/DatasmithFBXTranslator/Private/DatasmithFBXImportOptions.cpp b/Engine/Plugins/Enterprise/DatasmithFBXImporter/Source/DatasmithFBXTranslator/Private/DatasmithFBXImportOptions.cpp
deleted file mode 100644
diff --git a/Engine/Plugins/Enterprise/DatasmithFBXImporter/Source/DatasmithFBXTranslator/Public/DatasmithFBXImportOptions.h b/Engine/Plugins/Enterprise/DatasmithFBXImporter/Source/DatasmithFBXTranslator/Public/DatasmithFBXImportOptions.h
--- a/Engine/Plugins/Enterprise/DatasmithFBXImporter/Source/DatasmithFBXTranslator/Public/DatasmithFBXImportOptions.h
+++ b/Engine/Plugins/Enterprise/DatasmithFBXImporter/Source/DatasmithFBXTranslator/Public/DatasmithFBXImportOptions.h
@@ -2,6 +2,8 @@
 
 #pragma once
 
+#include "DatasmithImportOptions.h"
+
 #include "CoreMinimal.h"
 #include "Engine/EngineTypes.h"
 #include "UObject/Object.h"
@@ -10,9 +12,9 @@
 #include "DatasmithFBXImportOptions.generated.h"
 
 UCLASS(config = EditorPerProjectUserSettings, HideCategories=(DebugProperty))
-class DATASMITHFBXTRANSLATOR_API UDatasmithFBXImportOptions : public UObject
+class DATASMITHFBXTRANSLATOR_API UDatasmithFBXImportOptions : public UDatasmithOptionsBase
 {
-	GENERATED_UCLASS_BODY()
+	GENERATED_BODY()
 
 public:
 	UPROPERTY(EditAnywhere, BlueprintReadWrite, Category=AssetImporting, meta=(DisplayName="Texture folders", ToolTip="Where to look for textures"))
diff --git a/Engine/Plugins/Enterprise/DatasmithFBXImporter/Source/DatasmithVREDTranslator/Private/DatasmithVREDTranslator.cpp b/Engine/Plugins/Enterprise/DatasmithFBXImporter/Source/DatasmithVREDTranslator/Private/DatasmithVREDTranslator.cpp
--- a/Engine/Plugins/Enterprise/DatasmithFBXImporter/Source/DatasmithVREDTranslator/Private/DatasmithVREDTranslator.cpp
+++ b/Engine/Plugins/Enterprise/DatasmithFBXImporter/Source/DatasmithVREDTranslator/Private/DatasmithVREDTranslator.cpp
@@ -156,7 +156,7 @@ bool FDatasmithVREDTranslator::LoadLevelSequence(const TSharedRef<IDatasmithLeve
 	return false;
 }
 
-void FDatasmithVREDTranslator::GetSceneImportOptions(TArray<TStrongObjectPtr<UObject>>& Options)
+void FDatasmithVREDTranslator::GetSceneImportOptions(TArray<TStrongObjectPtr<UDatasmithOptionsBase>>& Options)
 {
 	if (!ImportOptions.IsValid())
 	{
@@ -169,12 +169,11 @@ void FDatasmithVREDTranslator::GetSceneImportOptions(TArray<TStrongObjectPtr<UOb
 	Options.Add(ImportOptions);
 }
 
-void FDatasmithVREDTranslator::SetSceneImportOptions(TArray<TStrongObjectPtr<UObject>>& Options)
+void FDatasmithVREDTranslator::SetSceneImportOptions(TArray<TStrongObjectPtr<UDatasmithOptionsBase>>& Options)
 {
-	for (TStrongObjectPtr<UObject>& OptionPtr : Options)
+	for (const TStrongObjectPtr<UDatasmithOptionsBase>& OptionPtr : Options)
 	{
-		UObject* Option = OptionPtr.Get();
-		if (UDatasmithVREDImportOptions* InImportOptions = Cast<UDatasmithVREDImportOptions>(Option))
+		if (UDatasmithVREDImportOptions* InImportOptions = Cast<UDatasmithVREDImportOptions>(OptionPtr.Get()))
 		{
 			ImportOptions.Reset(InImportOptions);
 		}
diff --git a/Engine/Plugins/Enterprise/DatasmithFBXImporter/Source/DatasmithVREDTranslator/Private/DatasmithVREDTranslator.h b/Engine/Plugins/Enterprise/DatasmithFBXImporter/Source/DatasmithVREDTranslator/Private/DatasmithVREDTranslator.h
--- a/Engine/Plugins/Enterprise/DatasmithFBXImporter/Source/DatasmithVREDTranslator/Private/DatasmithVREDTranslator.h
+++ b/Engine/Plugins/Enterprise/DatasmithFBXImporter/Source/DatasmithVREDTranslator/Private/DatasmithVREDTranslator.h
@@ -21,8 +21,8 @@ public:
 	virtual bool LoadStaticMesh(const TSharedRef<IDatasmithMeshElement> MeshElement, FDatasmithMeshElementPayload& OutMeshPayload) override;
 	virtual bool LoadLevelSequence(const TSharedRef<IDatasmithLevelSequenceElement> LevelSequenceElement, FDatasmithLevelSequencePayload& OutLevelSequencePayload) override;
 
-	virtual void GetSceneImportOptions(TArray<TStrongObjectPtr<UObject>>& Options) override;
-	virtual void SetSceneImportOptions(TArray<TStrongObjectPtr<UObject>>& Options) override;
+	virtual void GetSceneImportOptions(TArray<TStrongObjectPtr<UDatasmithOptionsBase>>& Options) override;
+	virtual void SetSceneImportOptions(TArray<TStrongObjectPtr<UDatasmithOptionsBase>>& Options) override;
 
 private:
     TStrongObjectPtr<UDatasmithVREDImportOptions> ImportOptions;
diff --git a/Engine/Plugins/Enterprise/DatasmithGLTFImporter/Source/DatasmithGLTFTranslator/Private/DatasmithGLTFImportOptions.cpp b/Engine/Plugins/Enterprise/DatasmithGLTFImporter/Source/DatasmithGLTFTranslator/Private/DatasmithGLTFImportOptions.cpp
deleted file mode 100644
diff --git a/Engine/Plugins/Enterprise/DatasmithGLTFImporter/Source/DatasmithGLTFTranslator/Private/DatasmithGLTFImportOptions.h b/Engine/Plugins/Enterprise/DatasmithGLTFImporter/Source/DatasmithGLTFTranslator/Private/DatasmithGLTFImportOptions.h
--- a/Engine/Plugins/Enterprise/DatasmithGLTFImporter/Source/DatasmithGLTFTranslator/Private/DatasmithGLTFImportOptions.h
+++ b/Engine/Plugins/Enterprise/DatasmithGLTFImporter/Source/DatasmithGLTFTranslator/Private/DatasmithGLTFImportOptions.h
@@ -2,6 +2,8 @@
 
 #pragma once
 
+#include "DatasmithImportOptions.h"
+
 #include "CoreMinimal.h"
 #include "Engine/EngineTypes.h"
 #include "UObject/Object.h"
@@ -10,22 +12,17 @@
 #include "DatasmithGLTFImportOptions.generated.h"
 
 UCLASS(config = EditorPerProjectUserSettings, HideCategories = (DebugProperty))
-class UDatasmithGLTFImportOptions : public UObject
+class UDatasmithGLTFImportOptions : public UDatasmithOptionsBase
 {
-	GENERATED_UCLASS_BODY()
+	GENERATED_BODY()
 
 public:
-	UPROPERTY(
-	    config, EditAnywhere, BlueprintReadWrite, Category = Lightmaps,
-	    meta =
-	        (ToolTip =
-	             "Generate new UV coordinates for lightmapping instead of using the highest index UV set. \nTurn this on to have Unreal Studio generate lightmap UV sets automatically.\nTurn this off to try using the highest index existing UV set (if available) as the lightmap UV set.\nFor both cases, geometry without existing UV sets will receive an empty UV set, which will by itself not be valid for use with Lightmass."))
-	bool bGenerateLightmapUVs;
+	UPROPERTY(config, EditAnywhere, BlueprintReadWrite, Category = Lightmaps, meta = (
+		ToolTip = "Generate new UV coordinates for lightmapping instead of using the highest index UV set. \nTurn this on to have Unreal Studio generate lightmap UV sets automatically.\nTurn this off to try using the highest index existing UV set (if available) as the lightmap UV set.\nFor both cases, geometry without existing UV sets will receive an empty UV set, which will by itself not be valid for use with Lightmass."))
+	bool bGenerateLightmapUVs = false;
 
-	UPROPERTY(
-	    config, EditAnywhere, BlueprintReadWrite, Category = AssetImporting,
-	    meta =
-	        (DisplayName = "Import Uniform Scale",
-	         ToolTip = "Scale factor used for importing assets, by default: 100, for conversion from meters(glTF) to centimeters(Unreal default)."))
-	float ImportScale;
+	UPROPERTY( config, EditAnywhere, BlueprintReadWrite, Category = AssetImporting, meta = (
+		DisplayName = "Import Uniform Scale",
+		ToolTip = "Scale factor used for importing assets, by default: 100, for conversion from meters(glTF) to centimeters(Unreal default)."))
+	float ImportScale = 100.f;
 };
diff --git a/Engine/Plugins/Enterprise/DatasmithGLTFImporter/Source/DatasmithGLTFTranslator/Private/DatasmithGLTFTranslator.cpp b/Engine/Plugins/Enterprise/DatasmithGLTFImporter/Source/DatasmithGLTFTranslator/Private/DatasmithGLTFTranslator.cpp
--- a/Engine/Plugins/Enterprise/DatasmithGLTFImporter/Source/DatasmithGLTFTranslator/Private/DatasmithGLTFTranslator.cpp
+++ b/Engine/Plugins/Enterprise/DatasmithGLTFImporter/Source/DatasmithGLTFTranslator/Private/DatasmithGLTFTranslator.cpp
@@ -165,17 +165,16 @@ bool FDatasmithGLTFTranslator::LoadLevelSequence(const TSharedRef<IDatasmithLeve
 	return false;
 }
 
-void FDatasmithGLTFTranslator::GetSceneImportOptions(TArray<TStrongObjectPtr<UObject>>& Options)
+void FDatasmithGLTFTranslator::GetSceneImportOptions(TArray<TStrongObjectPtr<UDatasmithOptionsBase>>& Options)
 {
 	Options.Add(GetOrCreateGLTFImportOptions());
 }
 
-void FDatasmithGLTFTranslator::SetSceneImportOptions(TArray<TStrongObjectPtr<UObject>>& Options)
+void FDatasmithGLTFTranslator::SetSceneImportOptions(TArray<TStrongObjectPtr<UDatasmithOptionsBase>>& Options)
 {
-	for (TStrongObjectPtr<UObject>& OptionPtr : Options)
+	for (const TStrongObjectPtr<UDatasmithOptionsBase>& OptionPtr : Options)
 	{
-		UObject* Option = OptionPtr.Get();
-		if (UDatasmithGLTFImportOptions* InImportOptions = Cast<UDatasmithGLTFImportOptions>(Option))
+		if (UDatasmithGLTFImportOptions* InImportOptions = Cast<UDatasmithGLTFImportOptions>(OptionPtr.Get()))
 		{
 			ImportOptions.Reset(InImportOptions);
 		}
diff --git a/Engine/Plugins/Enterprise/DatasmithGLTFImporter/Source/DatasmithGLTFTranslator/Private/DatasmithGLTFTranslator.h b/Engine/Plugins/Enterprise/DatasmithGLTFImporter/Source/DatasmithGLTFTranslator/Private/DatasmithGLTFTranslator.h
--- a/Engine/Plugins/Enterprise/DatasmithGLTFImporter/Source/DatasmithGLTFTranslator/Private/DatasmithGLTFTranslator.h
+++ b/Engine/Plugins/Enterprise/DatasmithGLTFImporter/Source/DatasmithGLTFTranslator/Private/DatasmithGLTFTranslator.h
@@ -20,8 +20,8 @@ public:
 	virtual bool LoadStaticMesh(const TSharedRef<IDatasmithMeshElement> MeshElement, FDatasmithMeshElementPayload& OutMeshPayload) override;
 	virtual bool LoadLevelSequence(const TSharedRef<IDatasmithLevelSequenceElement> LevelSequenceElement, FDatasmithLevelSequencePayload& OutLevelSequencePayload) override;
 
-	virtual void GetSceneImportOptions(TArray<TStrongObjectPtr<UObject>>& Options) override;
-	virtual void SetSceneImportOptions(TArray<TStrongObjectPtr<UObject>>& Options) override;
+	virtual void GetSceneImportOptions(TArray<TStrongObjectPtr<UDatasmithOptionsBase>>& Options) override;
+	virtual void SetSceneImportOptions(TArray<TStrongObjectPtr<UDatasmithOptionsBase>>& Options) override;
 	//~ End IDatasmithTranslator interface
 
 private:
diff --git a/Engine/Plugins/Enterprise/DatasmithIFCImporter/Source/DatasmithIFCTranslator/Private/DatasmithIFCImportOptions.h b/Engine/Plugins/Enterprise/DatasmithIFCImporter/Source/DatasmithIFCTranslator/Private/DatasmithIFCImportOptions.h
--- a/Engine/Plugins/Enterprise/DatasmithIFCImporter/Source/DatasmithIFCTranslator/Private/DatasmithIFCImportOptions.h
+++ b/Engine/Plugins/Enterprise/DatasmithIFCImporter/Source/DatasmithIFCTranslator/Private/DatasmithIFCImportOptions.h
@@ -2,6 +2,8 @@
 
 #pragma once
 
+#include "DatasmithImportOptions.h"
+
 #include "CoreMinimal.h"
 #include "Engine/EngineTypes.h"
 #include "UObject/Object.h"
@@ -10,7 +12,7 @@
 #include "DatasmithIFCImportOptions.generated.h"
 
 UCLASS(config = EditorPerProjectUserSettings, HideCategories = (DebugProperty))
-class UDatasmithIFCImportOptions : public UObject
+class UDatasmithIFCImportOptions : public UDatasmithOptionsBase
 {
 	GENERATED_BODY()
 
diff --git a/Engine/Plugins/Enterprise/DatasmithIFCImporter/Source/DatasmithIFCTranslator/Private/DatasmithIFCTranslator.cpp b/Engine/Plugins/Enterprise/DatasmithIFCImporter/Source/DatasmithIFCTranslator/Private/DatasmithIFCTranslator.cpp
--- a/Engine/Plugins/Enterprise/DatasmithIFCImporter/Source/DatasmithIFCTranslator/Private/DatasmithIFCTranslator.cpp
+++ b/Engine/Plugins/Enterprise/DatasmithIFCImporter/Source/DatasmithIFCTranslator/Private/DatasmithIFCTranslator.cpp
@@ -113,7 +113,7 @@ bool FDatasmithIFCTranslator::LoadLevelSequence(const TSharedRef<IDatasmithLevel
 	return false;
 }
 
-void FDatasmithIFCTranslator::GetSceneImportOptions(TArray<TStrongObjectPtr<UObject>>& Options)
+void FDatasmithIFCTranslator::GetSceneImportOptions(TArray<TStrongObjectPtr<UDatasmithOptionsBase>>& Options)
 {
 	if (!ImportOptions.IsValid())
 	{
@@ -123,12 +123,11 @@ void FDatasmithIFCTranslator::GetSceneImportOptions(TArray<TStrongObjectPtr<UObj
 	Options.Add(ImportOptions);
 }
 
-void FDatasmithIFCTranslator::SetSceneImportOptions(TArray<TStrongObjectPtr<UObject>>& Options)
+void FDatasmithIFCTranslator::SetSceneImportOptions(TArray<TStrongObjectPtr<UDatasmithOptionsBase>>& Options)
 {
-	for (TStrongObjectPtr<UObject>& OptionPtr : Options)
+	for (const TStrongObjectPtr<UDatasmithOptionsBase>& OptionPtr : Options)
 	{
-		UObject* Option = OptionPtr.Get();
-		if (UDatasmithIFCImportOptions* InImportOptions = Cast<UDatasmithIFCImportOptions>(Option))
+		if (UDatasmithIFCImportOptions* InImportOptions = Cast<UDatasmithIFCImportOptions>(OptionPtr.Get()))
 		{
 			ImportOptions.Reset(InImportOptions);
 		}
diff --git a/Engine/Plugins/Enterprise/DatasmithIFCImporter/Source/DatasmithIFCTranslator/Private/DatasmithIFCTranslator.h b/Engine/Plugins/Enterprise/DatasmithIFCImporter/Source/DatasmithIFCTranslator/Private/DatasmithIFCTranslator.h
--- a/Engine/Plugins/Enterprise/DatasmithIFCImporter/Source/DatasmithIFCTranslator/Private/DatasmithIFCTranslator.h
+++ b/Engine/Plugins/Enterprise/DatasmithIFCImporter/Source/DatasmithIFCTranslator/Private/DatasmithIFCTranslator.h
@@ -20,8 +20,8 @@ public:
 	virtual bool LoadStaticMesh(const TSharedRef<IDatasmithMeshElement> MeshElement, FDatasmithMeshElementPayload& OutMeshPayload) override;
 	virtual bool LoadLevelSequence(const TSharedRef<IDatasmithLevelSequenceElement> LevelSequenceElement, FDatasmithLevelSequencePayload& OutLevelSequencePayload) override;
 
-	virtual void GetSceneImportOptions(TArray<TStrongObjectPtr<UObject>>& Options) override;
-	virtual void SetSceneImportOptions(TArray<TStrongObjectPtr<UObject>>& Options) override;
+	virtual void GetSceneImportOptions(TArray<TStrongObjectPtr<UDatasmithOptionsBase>>& Options) override;
+	virtual void SetSceneImportOptions(TArray<TStrongObjectPtr<UDatasmithOptionsBase>>& Options) override;
 
 private:
     TStrongObjectPtr<UDatasmithIFCImportOptions> ImportOptions;
diff --git a/Engine/Plugins/Enterprise/DatasmithImporter/Source/DatasmithImporter/Private/DatasmithFileProducer.cpp b/Engine/Plugins/Enterprise/DatasmithImporter/Source/DatasmithImporter/Private/DatasmithFileProducer.cpp
--- a/Engine/Plugins/Enterprise/DatasmithImporter/Source/DatasmithImporter/Private/DatasmithFileProducer.cpp
+++ b/Engine/Plugins/Enterprise/DatasmithImporter/Source/DatasmithImporter/Private/DatasmithFileProducer.cpp
@@ -130,11 +130,11 @@ bool UDatasmithFileProducer::Initialize()
 	TSharedPtr<IDatasmithTranslator> TranslatorPtr = TranslatableSourcePtr->GetTranslator();
 	if(IDatasmithTranslator* Translator = TranslatorPtr.Get())
 	{
-		TArray< TStrongObjectPtr<UObject> > Options;
+		TArray< TStrongObjectPtr<UDatasmithOptionsBase> > Options;
 		Translator->GetSceneImportOptions( Options );
 
 		bool bUpdateOptions = false;
-		for(TStrongObjectPtr<UObject>& ObjectPtr : Options)
+		for(TStrongObjectPtr<UDatasmithOptionsBase>& ObjectPtr : Options)
 		{
 			if(UDatasmithCommonTessellationOptions* TessellationOption = Cast<UDatasmithCommonTessellationOptions>(ObjectPtr.Get()))
 			{
diff --git a/Engine/Plugins/Enterprise/DatasmithImporter/Source/DatasmithImporter/Private/DatasmithImportContext.cpp b/Engine/Plugins/Enterprise/DatasmithImporter/Source/DatasmithImporter/Private/DatasmithImportContext.cpp
--- a/Engine/Plugins/Enterprise/DatasmithImporter/Source/DatasmithImporter/Private/DatasmithImportContext.cpp
+++ b/Engine/Plugins/Enterprise/DatasmithImporter/Source/DatasmithImporter/Private/DatasmithImportContext.cpp
@@ -10,19 +10,43 @@
 #include "UI/DatasmithImportOptionsWindow.h"
 #include "Utility/DatasmithImporterUtils.h"
 
-#include "Misc/FeedbackContext.h"
-#include "PackageTools.h"
-#include "Misc/Paths.h"
-#include "Interfaces/IMainFrameModule.h"
 #include "Dom/JsonObject.h"
+#include "Engine/StaticMesh.h"
 #include "EngineUtils.h"
 #include "FileHelpers.h"
-#include "Engine/StaticMesh.h"
+#include "Interfaces/IMainFrameModule.h"
+#include "Misc/FeedbackContext.h"
+#include "Misc/Paths.h"
+#include "PackageTools.h"
 
 
 #define LOCTEXT_NAMESPACE "DatasmithImportPlugin"
 
-const TCHAR* UserOptionPath = TEXT("Unreal Engine/Enterprise/Datasmith/Config/UserOptions.ini");
+namespace DatasmithImportContextImpl
+{
+	const TCHAR* GetUserOptionPath()
+	{
+		static FString UserDatasmithOptionsFile = FPaths::Combine(FPlatformProcess::UserSettingsDir(), "Unreal Engine/Enterprise/Datasmith/Config/UserOptions.ini");
+		return *UserDatasmithOptionsFile;
+	}
+
+	void LoadOptionConfig(UObject* InOption)
+	{
+		if (InOption)
+		{
+			InOption->LoadConfig(nullptr, GetUserOptionPath());
+		}
+	}
+
+	void SaveOptionConfig(UObject* InOption)
+	{
+		if (InOption)
+		{
+			InOption->SaveConfig(CPF_Config, GetUserOptionPath());
+		}
+	}
+}
+
 
 class FDatasmithImportOptionHelper
 {
@@ -143,24 +167,25 @@ FDatasmithImportContext::FDatasmithImportContext(const FString& FileName, bool b
 	, CurrentSceneActorIndex(0)
 	, ReferenceCollector(this)
 {
-	AddOption(Options.Get(), bLoadConfig);
-	SetFileName(FileName);
-
-	// Force the SceneHandling to be on current level by default.
-	// Note: This is done because this option was previously persisted and can get overwritten
-	Options->BaseOptions.SceneHandling = EDatasmithImportScene::CurrentLevel;
-
 	if (SceneTranslator)
 	{
 		SceneTranslator->GetSceneImportOptions(AdditionalImportOptions);
+	}
+
+	if (bLoadConfig)
+	{
+		DatasmithImportContextImpl::LoadOptionConfig(Options.Get());
 		for (const TStrongObjectPtr<UObject>& AdditionalImportOption : AdditionalImportOptions)
 		{
-			AddOption(AdditionalImportOption.Get(), bLoadConfig);
+			DatasmithImportContextImpl::LoadOptionConfig(AdditionalImportOption.Get());
 		}
-
-		// Temporarily give Rhino translator access to BaseOptions (JIRA UE-81278)
-		AdditionalImportOptions.Add(Options);
 	}
+
+	// Force the SceneHandling to be on current level by default.
+	// Note: This is done because this option was previously persisted and can get overwritten
+	Options->BaseOptions.SceneHandling = EDatasmithImportScene::CurrentLevel;
+
+	SetFileName(FileName);
 }
 
 void FDatasmithImportContext::SetFileName(const FString& FileName)
@@ -171,15 +196,19 @@ void FDatasmithImportContext::SetFileName(const FString& FileName)
 	FileHash = FMD5Hash::HashFile(*Options->FilePath);
 }
 
-void FDatasmithImportContext::AddOption(UObject* InOption, bool bLoadConfig)
+void FDatasmithImportContext::UpdateImportOption(UDatasmithOptionsBase* NewOption)
 {
-	if (InOption != nullptr)
+	if (NewOption)
 	{
-		ImportOptions.Add(InOption);
-		if (bLoadConfig)
+		auto SameClass = [=](const TStrongObjectPtr<UDatasmithOptionsBase>& Option) { return Option->GetClass() == NewOption->GetClass(); };
+
+		if (TStrongObjectPtr<UDatasmithOptionsBase>* PreviousOptionPtr = AdditionalImportOptions.FindByPredicate(SameClass))
+		{
+			PreviousOptionPtr->Reset(NewOption);
+		}
+		else
 		{
-			FString UserDatasmithOptionsFile = FPaths::Combine(FPlatformProcess::UserSettingsDir(), UserOptionPath);
-			InOption->LoadConfig(nullptr, *UserDatasmithOptionsFile);
+			AdditionalImportOptions.Emplace(NewOption);
 		}
 	}
 }
@@ -210,6 +239,15 @@ bool FDatasmithImportContext::Init(TSharedRef< IDatasmithScene > InScene, const
 	check(Options->FileName.Len() > 0);
 	check(Options->FilePath.Len() > 0);
 
+	TArray<UObject*> ImportOptions;
+	ImportOptions.Reserve(1 + AdditionalImportOptions.Num());
+
+	ImportOptions.Add(Options.Get());
+	for (const TStrongObjectPtr<UDatasmithOptionsBase>& AdditionalOption : AdditionalImportOptions)
+	{
+		ImportOptions.Add(AdditionalOption.Get());
+	}
+
 	if (!FModuleManager::Get().IsModuleLoaded("AssetTools"))
 	{
 		UE_LOG(LogDatasmithImport, Warning, TEXT("Import failed. The AssetTools module can't be loaded."));
@@ -242,16 +280,17 @@ bool FDatasmithImportContext::Init(TSharedRef< IDatasmithScene > InScene, const
 		}
 
 		// Update config file with new values
-		FString UserDatasmithOptionsFile = FPaths::Combine(FPlatformProcess::UserSettingsDir(), UserOptionPath);
 		for (UObject* Option : ImportOptions)
 		{
-			Option->SaveConfig(CPF_Config, *UserDatasmithOptionsFile);
+			DatasmithImportContextImpl::SaveOptionConfig(Option);
 		}
 	}
 
 	if (SceneTranslator)
 	{
+ 		AdditionalImportOptions.Push(Options);
 		SceneTranslator->SetSceneImportOptions(AdditionalImportOptions);
+ 		AdditionalImportOptions.Pop();
 	}
 
 	Options->UpdateNotDisplayedConfig( bIsAReimport );
diff --git a/Engine/Plugins/Enterprise/DatasmithImporter/Source/DatasmithImporter/Private/DatasmithImportFactory.cpp b/Engine/Plugins/Enterprise/DatasmithImporter/Source/DatasmithImporter/Private/DatasmithImportFactory.cpp
--- a/Engine/Plugins/Enterprise/DatasmithImporter/Source/DatasmithImporter/Private/DatasmithImportFactory.cpp
+++ b/Engine/Plugins/Enterprise/DatasmithImporter/Source/DatasmithImporter/Private/DatasmithImportFactory.cpp
@@ -151,12 +151,18 @@ namespace DatasmithImportFactoryImpl
 			SceneAsset = NewObject< UDatasmithScene >( Package, FName(*AssetName), RF_Public | RF_Standalone );
 		}
 
-		UDatasmithSceneImportData* ReImportSceneData = NewObject< UDatasmithSceneImportData >( SceneAsset, UDatasmithTranslatedSceneImportData::StaticClass() );
+		UDatasmithTranslatedSceneImportData* ReImportSceneData = NewObject< UDatasmithTranslatedSceneImportData >( SceneAsset );
 		SceneAsset->AssetImportData = ReImportSceneData;
 
 		// Copy over the changes the user may have done on the options
 		ReImportSceneData->BaseOptions = InContext.Options->BaseOptions;
 
+		for (const TStrongObjectPtr<UDatasmithOptionsBase>& Option : InContext.AdditionalImportOptions)
+		{
+			UDatasmithOptionsBase* OptionObj = Option.Get();
+			OptionObj->Rename(nullptr, ReImportSceneData);
+			ReImportSceneData->AdditionalOptions.Add(OptionObj);
+		}
 		ReImportSceneData->Update( InContext.Options->FilePath, InContext.FileHash.IsValid() ? &InContext.FileHash : nullptr );
 
 		FAssetRegistryModule::AssetCreated(ReImportSceneData);
@@ -467,8 +473,6 @@ UObject* UDatasmithImportFactory::FactoryCreateFile(UClass* InClass, UObject* In
 	FDatasmithTranslatableSceneSource TranslatableSource(Source);
 	if (!TranslatableSource.IsTranslatable())
 	{
-		bOperationCanceled = true;
-		bOutOperationCanceled = true;
 		UE_LOG(LogDatasmithImport, Warning, TEXT("Datasmith import error: no suitable translator found for this source. Abort import."));
 		return nullptr;
 	}
@@ -498,16 +502,14 @@ UObject* UDatasmithImportFactory::FactoryCreateFile(UClass* InClass, UObject* In
 
 	if (!TranslatableSource.Translate(Scene))
 	{
-		bOperationCanceled = true;
-		bOutOperationCanceled = true;
 		UE_LOG(LogDatasmithImport, Warning, TEXT("Datasmith import error: Scene translation failure. Abort import."));
 		return nullptr;
 	}
 
-	bOutOperationCanceled = !Import( ImportContext );
-	if (bOutOperationCanceled)
+	if (!Import( ImportContext ))
 	{
 		bOperationCanceled = true;
+		bOutOperationCanceled = true;
 		UE_LOG(LogDatasmithImport, Warning, TEXT("Datasmith import error. Abort import."));
 		return nullptr;
 	}
@@ -745,7 +747,6 @@ EReimportResult::Type UDatasmithImportFactory::ReimportStaticMesh(UStaticMesh* M
 EReimportResult::Type UDatasmithImportFactory::ReimportScene(UDatasmithScene* SceneAsset)
 {
 	// #ueent_todo: unify with import, BP, python, DP.
-	// Missing a Pipeline object ?
 	if (!SceneAsset || !SceneAsset->AssetImportData)
 	{
 		return EReimportResult::Failed;
@@ -769,6 +770,13 @@ EReimportResult::Type UDatasmithImportFactory::ReimportScene(UDatasmithScene* Sc
 	FDatasmithImportContext ImportContext(Source.GetSourceFile(), bLoadConfig, GetLoggerName(), GetDisplayName(), TranslatableSource.GetTranslator());
 	ImportContext.SceneAsset = SceneAsset;
 	ImportContext.Options->BaseOptions = ReimportData.BaseOptions; // Restore options as used in original import
+	if (UDatasmithTranslatedSceneImportData* TranslatedSceneReimportData = Cast<UDatasmithTranslatedSceneImportData>(SceneAsset->AssetImportData))
+	{
+		for (UDatasmithOptionsBase* Option : TranslatedSceneReimportData->AdditionalOptions)
+		{
+			ImportContext.UpdateImportOption(Option);
+		}
+	}
 	ImportContext.bIsAReimport = true;
 
 	FString ImportPath = ImportContext.Options->BaseOptions.AssetOptions.PackagePath.ToString();
diff --git a/Engine/Plugins/Enterprise/DatasmithImporter/Source/DatasmithImporter/Public/DatasmithImportContext.h b/Engine/Plugins/Enterprise/DatasmithImporter/Source/DatasmithImporter/Public/DatasmithImportContext.h
--- a/Engine/Plugins/Enterprise/DatasmithImporter/Source/DatasmithImporter/Public/DatasmithImportContext.h
+++ b/Engine/Plugins/Enterprise/DatasmithImporter/Source/DatasmithImporter/Public/DatasmithImportContext.h
@@ -2,17 +2,17 @@
 
 #pragma once
 
+#include "Containers/Map.h"
 #include "CoreMinimal.h"
-#include "UObject/ObjectMacros.h"
 #include "Engine/EngineTypes.h"
-#include "UObject/StrongObjectPtr.h"
-
-#include "Containers/Map.h"
 #include "Logging/TokenizedMessage.h"
-#include "Misc/SecureHash.h"
 #include "MaterialShared.h"
+#include "Misc/SecureHash.h"
+#include "UObject/ObjectMacros.h"
+#include "UObject/StrongObjectPtr.h"
 
 #include "DatasmithUtils.h"
+#include "DatasmithImportOptions.h"
 
 class AActor;
 class ADatasmithSceneActor;
@@ -222,7 +222,7 @@ struct DATASMITHIMPORTER_API FDatasmithImportContext
 	TStrongObjectPtr<UDatasmithImportOptions> Options;
 
 	/** Per-Translator settings related to the import of a Datasmith scene */
-	TArray<TStrongObjectPtr<UObject>> AdditionalImportOptions;
+	TArray<TStrongObjectPtr<UDatasmithOptionsBase>> AdditionalImportOptions;
 
 	/** Root blueprint which will be used if user requires every components of the scene under one blueprint */
 	UBlueprint* RootBlueprint;
@@ -299,12 +299,9 @@ public:
 	bool Init(const FString& FileName, TSharedRef< IDatasmithScene > Scene, const FString& InImportPath, EObjectFlags InFlags, FFeedbackContext* InWarn, const TSharedPtr<FJsonObject>& ImportSettingsJson, bool bSilent);
 
 	/**
-	 * Add option to the list of options used by the import process
-	 *
-	 * @param Option		The option to be added to the array
-	 * @param bLoadConfig	Overwrite properties of input option with those saved by the user
+	 * Replace or add options based on it's UClass.
 	 */
-	void AddOption(UObject* InOption, bool bLoadConfig);
+	void UpdateImportOption(UDatasmithOptionsBase* Option);
 
 	/** Push messages for display to the end-user */
 	TSharedRef<FTokenizedMessage> LogError(const FText& InErrorMessage);
@@ -341,9 +338,6 @@ private:
 	/** Map of SceneComponent objects added to the World at import */
 	TMap<FString, USceneComponent*> ImportedSceneComponentMap;
 
-	/** Array of options used by the calling factory */
-	TArray<UObject*> ImportOptions;
-
 	int32 CurrentSceneActorIndex;
 
 private:
diff --git a/Engine/Plugins/Enterprise/DatasmithImporter/Source/DatasmithTranslator/DatasmithTranslator.Build.cs b/Engine/Plugins/Enterprise/DatasmithImporter/Source/DatasmithTranslator/DatasmithTranslator.Build.cs
--- a/Engine/Plugins/Enterprise/DatasmithImporter/Source/DatasmithTranslator/DatasmithTranslator.Build.cs
+++ b/Engine/Plugins/Enterprise/DatasmithImporter/Source/DatasmithTranslator/DatasmithTranslator.Build.cs
@@ -2,28 +2,29 @@
 
 namespace UnrealBuildTool.Rules
 {
-    public class DatasmithTranslator : ModuleRules
-    {
-        public DatasmithTranslator(ReadOnlyTargetRules Target) : base(Target)
-        {
-            PublicDependencyModuleNames.AddRange(
-                new string[]
-                {
-                    "Core",
-                    "CoreUObject",
-                    "DatasmithCore",
-                }
-            );
+	public class DatasmithTranslator : ModuleRules
+	{
+		public DatasmithTranslator(ReadOnlyTargetRules Target) : base(Target)
+		{
+			PublicDependencyModuleNames.AddRange(
+				new string[]
+				{
+					"Core",
+					"CoreUObject",
+					"DatasmithCore",
+					"DatasmithContent",
+				}
+			);
 
-            PrivateDependencyModuleNames.AddRange(
-                new string[]
-                {
-                    "Engine",
-                    "MeshDescription",
-                    "RawMesh",
-                    "StaticMeshDescription",
-                }
-            );
-        }
-    }
+			PrivateDependencyModuleNames.AddRange(
+				new string[]
+				{
+					"Engine",
+					"MeshDescription",
+					"RawMesh",
+					"StaticMeshDescription",
+				}
+			);
+		}
+	}
 }
diff --git a/Engine/Plugins/Enterprise/DatasmithImporter/Source/DatasmithTranslator/Public/DatasmithTranslator.h b/Engine/Plugins/Enterprise/DatasmithImporter/Source/DatasmithTranslator/Public/DatasmithTranslator.h
--- a/Engine/Plugins/Enterprise/DatasmithImporter/Source/DatasmithTranslator/Public/DatasmithTranslator.h
+++ b/Engine/Plugins/Enterprise/DatasmithImporter/Source/DatasmithTranslator/Public/DatasmithTranslator.h
@@ -3,6 +3,7 @@
 
 #include "DatasmithPayload.h"
 #include "DatasmithSceneSource.h"
+#include "DatasmithImportOptions.h"
 
 #include "CoreMinimal.h"
 #include "UObject/Object.h"
@@ -110,14 +111,14 @@ public:
 	 *
 	 * @param OptionClasses list of classes that will be displayed to the user
 	 */
-	virtual void GetSceneImportOptions(TArray<TStrongObjectPtr<UObject>>& Options) {}
+	virtual void GetSceneImportOptions(TArray<TStrongObjectPtr<UDatasmithOptionsBase>>& Options) {}
 
 	/**
 	 * Values of additional options as entered by the user
 	 *
 	 * @param Options Actual values for the displayed options.
 	 */
-	virtual void SetSceneImportOptions(TArray<TStrongObjectPtr<UObject>>& Options) {}
+	virtual void SetSceneImportOptions(TArray<TStrongObjectPtr<UDatasmithOptionsBase>>& Options) {}
 
 private:
 	FDatasmithSceneSource SceneSource;
diff --git a/Engine/Plugins/Enterprise/DatasmithImporter/Source/DatasmithTranslator/Public/DatasmithTranslatorManager.h b/Engine/Plugins/Enterprise/DatasmithImporter/Source/DatasmithTranslator/Public/DatasmithTranslatorManager.h
--- a/Engine/Plugins/Enterprise/DatasmithImporter/Source/DatasmithTranslator/Public/DatasmithTranslatorManager.h
+++ b/Engine/Plugins/Enterprise/DatasmithImporter/Source/DatasmithTranslator/Public/DatasmithTranslatorManager.h
@@ -4,7 +4,6 @@
 
 #include "DatasmithTranslator.h"
 
-#include "UObject/StrongObjectPtr.h"
 
 /**
  * Responsible of the Translators lifecycle.
