Import IGES : CT Flag CT_LOAD_FLAG_SEARCH_NEW_TOPOLOGY was activated only with sew option. But it's not suitable because patches are imported independantly for IGES files.

diff --git a/Engine/Plugins/Enterprise/DatasmithCADImporter/Source/CADInterfaces/Private/CoreTechFileParser.cpp b/Engine/Plugins/Enterprise/DatasmithCADImporter/Source/CADInterfaces/Private/CoreTechFileParser.cpp
--- a/Engine/Plugins/Enterprise/DatasmithCADImporter/Source/CADInterfaces/Private/CoreTechFileParser.cpp
+++ b/Engine/Plugins/Enterprise/DatasmithCADImporter/Source/CADInterfaces/Private/CoreTechFileParser.cpp
@@ -569,7 +569,6 @@ void FCoreTechFileParser::ReadMaterials()
 FCoreTechFileParser::FCoreTechFileParser(const FImportParameters& ImportParams, const FString& EnginePluginsPath, const FString& InCachePath)
 	: CachePath(InCachePath)
 	, ImportParameters(ImportParams)
-	, bFixTopoRequired(ImportParameters.StitchingTechnique != StitchingNone)
 {
 	CTKIO_InitializeKernel(ImportParameters.MetricUnit, *EnginePluginsPath);
 }
@@ -821,12 +820,8 @@ CT_FLAGS FCoreTechFileParser::SetCoreTechImportOption(const FString& MainFileExt
 	// Ask Kernel IO to complete or create missing topology
 	if (MainFileExt == TEXT("igs") || MainFileExt == TEXT("iges"))
 	{
-		bFixTopoRequired = false;
 		Flags |= CT_LOAD_FLAG_COMPLETE_TOPOLOGY;
-		if (ImportParameters.StitchingTechnique == StitchingSew)
-		{
-			Flags |= CT_LOAD_FLAG_SEARCH_NEW_TOPOLOGY;
-		}
+		Flags |= CT_LOAD_FLAG_SEARCH_NEW_TOPOLOGY;
 	}
 
 	// 3dxml file is zipped files, it's full managed by Kernel_io. We cannot read it in sequential mode
@@ -969,7 +964,7 @@ bool FCoreTechFileParser::ReadComponent(CT_OBJECT_ID ComponentId, uint32 Default
 	TArray<CT_OBJECT_ID> Instances, Bodies;
 	GetInstancesAndBodies(ComponentId, Instances, Bodies);
 
-	if (ImportParameters.StitchingTechnique != StitchingNone && bFixTopoRequired)
+	if (ImportParameters.StitchingTechnique != StitchingNone)
 	{
 		if (!Instances.Num() && Bodies.Num() > 1 && ImportParameters.StitchingTechnique == StitchingSew)
 		{
diff --git a/Engine/Plugins/Enterprise/DatasmithCADImporter/Source/CADInterfaces/Public/CoreTechFileParser.h b/Engine/Plugins/Enterprise/DatasmithCADImporter/Source/CADInterfaces/Public/CoreTechFileParser.h
--- a/Engine/Plugins/Enterprise/DatasmithCADImporter/Source/CADInterfaces/Public/CoreTechFileParser.h
+++ b/Engine/Plugins/Enterprise/DatasmithCADImporter/Source/CADInterfaces/Public/CoreTechFileParser.h
@@ -133,7 +133,6 @@ protected:
 	bool bNeedSaveCTFile = false;
 
 	const FImportParameters& ImportParameters;
-	bool bFixTopoRequired;
 };
 
 #endif // CAD_INTERFACE
