PLMXML: fixed error for not-wellformed(but manageable) example. Added support Occurrence's inline Representation

diff --git a/Engine/Plugins/Enterprise/DatasmithCADImporter/Source/DatasmithPLMXMLTranslator/Private/DatasmithPlmXmlImporter.cpp b/Engine/Plugins/Enterprise/DatasmithCADImporter/Source/DatasmithPLMXMLTranslator/Private/DatasmithPlmXmlImporter.cpp
--- a/Engine/Plugins/Enterprise/DatasmithCADImporter/Source/DatasmithPLMXMLTranslator/Private/DatasmithPlmXmlImporter.cpp
+++ b/Engine/Plugins/Enterprise/DatasmithCADImporter/Source/DatasmithPLMXMLTranslator/Private/DatasmithPlmXmlImporter.cpp
@@ -97,7 +97,7 @@ namespace PlmXml
 		TSharedPtr<IDatasmithActorElement> Parent;
 		FString InstanceDatasmithName;
 		FImportedMesh ImportedRepresentation;
-		bool bOverrideTransform = false;
+		bool bOverrideWorldTransform = false;
 	};
 
 	struct FImportedExternalFile
@@ -211,7 +211,7 @@ namespace PlmXml
 		}
 	};
 
-	void SetActorTransform(const TSharedRef<IDatasmithActorElement>& ActorElement, const FTransform& Transform)
+	void SetActorWorldTransform(const TSharedRef<IDatasmithActorElement>& ActorElement, const FTransform& Transform)
 	{
 		ActorElement->SetTranslation(Transform.GetTranslation()); // TODO: convert units
 		ActorElement->SetRotation(Transform.GetRotation());
@@ -395,11 +395,6 @@ namespace PlmXml
 		TMap<FString, const FXmlNode*> ProductRevisionViewNodes;
 	};
 
-	struct FImportedProduct
-	{
-		TMap<FString, FImportedMesh> RepresentationsById;
-	};
-	
 	struct FImportContext
 	{
 		FString MainPlmXmlFilePath;
@@ -415,9 +410,13 @@ namespace PlmXml
 		TArray<FIdRef> InstanceGraphRootRefs;
 		TArray<FImportedRepresentationInstance> RepresentationInstances;
 		TSharedRef<FImportedInstanceTreeNode> ImportedInstanceTreeRootNode;
+		TMap<FString, TSharedPtr<FImportedInstanceTreeNode> > ImportedInstanceTreeNodes;
 		TArray<FImportedProductInstance> ProductInstances;
 		TSharedPtr<IDatasmithActorElement> RootActorElementForProductRevisions;
 
+		TMap<FString, FImportedMesh> RepresentationsById;
+
+
 		FImportContext(FString InRootFilePath,
 		               FPlmXmlMeshLoader& InMeshLoader,
 		               TSharedRef<IDatasmithScene> InDatasmithScene,
@@ -439,7 +438,6 @@ namespace PlmXml
 		FImportContext& MainImportContext;
 		
 		const FParsedProductDef& ProductDef;
-		FImportedProduct ImportedProduct;
 
 		FProductDefImportContext(FImportContext& InMainImportContext, const FParsedProductDef& InProductDef)
 			: MainImportContext(InMainImportContext)
@@ -523,6 +521,8 @@ namespace PlmXml
 		void AddChildProductInstance(FString Id, TSharedRef<FImportedInstanceTreeNode> InImportedInstanceTreeNode) override
 		{
 			ImportContext.MainImportContext.ImportedInstanceTreeRootNode->ChildInstances.Add(Id, InImportedInstanceTreeNode);
+
+			ImportContext.MainImportContext.ImportedInstanceTreeNodes.Add(Id, InImportedInstanceTreeNode);
 		}
 	};
 
@@ -651,7 +651,7 @@ namespace PlmXml
 	void TraverseProductRevisionView(FTraverseProductRevisionViewContext&& Context);
 	void TraverseProductInstance(FTraverseProductInstanceContext&& Context);
 
-	FImportedMesh GetOrImportRepresentation(FProductDefImportContext& ImportContext, FString ParentId, const FRepresentation& Representation)
+	FImportedMesh GetOrImportRepresentation(FImportContext& ImportContext, FString ParentId, const FRepresentation& Representation)
 	{
 		const FXmlNode* RepresentationNode = Representation.Node;
 
@@ -663,7 +663,7 @@ namespace PlmXml
 		// Combine parent id(in case representation itself has no id), id and format into representation unique key
 		FString RepresentationHashId = FMD5::HashAnsiString(*(ParentId + RepresentationId + RepresentationFormat));
 
-		if (FImportedMesh* ImportedRepresentationPtr = ImportContext.ImportedProduct.RepresentationsById.Find(RepresentationHashId))
+		if (FImportedMesh* ImportedRepresentationPtr = ImportContext.RepresentationsById.Find(RepresentationHashId))
 		{
 			return *ImportedRepresentationPtr;
 		}
@@ -676,15 +676,15 @@ namespace PlmXml
 			// TODO:  Possible(not included in samples yet) - <EntityMaterial>, <Material>, <Transform>, <CompoundRep>
 			if (LocationText.IsEmpty())
 			{
-				ImportContext.ImportedProduct.RepresentationsById.Add(RepresentationHashId, ImportedRepresentation);
+				ImportContext.RepresentationsById.Add(RepresentationHashId, ImportedRepresentation);
 				return ImportedRepresentation;
 			}
 			
-			FString FullPath = FPaths::IsRelative(LocationText) ? FPaths::ConvertRelativePathToFull(FPaths::GetPath(ImportContext.MainImportContext.MainPlmXmlFilePath) / LocationText) : LocationText;
-			ImportedRepresentation.ImportedMeshId = ImportContext.MainImportContext.MeshLoader.AddMeshToLoad(FullPath);
+			FString FullPath = FPaths::IsRelative(LocationText) ? FPaths::ConvertRelativePathToFull(FPaths::GetPath(ImportContext.MainPlmXmlFilePath) / LocationText) : LocationText;
+			ImportedRepresentation.ImportedMeshId = ImportContext.MeshLoader.AddMeshToLoad(FullPath);
 		}
 
-		ImportContext.ImportedProduct.RepresentationsById.Add(RepresentationHashId, ImportedRepresentation);
+		ImportContext.RepresentationsById.Add(RepresentationHashId, ImportedRepresentation);
 		return ImportedRepresentation;
 	}
 
@@ -765,6 +765,32 @@ namespace PlmXml
 			}
 		}
 	}
+
+	void ParseRepresentation(FImportContext& ImportContext, const FXmlNode* Node, const TSharedPtr<IDatasmithActorElement>& ParentElement, FString Name)
+	{
+		// TODO: probably need to check if Representation is of format 'JT', not another type and handle those other types separately(or just ignore)
+		if (TEXT("Representation") == Node->GetTag())
+		{
+			FRepresentation Representation = { Node };
+
+			FImportedMesh ImportedRepresentation = GetOrImportRepresentation(ImportContext, Name, Representation);
+
+			if (ImportedRepresentation.IsValid())
+			{
+				// Record instance info to instantiate later(after finishing InstanceGraph traversal and geometry loading)
+
+				FImportedRepresentationInstance RepresentationInstance;
+				RepresentationInstance.ImportedRepresentation = ImportedRepresentation;
+
+				RepresentationInstance.bOverrideWorldTransform = false;
+				
+				RepresentationInstance.Parent = ParentElement;
+				RepresentationInstance.InstanceDatasmithName = CreateDatasmithName(Name, TEXT("RepresentationInstance_") + FString::FromInt(ImportedRepresentation.ImportedMeshId));
+
+				ImportContext.RepresentationInstances.Add(RepresentationInstance);
+			}
+		}
+	}
 	
 	void TraverseProductRevisionView(FTraverseProductRevisionViewContext&& Context)
 	{
@@ -793,29 +819,9 @@ namespace PlmXml
 			// Consider every Representation child node to represent visible geometry
 			for (const FXmlNode* ProductRevisionViewChildNode : ProductRevisionView.Node->GetChildrenNodes())
 			{
-				// TODO: probably need to check if Representation is of format 'JT', not another type and handle those other types separately(or just ignore)
-				if (TEXT("Representation") == ProductRevisionViewChildNode->GetTag())
-				{
-					FRepresentation Representation = { ProductRevisionViewChildNode };
-					
-					FImportedMesh ImportedRepresentation = GetOrImportRepresentation(Context.ImportContext, ProductRevisionViewId, Representation);
-
-					if (ImportedRepresentation.IsValid())
-					{
-						// Record instance info to instantiate later(after finishing InstanceGraph traversal and geometry loading)
-						
-						FImportedRepresentationInstance RepresentationInstance;
-						RepresentationInstance.ImportedRepresentation = ImportedRepresentation;
-
-						// <Transform> elements associated with <Instance> elements are relative
-						RepresentationInstance.Transform = Context.Transform;
-						RepresentationInstance.bOverrideTransform = true;
-						RepresentationInstance.Parent = Context.GetActorElement();
-						RepresentationInstance.InstanceDatasmithName = CreateDatasmithName(Context.CurrentDatasmithName, TEXT("RepresentationInstance_") + FString::FromInt(ImportedRepresentation.ImportedMeshId));
-						
-						Context.ImportContext.MainImportContext.RepresentationInstances.Add(RepresentationInstance);
-					}
-				}
+				FTransform LocalTransform;
+				LocalTransform.SetIdentity(); 
+				ParseRepresentation(Context.ImportContext.MainImportContext, ProductRevisionViewChildNode, Context.GetActorElement(), Context.CurrentDatasmithName);
 			}
 		}
 
@@ -896,6 +902,11 @@ namespace PlmXml
 			Transform.SetIdentity();
 		}
 
+		FTransform ParentWorldTransform = ProductInstanceContext.Transform;
+		FTransform ActorWorldTransform;
+		FTransform::Multiply(&ActorWorldTransform, &Transform, &ParentWorldTransform);
+		PlmXml::SetActorWorldTransform(ActorElement, ActorWorldTransform);
+
 		FIdRef PartRefId = GetUriReferenceId(ProductInstance.Node, TEXT("partRef")); // partRef can be <Product>, <ProductRevision>, or<ProductRevisionView>.
 
 		if (!PartRefId)
@@ -983,7 +994,7 @@ namespace PlmXml
 
 								FString FullPath = FPaths::ConvertRelativePathToFull(FPaths::GetPath(ImportContext.MainPlmXmlFilePath) / LocationText);
 
-								// TODO: fix here, it's not representation
+								// note, it's not representation in terms of plmxml, but geometry data
 								PlmXml::FImportedMesh ImportedRepresentation;
 								ImportedRepresentation.ImportedMeshId = ImportContext.MeshLoader.AddMeshToLoad(FullPath);
 
@@ -993,6 +1004,10 @@ namespace PlmXml
 
 									PlmXml::FImportedRepresentationInstance RepresentationInstance;
 									RepresentationInstance.ImportedRepresentation = ImportedRepresentation;
+
+									// Keep representation relative to its actor
+									RepresentationInstance.bOverrideWorldTransform = false;
+									
 									// TODO: transform? RepresentationInstance.Transform = ;
 									RepresentationInstance.Parent = DataSetActorElement; // TODO: ActorElement -
 									// TODO: name - should be unique(made from hierarchy from root to dataset/file)
@@ -1052,38 +1067,55 @@ namespace PlmXml
 					
 					// TODO: rootRefs - if present use it to collect occurrence tree, else take all occurrences
 					const FXmlNode* OccurrenceNode = ProductViewChildNode;
-					const TSharedRef<IDatasmithActorElement> OccurrenceActorElement = CreateActor(ProductViewActorElement, *GetAttributeId(OccurrenceNode), *GetLabel(OccurrenceNode));
+					const FString DatasmithName = *GetAttributeId(OccurrenceNode);
+					const TSharedRef<IDatasmithActorElement> OccurrenceActorElement = CreateActor(ProductViewActorElement, DatasmithName, *GetLabel(OccurrenceNode));
 
 					ParseUserDataToDatasmithMetadata(OccurrenceNode, OccurrenceActorElement, ImportContext.DatasmithScene);
-					
-					TArray<FIdRef> InstanceRefs = PlmXml::GetAttributeUriReferenceList(OccurrenceNode, TEXT("instanceRefs"));
 
+					// Leaf Occurence actor could change depending what instanceRefs setup
+					TSharedRef<IDatasmithActorElement> OccurrenceLeafActorElement = OccurrenceActorElement;
+					TArray<FIdRef> InstanceRefs = PlmXml::GetAttributeUriReferenceList(OccurrenceNode, TEXT("instanceRefs"));
 					if (InstanceRefs.Num())
 					{
 
-						// !!!!!!!!!!!!!! NICE. Every 'Occurrence' adds an object to a 'view'(ProductView, PRV only view or others too?) that should
+						// Every 'Occurrence' adds an object to a 'view'(ProductView, PRV only view or others too?) that should
 						// be displayed when this 'view' is active(ro whatever it's called, 'shown'). Difference Occurrence's can reference same 'instance' object but
 						// modify it, with, for example transform. transformRef is in he doc. or child nodes Transform(takes precedence over transformRef) - absolute, additional Representation, EntityMaterial, AssociatedAttachment>.
 						// see 7.5.3 <Occurrence> for details
-						// !!!!!!!!!!!!!!
-						// instancedRef is supposed to be path through instance graph, but 'stand' example has instanceRefs not on InstanceGraph at all, instanceRefs referencing ProductRevision(which is not ProductRevisionView - another contraditiction with the docs)
-						FIdRef FirstInstanceRef = InstanceRefs[0];
-						ensure(ImportContext.InstanceGraphRootRefs.Find(FirstInstanceRef) != INDEX_NONE); // just a check for plmxml validity, first of instanceRefs should be present in rootRefs
+
+						// Few examples, contradicting with PLMXML spec:
+						// - instancedRef is supposed to be path through instance graph, but 'stand' example has instanceRefs not on InstanceGraph at all, instanceRefs referencing ProductRevision(which is not ProductRevisionView - another contraditiction with the docs)
+						// - PLMXML_Occurrence(UE-90772) has Occurrence's instanceRefs not being "a path through instanceGraph starting from a rootRef", doesn't start from rootRef
 
 						TSharedPtr<PlmXml::FImportedInstanceTreeNode> FoundImportedInstanceTreeNode = ImportContext.ImportedInstanceTreeRootNode;
 						for (const FIdRef& InstanceRef : InstanceRefs)
 						{
-							TSharedPtr<PlmXml::FImportedInstanceTreeNode>* Ptr = FoundImportedInstanceTreeNode->ChildInstances.Find(InstanceRef.Value);
-							if (!ensure(Ptr))
+							// Check if instance path is valid take the child instance
+							if (FoundImportedInstanceTreeNode.IsValid())
 							{
-								UE_LOG(LogDatasmithPlmXmlImport, Warning, TEXT("instanceRef '%s' doesn't correspond any ProductInstance"), *InstanceRef.Value);
-								break;
+								TSharedPtr<PlmXml::FImportedInstanceTreeNode>* Ptr = FoundImportedInstanceTreeNode->ChildInstances.Find(InstanceRef.Value);
+								if (!ensure(Ptr))
+								{
+									UE_LOG(LogDatasmithPlmXmlImport, Warning, TEXT("instanceRef '%s' doesn't correspond any ProductInstance"), *InstanceRef.Value);
+								}
+								else
+								{
+									FoundImportedInstanceTreeNode = *Ptr;
+								}
+							}
+							else
+							{
+								// In case PLMXML is not well-formed and instance path doesnt exist in instanceGraph just take an instance node(by its id)
+								if (TSharedPtr<PlmXml::FImportedInstanceTreeNode>* Ptr = ImportContext.ImportedInstanceTreeNodes.Find(InstanceRef.Value))
+								{
+									FoundImportedInstanceTreeNode = *Ptr;
+								}
 							}
-							FoundImportedInstanceTreeNode = *Ptr;
 						}
-						if (FoundImportedInstanceTreeNode->DatasmithActorElement.IsValid())
+						if (FoundImportedInstanceTreeNode.IsValid() && FoundImportedInstanceTreeNode->DatasmithActorElement.IsValid())
 						{
-							ActorOptions.Add(FoundImportedInstanceTreeNode->DatasmithActorElement.ToSharedRef());
+							OccurrenceLeafActorElement = FoundImportedInstanceTreeNode->DatasmithActorElement.ToSharedRef();
+							ActorOptions.Add(OccurrenceLeafActorElement);
 						}
 					}
 
@@ -1106,12 +1138,18 @@ namespace PlmXml
 							// TODO: right now it's only implemented for existing <Transform> element. Which replaces transform of the referenced object and does it as absolute transform.
 							// need also:
 							// - transformRef(optional, <Transform> takes precedence over it)
-							PlmXml::SetActorTransform(OccurrenceActorElement, Transform);
+							PlmXml::SetActorWorldTransform(OccurrenceActorElement, Transform);
 
 							const FXmlNode* ProductRevisionNode = ImportContext.ParsedPlmXml.ProductRevisionNodes[InstancedRef.Value];
-							FImportedProductRevision ImportedProductRevision = ParseProductRevision(ImportContext, OccurrenceActorElement, ProductRevisionNode, OccurrenceActorElement->GetName());
+							FImportedProductRevision ImportedProductRevision = ParseProductRevision(ImportContext, OccurrenceLeafActorElement, ProductRevisionNode, OccurrenceActorElement->GetName());
 						}
 					}
+
+					// TODO: Transform for Occurrence may come from parent Occurrence(by occurrenceRefs)'s referenced instance(partRef?)
+					for (const FXmlNode* OccurrenceNodeChildNode : OccurrenceNode->GetChildrenNodes())
+					{
+						ParseRepresentation(ImportContext, OccurrenceNodeChildNode, OccurrenceLeafActorElement, DatasmithName);
+					}
 				}
 			}
 
@@ -1330,13 +1368,17 @@ bool FDatasmithPlmXmlImporter::OpenFile(const FString& InFilePath, const FDatasm
 		{
 			TSharedRef<IDatasmithActorElement> ActorElement = ActorElementPtr.ToSharedRef();
 
-			if (RepresentationInstance.bOverrideTransform)
+			if (RepresentationInstance.bOverrideWorldTransform)
 			{
-				PlmXml::SetActorTransform(ActorElement, RepresentationInstance.Transform);
+				// Set Actor World Transform
+				PlmXml::SetActorWorldTransform(ActorElement, RepresentationInstance.Transform);
+
+				// Add actor as child keeping its transform set above as absolute
 				RepresentationInstance.Parent->AddChild(ActorElement, EDatasmithActorAttachmentRule::KeepWorldTransform);
 			}
 			else
 			{
+				// Add actor as child interpreting its transform as relative
 				// TODO: this needs to be relative - FTransform::Multiply(WorldTransform, ParentTransform, LocalTransform);
 				RepresentationInstance.Parent->AddChild(ActorElement, EDatasmithActorAttachmentRule::KeepRelativeTransform);
 			}
