Improvement of the choice of actor name. Original names (CAD system name) of instance and reference are preferred at names of Kernel_IO.

diff --git a/Engine/Plugins/Enterprise/DatasmithCADImporter/Source/DatasmithCADTranslator/Private/DatasmithSceneGraphBuilder.cpp b/Engine/Plugins/Enterprise/DatasmithCADImporter/Source/DatasmithCADTranslator/Private/DatasmithSceneGraphBuilder.cpp
--- a/Engine/Plugins/Enterprise/DatasmithCADImporter/Source/DatasmithCADTranslator/Private/DatasmithSceneGraphBuilder.cpp
+++ b/Engine/Plugins/Enterprise/DatasmithCADImporter/Source/DatasmithCADTranslator/Private/DatasmithSceneGraphBuilder.cpp
@@ -273,59 +273,64 @@ void FDatasmithSceneGraphBuilder::GetNodeUUIDAndName(
 	FString& OutName
 )
 {
-	FString* IName = InInstanceNodeMetaDataMap.Find(TEXT("CTName"));
-	FString* IOriginalName = InInstanceNodeMetaDataMap.Find(TEXT("Name"));
-	FString* IUUID = InInstanceNodeMetaDataMap.Find(TEXT("UUID"));
+	FString* InstanceKernelIOName = InInstanceNodeMetaDataMap.Find(TEXT("CTName"));
+	FString* InstanceCADName = InInstanceNodeMetaDataMap.Find(TEXT("Name"));
+	FString* InstanceUUID = InInstanceNodeMetaDataMap.Find(TEXT("UUID"));
 
-	FString* RName = InReferenceNodeMetaDataMap.Find(TEXT("CTName"));
-	FString* ROriginalName = InReferenceNodeMetaDataMap.Find(TEXT("Name"));
-	FString* RUUID = InReferenceNodeMetaDataMap.Find(TEXT("UUID"));
+	FString* ReferenceKernelIOName = InReferenceNodeMetaDataMap.Find(TEXT("CTName"));
+	FString* ReferenceCADName = InReferenceNodeMetaDataMap.Find(TEXT("Name"));
+	FString* ReferenceUUID = InReferenceNodeMetaDataMap.Find(TEXT("UUID"));
 
-	FString ReferenceName;
-
-	// Reference Name
-	if (ROriginalName)
+	// Outname Name
+	// IName and RName are KernelIO build Name. Original names (CAD system name) are preferred
+	if (InstanceCADName && !InstanceCADName->IsEmpty())
+	{
+		OutName = *InstanceCADName;
+	}
+	else if (ReferenceCADName && !ReferenceCADName->IsEmpty())
 	{
-		ReferenceName = *ROriginalName;
+		OutName = *ReferenceCADName;
 	}
-	else if (RName)
+	else if (InstanceKernelIOName && !InstanceKernelIOName->IsEmpty())
 	{
-		ReferenceName = *RName;
+		OutName = *InstanceKernelIOName;
+	}
+	else if (ReferenceKernelIOName && !ReferenceKernelIOName->IsEmpty())
+	{
+		OutName = *ReferenceKernelIOName;
 	}
 	else
 	{
-		ReferenceName = "NoName";
+		OutName = "NoName";
 	}
-
-	OutName = IOriginalName ? *IOriginalName : IName ? *IName : ReferenceName;
 	CleanName(OutName);
 
 	CADUUID UEUUID = HashCombine(GetTypeHash(InParentUEUUID), GetTypeHash(InComponentIndex));
 
-	if (IUUID)
+	if (InstanceUUID)
 	{
-		UEUUID = HashCombine(UEUUID, GetTypeHash(*IUUID));
+		UEUUID = HashCombine(UEUUID, GetTypeHash(*InstanceUUID));
 	}
-	if (IOriginalName)
+	if (InstanceCADName)
 	{
-		UEUUID = HashCombine(UEUUID, GetTypeHash(*IOriginalName));
+		UEUUID = HashCombine(UEUUID, GetTypeHash(*InstanceCADName));
 	}
-	if (IName)
+	if (InstanceKernelIOName)
 	{
-		UEUUID = HashCombine(UEUUID, GetTypeHash(*IName));
+		UEUUID = HashCombine(UEUUID, GetTypeHash(*InstanceKernelIOName));
 	}
 
-	if (RUUID)
+	if (ReferenceUUID)
 	{
-		UEUUID = HashCombine(UEUUID, GetTypeHash(*RUUID));
+		UEUUID = HashCombine(UEUUID, GetTypeHash(*ReferenceUUID));
 	}
-	if (ROriginalName)
+	if (ReferenceCADName)
 	{
-		UEUUID = HashCombine(UEUUID, GetTypeHash(*ROriginalName));
+		UEUUID = HashCombine(UEUUID, GetTypeHash(*ReferenceCADName));
 	}
-	if (RName)
+	if (ReferenceKernelIOName)
 	{
-		UEUUID = HashCombine(UEUUID, GetTypeHash(*RName));
+		UEUUID = HashCombine(UEUUID, GetTypeHash(*ReferenceKernelIOName));
 	}
 
 	OutUEUUID = FString::Printf(TEXT("0x%08x"), UEUUID);
@@ -529,7 +534,6 @@ void FDatasmithSceneGraphBuilder::AddMetaData(TSharedPtr< IDatasmithActorElement
 		UnwantedAttributes.Add(TEXT("OriginalUnitsMass"));
 		UnwantedAttributes.Add(TEXT("OriginalUnitsLength"));
 		UnwantedAttributes.Add(TEXT("OriginalUnitsDuration"));
-		UnwantedAttributes.Add(TEXT("OriginalId"));
 		UnwantedAttributes.Add(TEXT("OriginalIdStr"));
 		UnwantedAttributes.Add(TEXT("ShowAttribute"));
 		UnwantedAttributes.Add(TEXT("Identification"));
@@ -551,6 +555,11 @@ void FDatasmithSceneGraphBuilder::AddMetaData(TSharedPtr< IDatasmithActorElement
 			continue;
 		}
 
+		if (Attribute.Value.IsEmpty())
+		{
+			continue;
+		}
+
 		// If file information are attached to object, make sure to set a workable and full path
 		if (Attribute.Key == TEXT("FileName"))
 		{
@@ -596,6 +605,11 @@ void FDatasmithSceneGraphBuilder::AddMetaData(TSharedPtr< IDatasmithActorElement
 			continue;
 		}
 
+		if (Attribute.Value.IsEmpty())
+		{
+			continue;
+		}
+
 		FString MetaName = TEXT("Instance ");
 		MetaName += Attribute.Key;
 		TSharedRef< IDatasmithKeyValueProperty > KeyValueProperty = FDatasmithSceneFactory::CreateKeyValueProperty(*MetaName);
