Kernel_IO forget meshing parameters after a Sew.
SetCoreTechTessellationState(ImportParameters) is launched after each Sew to set tessellation parameters

diff --git a/Engine/Plugins/Enterprise/DatasmithCADImporter/Source/CADInterfaces/Private/CoreTechFileParser.cpp b/Engine/Plugins/Enterprise/DatasmithCADImporter/Source/CADInterfaces/Private/CoreTechFileParser.cpp
--- a/Engine/Plugins/Enterprise/DatasmithCADImporter/Source/CADInterfaces/Private/CoreTechFileParser.cpp
+++ b/Engine/Plugins/Enterprise/DatasmithCADImporter/Source/CADInterfaces/Private/CoreTechFileParser.cpp
@@ -978,6 +978,7 @@ bool FCoreTechFileParser::ReadComponent(CT_OBJECT_ID ComponentId, uint32 Default
 			}
 		}
 		GetInstancesAndBodies(ComponentId, Instances, Bodies);
+		SetCoreTechTessellationState(ImportParameters);
 	}
 
 	for (CT_OBJECT_ID InstanceId : Instances)
