Kernel_IO's Stitching action always ends by the split of the new bodies into a set of connected patches
a CT_Component (reference in the concept of instance/reference) can have bodies and instances
The SEW stitching rule in UE is:
- case 1 : the component has only a set of bodies: the bodies are merged, stitched and splitted into a new set of topologically correct bodies
- case 2 : the component has only one body or has bodies + instance, only a topology healing is done on the bodies

Case 1: Repair is done before processing bodies of a component
Case 2: Repair is done before getting the mesh of the body. As the body is exploded, the new bodies have to be discovered by comparing parent's bodies before and after the repair action

diff --git a/Engine/Plugins/Enterprise/DatasmithCADImporter/Source/CADInterfaces/Private/CoreTechFileParser.cpp b/Engine/Plugins/Enterprise/DatasmithCADImporter/Source/CADInterfaces/Private/CoreTechFileParser.cpp
--- a/Engine/Plugins/Enterprise/DatasmithCADImporter/Source/CADInterfaces/Private/CoreTechFileParser.cpp
+++ b/Engine/Plugins/Enterprise/DatasmithCADImporter/Source/CADInterfaces/Private/CoreTechFileParser.cpp
@@ -866,13 +866,9 @@ bool FCoreTechFileParser::ReadNode(CT_OBJECT_ID NodeId, uint32 DefaultMaterialHa
 		return ReadUnloadedComponent(NodeId);
 
 	case CT_BODY_TYPE:
-		if (int32* Index = SceneGraphArchive.CADIdToBodyIndex.Find(NodeId))
-		{
-			return true;
-		}
-		return ReadBody(NodeId, DefaultMaterialHash);
+		break;
 
-		//Treat all CT_CURVE_TYPE :
+	//Treat all CT_CURVE_TYPE :
 	case CT_CURVE_TYPE:
 	case CT_C_NURBS_TYPE:
 	case CT_CONICAL_TYPE:
@@ -964,21 +960,25 @@ bool FCoreTechFileParser::ReadComponent(CT_OBJECT_ID ComponentId, uint32 Default
 	TArray<CT_OBJECT_ID> Instances, Bodies;
 	GetInstancesAndBodies(ComponentId, Instances, Bodies);
 
-	if (ImportParameters.StitchingTechnique != StitchingNone)
+
+	// Kernel_IO's Stitching action always ends by the split of the new bodies into a set of connected patches
+	// a CT_Component (reference in the concept of instance/reference) can have bodies and instances
+	// The SEW stitching rule in UE is:
+	//   - case 1 : the component has only a set of bodies: the bodies are merged, stitched and splitted into a new set of topologically correct bodies
+	//   - case 2 : the component has only one body or has bodies + instance, only a topology healing is done on the bodies 
+	//
+	// Case 1: Repair is done before processing bodies of a component
+	// Case 2: Repair is done before getting the mesh of the body. As the body is exploded, the new bodies have to be discovered by comparing parent's bodies before and after the repair action
+
+	bool bNeedRepair = true;
+	if (!Instances.Num() && Bodies.Num() > 1 && ImportParameters.StitchingTechnique == StitchingSew)
 	{
-		if (!Instances.Num() && Bodies.Num() > 1 && ImportParameters.StitchingTechnique == StitchingSew)
-		{
-			Repair(ComponentId, StitchingSew);
-		}
-		else
-		{
-			for (CT_OBJECT_ID Body : Bodies)
-			{
-				Repair(Body, StitchingHeal);
-			}
-		}
+		// Case 1: Repair is done before processing bodies of a component
+		// Bodies.Num() > 1 so merge all bodies, sew and split into connected bodies (Repair)
+		Repair(ComponentId, StitchingSew);
 		GetInstancesAndBodies(ComponentId, Instances, Bodies);
 		SetCoreTechTessellationState(ImportParameters);
+		bNeedRepair = false;
 	}
 
 	for (CT_OBJECT_ID InstanceId : Instances)
@@ -991,7 +991,7 @@ bool FCoreTechFileParser::ReadComponent(CT_OBJECT_ID ComponentId, uint32 Default
 
 	for (CT_OBJECT_ID BodyId : Bodies)
 	{
-		if (ReadBody(BodyId, DefaultMaterialHash))
+		if (ReadBody(BodyId, ComponentId, DefaultMaterialHash, bNeedRepair))
 		{
 			SceneGraphArchive.ComponentSet[Index].Children.Add(BodyId);
 		}
@@ -1069,45 +1069,89 @@ uint32 GetBodiesFaceSetNum(TArray<CT_OBJECT_ID>& BodySet)
 	return size;
 }
 
-void FCoreTechFileParser::GetBodyTessellation(CT_OBJECT_ID BodyId, FBodyMesh& OutBodyMesh, uint32 DefaultMaterialHash)
+void FCoreTechFileParser::GetBodyTessellation(CT_OBJECT_ID BodyId, CT_OBJECT_ID ParentId, FBodyMesh& OutBodyMesh, uint32 DefaultMaterialHash, bool bNeedRepair)
 {
-	CT_LIST_IO FaceList;
-	CT_BODY_IO::AskFaces(BodyId, FaceList);
+	TArray<CT_OBJECT_ID> BodyFaces;
+	{
+		CT_LIST_IO FaceList;
+		CT_BODY_IO::AskFaces(BodyId, FaceList);
+		BodyFaces.Reserve((int32)(1.3 * FaceList.Count()));
+	}
+
+	FObjectDisplayDataId BodyMaterial;
+	BodyMaterial.DefaultMaterialName = DefaultMaterialHash;
+	GetCTObjectDisplayDataIds(BodyId, BodyMaterial);
 
-	// Compute Body BBox based on CAD data
-	uint32 VerticesSize;
-	CT_BODY_IO::AskVerticesSizeArray(BodyId, VerticesSize);
+	TArray<CT_OBJECT_ID> BodiesToProcess ;
+	if (bNeedRepair && ImportParameters.StitchingTechnique != StitchingNone)
+	{
+		// Case 2: Repair is done before getting the mesh of the body.
+		// Repair may have created new bodies including discarding initial body, so the list of bodies has to be retrieved by comparing parent's bodies before (InitialBodies) and after (AfterRepairBodies) the repair action
 
-	TArray<CT_COORDINATE> VerticesArray;
-	VerticesArray.SetNum(VerticesSize);
-	CT_BODY_IO::AskVerticesArray(BodyId, VerticesArray.GetData());
+		TArray<CT_OBJECT_ID> Instances;
+		TArray<CT_OBJECT_ID> InitialBodies;
+		GetInstancesAndBodies(ParentId, Instances, InitialBodies);
 
-	FBox& BBox = OutBodyMesh.BBox;
-	for (const CT_COORDINATE& Point : VerticesArray)
+		Repair(BodyId, ImportParameters.StitchingTechnique);
+		SetCoreTechTessellationState(ImportParameters);
+
+		TArray<CT_OBJECT_ID> AfterRepairBodies;
+		GetInstancesAndBodies(ParentId, Instances, AfterRepairBodies);
+
+		BodiesToProcess .Reserve(AfterRepairBodies.Num());
+		for (CT_OBJECT_ID Body : AfterRepairBodies)
+		{
+			if (Body == BodyId)
+			{
+				BodiesToProcess .Add(Body);
+			}
+			else if(InitialBodies.Find(Body) == INDEX_NONE)
+			{
+				BodiesToProcess .Add(Body);
+			}
+		}
+	}
+	else
 	{
-		BBox += FVector((float)Point.xyz[0], (float)Point.xyz[1], (float)Point.xyz[2]);
+		BodiesToProcess .Add(BodyId);
 	}
 
-	uint32 FaceSize = FaceList.Count();
+	FBox& BBox = OutBodyMesh.BBox;
+	for (CT_OBJECT_ID Body : BodiesToProcess )
+	{
+		CT_LIST_IO FaceList;
+		CT_BODY_IO::AskFaces(Body, FaceList);
+
+		// Compute Body BBox based on CAD data
+		uint32 VerticesSize;
+		CT_BODY_IO::AskVerticesSizeArray(Body, VerticesSize);
+
+		TArray<CT_COORDINATE> VerticesArray;
+		VerticesArray.SetNum(VerticesSize);
+		CT_BODY_IO::AskVerticesArray(Body, VerticesArray.GetData());
+
+		for (const CT_COORDINATE& Point : VerticesArray)
+		{
+			BBox += FVector((float)Point.xyz[0], (float)Point.xyz[1], (float)Point.xyz[2]);
+		}
+
+		CT_OBJECT_ID FaceID;
+		FaceList.IteratorInitialize();
+		while ((FaceID = FaceList.IteratorIter()) != 0)
+		{
+			BodyFaces.Add(FaceID);
+		}
+	}
+	uint32 FaceSize = BodyFaces.Num();
 
 	// Allocate memory space for tessellation data
 	OutBodyMesh.Faces.Reserve(FaceSize);
 	OutBodyMesh.ColorSet.Reserve(FaceSize);
 	OutBodyMesh.MaterialSet.Reserve(FaceSize);
 
-	CT_MATERIAL_ID BodyCtMaterialId = 0;
-
-	FObjectDisplayDataId BodyMaterial;
-	BodyMaterial.DefaultMaterialName = DefaultMaterialHash;
-	GetCTObjectDisplayDataIds(BodyId, BodyMaterial);
-
-	// Loop through the face of the first body and collect all tessellation data
-	FaceList.IteratorInitialize();
-
+	// Loop through the face of bodies and collect all tessellation data
 	int32 FaceIndex = 0;
-
-	CT_OBJECT_ID FaceID;
-	while ((FaceID = FaceList.IteratorIter()) != 0)
+	for(CT_OBJECT_ID FaceID : BodyFaces)
 	{
 		uint32 TriangleNum = GetFaceTessellation(FaceID, OutBodyMesh.Faces, ImportParameters);
 
@@ -1125,7 +1169,7 @@ void FCoreTechFileParser::GetBodyTessellation(CT_OBJECT_ID BodyId, FBodyMesh& Ou
 	}
 }
 
-bool FCoreTechFileParser::ReadBody(CT_OBJECT_ID BodyId, uint32 DefaultMaterialHash)
+bool FCoreTechFileParser::ReadBody(CT_OBJECT_ID BodyId, CT_OBJECT_ID ParentId, uint32 DefaultMaterialHash, bool bNeedRepair)
 {
 	// Is this body a constructive geometry ?
 	CT_LIST_IO FaceList;
@@ -1157,17 +1201,17 @@ bool FCoreTechFileParser::ReadBody(CT_OBJECT_ID BodyId, uint32 DefaultMaterialHa
 	SceneGraphArchive.BodySet[Index].MeshActorName = GetStaticMeshUuid(*SceneGraphArchive.ArchiveFileName, BodyId);
 	BodyMeshes[BodyMeshIndex].MeshActorName = SceneGraphArchive.BodySet[Index].MeshActorName;
 
-	GetBodyTessellation(BodyId, BodyMeshes[BodyMeshIndex], DefaultMaterialHash);
-
-	SceneGraphArchive.BodySet[Index].ColorFaceSet = BodyMeshes[BodyMeshIndex].ColorSet;
-	SceneGraphArchive.BodySet[Index].MaterialFaceSet = BodyMeshes[BodyMeshIndex].MaterialSet;
-
-	// Save Body in CT file for re-tessellation
+	// Save Body in CT file for re-tessellation before getBody because GetBody can call repair and modify the body (delete and build a new one with a new Id)
 	CT_LIST_IO ObjectList;
 	ObjectList.PushBack(BodyId);
 	FString BodyFile = FString::Printf(TEXT("UEx%08x"), SceneGraphArchive.BodySet[Index].MeshActorName);
 	CT_KERNEL_IO::SaveFile(ObjectList, *FPaths::Combine(CachePath, TEXT("body"), BodyFile + TEXT(".ct")), L"Ct");
 
+	GetBodyTessellation(BodyId, ParentId, BodyMeshes[BodyMeshIndex], DefaultMaterialHash, bNeedRepair);
+
+	SceneGraphArchive.BodySet[Index].ColorFaceSet = BodyMeshes[BodyMeshIndex].ColorSet;
+	SceneGraphArchive.BodySet[Index].MaterialFaceSet = BodyMeshes[BodyMeshIndex].MaterialSet;
+
 	return true;
 }
 
diff --git a/Engine/Plugins/Enterprise/DatasmithCADImporter/Source/CADInterfaces/Public/CoreTechFileParser.h b/Engine/Plugins/Enterprise/DatasmithCADImporter/Source/CADInterfaces/Public/CoreTechFileParser.h
--- a/Engine/Plugins/Enterprise/DatasmithCADImporter/Source/CADInterfaces/Public/CoreTechFileParser.h
+++ b/Engine/Plugins/Enterprise/DatasmithCADImporter/Source/CADInterfaces/Public/CoreTechFileParser.h
@@ -62,7 +62,8 @@ public:
 	double GetMetricUnit() const { return 0.01; }
 
 	EProcessResult ProcessFile(const CADLibrary::FFileDescription& InCTFileDescription);
-	void GetBodyTessellation(CT_OBJECT_ID BodyId, FBodyMesh& OutBodyMesh, uint32 ParentMaterialHash);
+
+	void GetBodyTessellation(CT_OBJECT_ID BodyId, CT_OBJECT_ID ParentId, FBodyMesh& OutBodyMesh, uint32 ParentMaterialHash, bool bNeedRepair);
 
 	TSet<FFileDescription>& GetExternalRefSet()
 	{
@@ -95,7 +96,7 @@ private:
 	bool ReadInstance(CT_OBJECT_ID NodeId, uint32 ParentMaterialHash);
 	bool ReadComponent(CT_OBJECT_ID NodeId, uint32 ParentMaterialHash);
 	bool ReadUnloadedComponent(CT_OBJECT_ID NodeId);
-	bool ReadBody(CT_OBJECT_ID NodeId, uint32 ParentMaterialHash);
+	bool ReadBody(CT_OBJECT_ID NodeId, CT_OBJECT_ID ParentId, uint32 ParentMaterialHash, bool bNeedRepair);
 
 	bool FindFile(FFileDescription& FileDescription);
 
diff --git a/Engine/Plugins/Enterprise/DatasmithCADImporter/Source/CADLibrary/Private/CoreTechHelper.cpp b/Engine/Plugins/Enterprise/DatasmithCADImporter/Source/CADLibrary/Private/CoreTechHelper.cpp
--- a/Engine/Plugins/Enterprise/DatasmithCADImporter/Source/CADLibrary/Private/CoreTechHelper.cpp
+++ b/Engine/Plugins/Enterprise/DatasmithCADImporter/Source/CADLibrary/Private/CoreTechHelper.cpp
@@ -591,7 +591,7 @@ CT_IO_ERROR Tessellate(CT_OBJECT_ID MainObjectId, const FImportParameters& Impor
 
 	while (CT_OBJECT_ID BodyId = Objects.IteratorIter())
 	{
-		Parser.GetBodyTessellation(BodyId, BodyMesh, 0);
+		Parser.GetBodyTessellation(BodyId, MainObjectId, BodyMesh, 0, /*bNeedRepair*/ true);
 	}
 
 	bool bTessellated = ConvertCTBodySetToMeshDescription(ImportParams, MeshParameters, BodyMesh, MeshDesc);
