DatasmithScene, SetVendor, SetProductName, SetProductVersion in CAD & Alias import

diff --git a/Engine/Plugins/Enterprise/DatasmithCADImporter/Source/CADInterfaces/Private/CoreTechFileParser.cpp b/Engine/Plugins/Enterprise/DatasmithCADImporter/Source/CADInterfaces/Private/CoreTechFileParser.cpp
--- a/Engine/Plugins/Enterprise/DatasmithCADImporter/Source/CADInterfaces/Private/CoreTechFileParser.cpp
+++ b/Engine/Plugins/Enterprise/DatasmithCADImporter/Source/CADInterfaces/Private/CoreTechFileParser.cpp
@@ -779,6 +779,12 @@ FCoreTechFileParser::EProcessResult FCoreTechFileParser::ReadFileWithKernelIO()
 	bool bReadNodeSucceed = ReadNode(MainId, DefaultMaterialHash);
 	// End of parsing
 
+	CT_STR KernelIO_Version = CT_KERNEL_IO::AskVersion();
+	if (!KernelIO_Version.IsEmpty())
+	{
+		SceneGraphArchive.ComponentSet[0].MetaData.Add(TEXT("KernelIOVersion"), AsFString(KernelIO_Version));
+	}
+
 	if (bNeedSaveCTFile)
 	{
 		CT_LIST_IO ObjectList;
diff --git a/Engine/Plugins/Enterprise/DatasmithCADImporter/Source/DatasmithCADTranslator/Private/DatasmithSceneGraphBuilder.cpp b/Engine/Plugins/Enterprise/DatasmithCADImporter/Source/DatasmithCADTranslator/Private/DatasmithSceneGraphBuilder.cpp
--- a/Engine/Plugins/Enterprise/DatasmithCADImporter/Source/DatasmithCADTranslator/Private/DatasmithSceneGraphBuilder.cpp
+++ b/Engine/Plugins/Enterprise/DatasmithCADImporter/Source/DatasmithCADTranslator/Private/DatasmithSceneGraphBuilder.cpp
@@ -103,9 +103,93 @@ bool FDatasmithSceneGraphBuilder::Build()
 	}
 
 	ActorData Data(TEXT(""));
-	TSharedPtr< IDatasmithActorElement > RootActor = BuildComponent(SceneGraph->ComponentSet[*Index], Data);
+	CADLibrary::FArchiveComponent& Component = SceneGraph->ComponentSet[*Index];
+	TSharedPtr< IDatasmithActorElement > RootActor = BuildComponent(Component, Data);
 	DatasmithScene->AddActor(RootActor);
 
+	// Set ProductName, ProductVersion in DatasmithScene for Analytics purpose
+	// application_name is something like "Catia V5"
+	DatasmithScene->SetVendor(TEXT("CoreTechnologie"));
+
+	if (FString* ProductVersion = Component.MetaData.Find(TEXT("KernelIOVersion")))
+	{
+		DatasmithScene->SetProductVersion(**ProductVersion);
+	}
+
+	FString* ProductName = Component.MetaData.Find(TEXT("Input_Format_and_Emitter"));
+	if(ProductName)
+	{
+		ProductName->TrimStartAndEndInline();
+		if (!ProductName->IsEmpty())
+		{
+			DatasmithScene->SetProductName(**ProductName);
+		}
+		else
+		{
+			ProductName = nullptr;
+		}
+	}
+
+	if(!ProductName)
+	{
+		if (rootFileDescription.Extension == TEXT("jt"))
+		{
+			DatasmithScene->SetProductName(TEXT("Jt"));
+		}
+		else if (rootFileDescription.Extension == TEXT("sldprt") || rootFileDescription.Extension == TEXT("sldasm"))
+		{
+			DatasmithScene->SetProductName(TEXT("SolidWorks"));
+		}
+		else if (rootFileDescription.Extension == TEXT("catpart") || rootFileDescription.Extension == TEXT("catproduct") || rootFileDescription.Extension == TEXT("cgr"))
+		{
+			DatasmithScene->SetProductName(TEXT("CATIA V5"));
+		}
+		else if (rootFileDescription.Extension == TEXT("3dxml") || rootFileDescription.Extension == TEXT("3drep"))
+		{
+			DatasmithScene->SetProductName(TEXT("3D XML"));
+		}
+		else if (rootFileDescription.Extension == TEXT("iam") || rootFileDescription.Extension == TEXT("ipt"))
+		{
+			DatasmithScene->SetProductName(TEXT("Inventor"));
+		}
+		else if (rootFileDescription.Extension == TEXT("prt") || rootFileDescription.Extension == TEXT("asm"))
+		{
+			DatasmithScene->SetProductName(TEXT("NX"));
+		}
+		else if (rootFileDescription.Extension == TEXT("stp") || rootFileDescription.Extension == TEXT("step"))
+		{
+			DatasmithScene->SetProductName(TEXT("STEP"));
+		}
+		else if (rootFileDescription.Extension == TEXT("igs") || rootFileDescription.Extension == TEXT("iges"))
+		{
+			DatasmithScene->SetProductName(TEXT("IGES"));
+		}
+		else if (rootFileDescription.Extension == TEXT("x_t") || rootFileDescription.Extension == TEXT("x_b"))
+		{
+			DatasmithScene->SetProductName(TEXT("Parasolid"));
+		}
+		else if (rootFileDescription.Extension == TEXT("dwg"))
+		{
+			DatasmithScene->SetProductName(TEXT("AutoCAD"));
+		}
+		else if (rootFileDescription.Extension == TEXT("dgn"))
+		{
+			DatasmithScene->SetProductName(TEXT("Micro Station"));
+		}
+		else if (rootFileDescription.Extension == TEXT("sat"))
+		{
+			DatasmithScene->SetProductName(TEXT("3D ACIS"));
+		}
+		else if (rootFileDescription.Extension.StartsWith(TEXT("asm")) || rootFileDescription.Extension.StartsWith(TEXT("creo")) || rootFileDescription.Extension.StartsWith(TEXT("prt")) || rootFileDescription.Extension.StartsWith(TEXT("neu")))
+		{
+			DatasmithScene->SetProductName(TEXT("Creo"));
+		}
+		else
+		{
+			DatasmithScene->SetProductName(TEXT("Unknown"));
+		}
+	}
+
 	return true;
 }
 
@@ -540,6 +624,7 @@ void FDatasmithSceneGraphBuilder::AddMetaData(TSharedPtr< IDatasmithActorElement
 		UnwantedAttributes.Add(TEXT("MaterialId"));
 		UnwantedAttributes.Add(TEXT("ColorUEId"));
 		UnwantedAttributes.Add(TEXT("ColorId"));
+		UnwantedAttributes.Add(TEXT("KernelIOVersion"));
 		return UnwantedAttributes;
 	};
 
diff --git a/Engine/Plugins/Enterprise/DatasmithCADImporter/Source/DatasmithWireTranslator/Private/DatasmithWireTranslator.cpp b/Engine/Plugins/Enterprise/DatasmithCADImporter/Source/DatasmithWireTranslator/Private/DatasmithWireTranslator.cpp
--- a/Engine/Plugins/Enterprise/DatasmithCADImporter/Source/DatasmithWireTranslator/Private/DatasmithWireTranslator.cpp
+++ b/Engine/Plugins/Enterprise/DatasmithCADImporter/Source/DatasmithWireTranslator/Private/DatasmithWireTranslator.cpp
@@ -141,9 +141,14 @@ public:
 		, FileLength(0)
 		, NumCRCErrors(0)
 	{
+		// Set ProductName, ProductVersion in DatasmithScene for Analytics purpose
+		// application_name is something like "Catia V5"
 		DatasmithScene->SetHost(TEXT("Alias"));
 		DatasmithScene->SetVendor(TEXT("Autodesk"));
 		DatasmithScene->SetExporterSDKVersion(TEXT("2019"));
+		DatasmithScene->SetProductName(TEXT("Alias Tools"));
+		DatasmithScene->SetProductVersion(TEXT("Alias 2019"));
+
 #ifdef CAD_LIBRARY
 		LocalSession = FAliasCoretechWrapper::GetSharedSession();
 #endif
