PLMXML - fixed crash with non-valid plmxml(null-pointer check)

diff --git a/Engine/Plugins/Enterprise/DatasmithCADImporter/Source/DatasmithPLMXMLTranslator/Private/DatasmithPlmXmlImporter.cpp b/Engine/Plugins/Enterprise/DatasmithCADImporter/Source/DatasmithPLMXMLTranslator/Private/DatasmithPlmXmlImporter.cpp
--- a/Engine/Plugins/Enterprise/DatasmithCADImporter/Source/DatasmithPLMXMLTranslator/Private/DatasmithPlmXmlImporter.cpp
+++ b/Engine/Plugins/Enterprise/DatasmithCADImporter/Source/DatasmithPLMXMLTranslator/Private/DatasmithPlmXmlImporter.cpp
@@ -1396,6 +1396,9 @@ bool FDatasmithPlmXmlImporter::LoadStaticMesh(const TSharedRef<IDatasmithMeshEle
 
 void FDatasmithPlmXmlImporter::UnloadScene()
 {
-	MeshLoader->UnloadScene();
+	if (MeshLoader.IsValid())
+	{
+		MeshLoader->UnloadScene();
+	}
 }
 #undef LOCTEXT_NAMESPACE
