CADKernel v0:
 - CADKernel is integrated as a static library
 - CADKernel::Toolkit::FCoreTechBridge allows to transfer CAD data to CADKernel
 - The mesh is generated into CADKernel
 - The mesh is extracted from CADKernel and is treated as with KernelIO

diff --git a/Engine/Plugins/Enterprise/DatasmithCADImporter/Source/CADInterfaces/Private/CoreTechFileParser.cpp b/Engine/Plugins/Enterprise/DatasmithCADImporter/Source/CADInterfaces/Private/CoreTechFileParser.cpp
--- a/Engine/Plugins/Enterprise/DatasmithCADImporter/Source/CADInterfaces/Private/CoreTechFileParser.cpp
+++ b/Engine/Plugins/Enterprise/DatasmithCADImporter/Source/CADInterfaces/Private/CoreTechFileParser.cpp
@@ -16,6 +16,9 @@
 #include "Misc/Paths.h"
 #include "Templates/TypeHash.h"
 
+#include "CADKernel/Mesh/Structure/SurfaceMesh.h"
+#include "CADKernel/Math/Point.h"
+
 namespace CADLibrary 
 {
 
@@ -212,7 +215,7 @@ void FillArrayOfInt(int32 ElementCount, void* InCTValueArray, int32* OutValueArr
 	}
 }
 
-uint32 GetFaceTessellation(CT_OBJECT_ID FaceID, TArray<FTessellationData>& FaceTessellationSet, const FImportParameters& ImportParams)
+uint32 GetKioFaceTessellation(CT_OBJECT_ID FaceID, TArray<FTessellationData>& FaceTessellationSet, const FImportParameters& ImportParams)
 {
 	CT_IO_ERROR Error = IO_OK;
 
@@ -244,29 +247,31 @@ uint32 GetFaceTessellation(CT_OBJECT_ID FaceID, TArray<FTessellationData>& FaceT
 	}
 
 	FTessellationData& Tessellation = FaceTessellationSet.Emplace_GetRef();
-	Tessellation.IndexArray.SetNum(IndexCount);
+	Tessellation.VertexIndices.SetNum(IndexCount);
+
+	Tessellation.PatchId = (int32) FaceID;
 
 	switch (IndexType)
 	{
 	case CT_TESS_UBYTE:
-		FillArrayOfInt<uint8>(IndexCount, IndexArray, Tessellation.IndexArray.GetData());
+		FillArrayOfInt<uint8>(IndexCount, IndexArray, Tessellation.VertexIndices.GetData());
 		break;
 	case CT_TESS_USHORT:
-		FillArrayOfInt<uint16>(IndexCount, IndexArray, Tessellation.IndexArray.GetData());
+		FillArrayOfInt<uint16>(IndexCount, IndexArray, Tessellation.VertexIndices.GetData());
 		break;
 	case CT_TESS_UINT:
-		FillArrayOfInt<uint32>(IndexCount, IndexArray, Tessellation.IndexArray.GetData());
+		FillArrayOfInt<uint32>(IndexCount, IndexArray, Tessellation.VertexIndices.GetData());
 		break;
 	}
 
-	Tessellation.VertexArray.SetNum(VertexCount);
+	Tessellation.PositionArray.SetNum(VertexCount);
 	switch (VertexType)
 	{
 	case CT_TESS_FLOAT:
-		FillArrayOfVector<float>(VertexCount, VertexArray, Tessellation.VertexArray.GetData());
+		FillArrayOfVector<float>(VertexCount, VertexArray, Tessellation.PositionArray.GetData());
 		break;
 	case CT_TESS_DOUBLE:
-		FillArrayOfVector<double>(VertexCount, VertexArray, Tessellation.VertexArray.GetData());
+		FillArrayOfVector<double>(VertexCount, VertexArray, Tessellation.PositionArray.GetData());
 		break;
 	}
 
@@ -318,10 +323,9 @@ uint32 GetFaceTessellation(CT_OBJECT_ID FaceID, TArray<FTessellationData>& FaceT
 		ScaleUV(FaceID, Tessellation.TexCoordArray, (float) ImportParams.ScaleFactor);
 	}
 
-	return Tessellation.IndexArray.Num() / 3;
+	return Tessellation.VertexIndices.Num() / 3;
 }
 
-
 void GetCTObjectDisplayDataIds(CT_OBJECT_ID ObjectID, FObjectDisplayDataId& Material)
 {
 	if (CT_OBJECT_IO::SearchAttribute(ObjectID, CT_ATTRIB_MATERIALID) == IO_OK)
@@ -485,7 +489,7 @@ void FCoreTechFileParser::SetFaceMainMaterial(FObjectDisplayDataId& InFaceMateri
 		FaceTessellations.ColorName = Color.UEMaterialName;
 		BodyMesh.ColorSet.Add(Color.UEMaterialName);
 	}
-	else if(InBodyMaterial.DefaultMaterialName)
+	else if (InBodyMaterial.DefaultMaterialName)
 	{
 		FaceTessellations.ColorName = InBodyMaterial.DefaultMaterialName;
 		BodyMesh.ColorSet.Add(InBodyMaterial.DefaultMaterialName);
@@ -506,7 +510,7 @@ void FCoreTechFileParser::ExportSceneGraphFile()
 
 void FCoreTechFileParser::ExportMeshArchiveFile()
 {
-	SerializeBodyMeshSet(*MeshArchiveFilePath, BodyMeshes);
+		SerializeBodyMeshSet(*MeshArchiveFilePath, BodyMeshes);
 }
 
 void FCoreTechFileParser::LoadSceneGraphArchive(const FString& SGFile)
@@ -568,8 +572,10 @@ void FCoreTechFileParser::ReadMaterials()
 FCoreTechFileParser::FCoreTechFileParser(const FImportParameters& ImportParams, const FString& EnginePluginsPath, const FString& InCachePath)
 	: CachePath(InCachePath)
 	, ImportParameters(ImportParams)
+	, CoreTechBridge()
 {
 	CTKIO_InitializeKernel(ImportParameters.MetricUnit, *EnginePluginsPath);
+	CoreTechBridge.InitializeCADKernel(0.00001 / ImportParams.MetricUnit);
 }
 
 bool FCoreTechFileParser::FindFile(FFileDescription& File)
@@ -1125,7 +1131,131 @@ uint32 GetBodiesFaceSetNum(TArray<CT_OBJECT_ID>& BodySet)
 	return size;
 }
 
-void FCoreTechFileParser::GetBodyTessellation(CT_OBJECT_ID BodyId, CT_OBJECT_ID ParentId, FBodyMesh& OutBodyMesh, uint32 DefaultMaterialHash, bool bNeedRepair)
+void FCoreTechFileParser::DefineMeshCriteria(FIdent MeshModelId)
+{
+	TArray<FIdent> Criteria;
+
+	{
+		FIdent CriterionId = CoreTechBridge.CreateCadCurvatureCriterion();
+		Criteria.Add(CriterionId);
+	}
+	
+	if (ImportParameters.MaxEdgeLength > SMALL_NUMBER)
+	{
+		FIdent CriterionId = CoreTechBridge.CreateMaxLengthCriterion(ImportParameters.MaxEdgeLength);
+		Criteria.Add(CriterionId);
+	}
+
+	if (ImportParameters.ChordTolerance > SMALL_NUMBER)
+	{
+		FIdent CriterionId = CoreTechBridge.CreateSAGCriterion(ImportParameters.ChordTolerance);
+		Criteria.Add(CriterionId);
+	}
+
+	if (ImportParameters.MaxNormalAngle > SMALL_NUMBER)
+	{
+		FIdent CriterionId = CoreTechBridge.CreateAngleCriterion(ImportParameters.MaxNormalAngle);
+		Criteria.Add(CriterionId);
+	}
+	CoreTechBridge.AddCriteria(MeshModelId, Criteria);
+}
+
+uint32 FCoreTechFileParser::GetFaceTessellation(FIdent FaceID, TArray<FTessellationData>& FaceTessellationSet, const FImportParameters& ImportParams)
+{
+	CT_IO_ERROR Error = IO_OK;
+
+	const CADKernel::FSurfaceMesh& SurfaceMesh = CoreTechBridge.GetFaceTesselation(FaceID);
+
+	int32 IndexCount = SurfaceMesh.TrianglesVerticesIndex.Num();
+	if (IndexCount == 0)
+	{
+		return 0;
+	}
+
+	FTessellationData& Tessellation = FaceTessellationSet.Emplace_GetRef();
+	Tessellation.PatchId = (int32) SurfaceMesh.GetId();
+
+	Tessellation.PositionIndices.Reserve(SurfaceMesh.VerticesGlobalIndex.Num());
+	for (int32 VertexIndex : SurfaceMesh.VerticesGlobalIndex)
+	{
+		Tessellation.PositionIndices.Add(VertexIndex);
+	}
+
+	Tessellation.VertexIndices.Reserve(IndexCount);
+	for(int32 Index : SurfaceMesh.TrianglesVerticesIndex)
+	{
+		Tessellation.VertexIndices.Add(Index);
+	}
+
+	Tessellation.NormalArray.Reserve(SurfaceMesh.Normals.Num());
+	for (const CADKernel::FPoint& Normal : SurfaceMesh.Normals)
+	{
+		Tessellation.NormalArray.Emplace((float)Normal.X, (float)Normal.Y, (float)Normal.Z);
+	}
+
+	Tessellation.TexCoordArray.Reserve(SurfaceMesh.UVMap.Num());
+	for (const CADKernel::FPoint2D& TexCoord : SurfaceMesh.UVMap)
+	{
+		Tessellation.TexCoordArray.Emplace((float)TexCoord.U, (float) TexCoord.V);
+	}
+
+	return Tessellation.VertexIndices.Num() / 3;
+}
+
+void FCoreTechFileParser::GetBodyTessellation(CT_OBJECT_ID CTBodyId, CT_OBJECT_ID ParentId, FBodyMesh& OutBodyMesh, uint32 DefaultMaterialHash, bool bNeedRepair, FString BodyFile)
+{
+	FObjectDisplayDataId BodyMaterial;
+	BodyMaterial.DefaultMaterialName = DefaultMaterialHash;
+	GetCTObjectDisplayDataIds(CTBodyId, BodyMaterial);
+
+	CoreTechBridge.Empty();
+	CoreTechBridge.DefineLogFile(*FPaths::Combine(CachePath, TEXT("CADKernel"), BodyFile + TEXT(".log")));
+
+	FIdent DSModelId = CoreTechBridge.AddModel();
+
+	TArray<FIdent> DSBodies;
+	FIdent DSBodyId = CoreTechBridge.AddBody(CTBodyId);
+	DSBodies.Add(DSBodyId);
+	CoreTechBridge.AddToModel(DSModelId, DSBodyId);
+
+	CoreTechBridge.SaveDatabase(*FPaths::Combine(CachePath, TEXT("CADKernel"), BodyFile + TEXT(".dtc")));
+
+	FIdent DSMeshModelId = CoreTechBridge.CreateMeshModel();
+	DefineMeshCriteria(DSMeshModelId);
+
+	CoreTechBridge.Mesh(DSMeshModelId, DSBodies);
+
+	CoreTechBridge.GetMeshModelPointCloud(DSMeshModelId, OutBodyMesh.VertexArray);
+   
+	const TArray<FIdent>& FacesIds = CoreTechBridge.GetFaceIds();
+	uint32 FaceSize = FacesIds.Num();
+
+	// Allocate memory space for tessellation data
+	OutBodyMesh.Faces.Reserve(FaceSize);
+	OutBodyMesh.ColorSet.Reserve(FaceSize);
+	OutBodyMesh.MaterialSet.Reserve(FaceSize);
+
+	// Loop through the face of bodies and collect all tessellation data
+	int32 FaceIndex = 0;
+	for (FIdent FaceID : FacesIds)
+	{
+		uint32 TriangleNum = GetFaceTessellation(FaceID, OutBodyMesh.Faces, ImportParameters);
+
+		if (TriangleNum == 0)
+		{
+			continue;
+		}
+
+		OutBodyMesh.TriangleCount += TriangleNum;
+
+		FObjectDisplayDataId FaceMaterial;
+		GetCTObjectDisplayDataIds(FaceID, FaceMaterial);
+		SetFaceMainMaterial(FaceMaterial, BodyMaterial, OutBodyMesh, FaceIndex);
+		FaceIndex++;
+	}
+}
+
+void FCoreTechFileParser::GetKioBodyTessellation(CT_OBJECT_ID BodyId, CT_OBJECT_ID ParentId, FBodyMesh& OutBodyMesh, uint32 DefaultMaterialHash, bool bNeedRepair)
 {
 	TArray<CT_OBJECT_ID> BodyFaces;
 	{
@@ -1209,7 +1339,7 @@ void FCoreTechFileParser::GetBodyTessellation(CT_OBJECT_ID BodyId, CT_OBJECT_ID
 	int32 FaceIndex = 0;
 	for(CT_OBJECT_ID FaceID : BodyFaces)
 	{
-		uint32 TriangleNum = GetFaceTessellation(FaceID, OutBodyMesh.Faces, ImportParameters);
+		uint32 TriangleNum = GetKioFaceTessellation(FaceID, OutBodyMesh.Faces, ImportParameters);
 
 		if (TriangleNum == 0)
 		{
@@ -1261,11 +1391,20 @@ bool FCoreTechFileParser::ReadBody(CT_OBJECT_ID BodyId, CT_OBJECT_ID ParentId, u
 	FString BodyFile = FString::Printf(TEXT("UEx%08x"), SceneGraphArchive.BodySet[Index].MeshActorName);
 	CT_KERNEL_IO::SaveFile(ObjectList, *FPaths::Combine(CachePath, TEXT("body"), BodyFile + TEXT(".ct")), L"Ct");
 
-	GetBodyTessellation(BodyId, ParentId, BodyMeshes[BodyMeshIndex], DefaultMaterialHash, bNeedRepair);
+	CT_FLAGS BodyProperties;
+	CT_BODY_IO::AskProperties(BodyId, BodyProperties);
 
+	if (ImportParameters.bEnableKernelIOTessellation || !(BodyProperties & CT_BODY_PROP_EXACT))
+	{
+		GetKioBodyTessellation(BodyId, ParentId, BodyMeshes[BodyMeshIndex], DefaultMaterialHash, bNeedRepair);
+	}
+	else
+	{
+		GetBodyTessellation(BodyId, ParentId, BodyMeshes[BodyMeshIndex], DefaultMaterialHash, bNeedRepair, BodyFile);
+	}
 	SceneGraphArchive.BodySet[Index].ColorFaceSet = BodyMeshes[BodyMeshIndex].ColorSet;
 	SceneGraphArchive.BodySet[Index].MaterialFaceSet = BodyMeshes[BodyMeshIndex].MaterialSet;
-
+	
 	return true;
 }
 
diff --git a/Engine/Plugins/Enterprise/DatasmithCADImporter/Source/CADInterfaces/Public/CoreTechFileParser.h b/Engine/Plugins/Enterprise/DatasmithCADImporter/Source/CADInterfaces/Public/CoreTechFileParser.h
--- a/Engine/Plugins/Enterprise/DatasmithCADImporter/Source/CADInterfaces/Public/CoreTechFileParser.h
+++ b/Engine/Plugins/Enterprise/DatasmithCADImporter/Source/CADInterfaces/Public/CoreTechFileParser.h
@@ -8,6 +8,7 @@
 #include "CADSceneGraph.h"
 #include "Containers/Map.h"
 #include "Containers/Queue.h"
+#include "Containers/UnrealString.h"
 #include "Misc/Paths.h"
 
 #ifdef CAD_INTERFACE
@@ -17,6 +18,9 @@
 #include "kernel_io/material_io/material_io.h"
 #endif // CAD_INTERFACE
 
+#ifdef CAD_INTERFACE // temporary in the wait of a full integration in UE
+#include "CADKernel/CADInterface/CoreTech/CoreTechBridge.h"
+#endif // CAD_INTERFACE
 
 struct FFileStatData;
 struct FFileDescription;
@@ -63,7 +67,11 @@ public:
 
 	EProcessResult ProcessFile(const CADLibrary::FFileDescription& InCTFileDescription);
 
-	void GetBodyTessellation(CT_OBJECT_ID BodyId, CT_OBJECT_ID ParentId, FBodyMesh& OutBodyMesh, uint32 ParentMaterialHash, bool bNeedRepair);
+	void GetKioBodyTessellation(CT_OBJECT_ID BodyId, CT_OBJECT_ID ParentId, FBodyMesh& OutBodyMesh, uint32 ParentMaterialHash, bool bNeedRepair);
+
+	void GetBodyTessellation(CT_OBJECT_ID BodyId, CT_OBJECT_ID ParentId, FBodyMesh& OutBodyMesh, uint32 ParentMaterialHash, bool bNeedRepair, FString BodyFile);
+	uint32 GetFaceTessellation(FIdent FaceID, TArray<FTessellationData>& FaceTessellationSet, const FImportParameters& ImportParams);
+	void DefineMeshCriteria(FIdent MeshModelId);
 
 	TSet<FFileDescription>& GetExternalRefSet()
 	{
@@ -134,6 +142,8 @@ protected:
 	TArray<FBodyMesh> BodyMeshes;
 
 	const FImportParameters& ImportParameters;
+
+	CADKernel::FCoreTechBridge CoreTechBridge;
 };
 
 #endif // CAD_INTERFACE
diff --git a/Engine/Plugins/Enterprise/DatasmithCADImporter/Source/CADLibrary/Private/CoreTechHelper.cpp b/Engine/Plugins/Enterprise/DatasmithCADImporter/Source/CADLibrary/Private/CoreTechHelper.cpp
--- a/Engine/Plugins/Enterprise/DatasmithCADImporter/Source/CADLibrary/Private/CoreTechHelper.cpp
+++ b/Engine/Plugins/Enterprise/DatasmithCADImporter/Source/CADLibrary/Private/CoreTechHelper.cpp
@@ -68,7 +68,7 @@ namespace CADLibrary
 		return (VertexIDs[0] == VertexIDs[1] || VertexIDs[0] == VertexIDs[2] || VertexIDs[1] == VertexIDs[2]);
 	}
 
-	void FillVertexPosition(const FImportParameters& ImportParams, const FMeshParameters& MeshParameters, FBodyMesh& Body, FMeshDescription& MeshDescription)
+	void FillKioVertexPosition(const FImportParameters& ImportParams, const FMeshParameters& MeshParameters, FBodyMesh& Body, FMeshDescription& MeshDescription)
 	{
 		int32 TriangleCount = Body.TriangleCount;
 		TArray<FTessellationData>& FaceTessellationSet = Body.Faces;
@@ -87,11 +87,10 @@ namespace CADLibrary
 		VertexDataSet.Reserve(TriangleCount * 3);
 
 		FVector Position;
-		uint32 GlobalVertexCount = 0;
+		int32 GlobalVertexCount = 0;
 		for (FTessellationData& CTTessellation : FaceTessellationSet)
 		{
-			CTTessellation.StartVertexIndex = GlobalVertexCount;
-			for (const FVector& Vertex : CTTessellation.VertexArray)
+			for (const FVector& Vertex : CTTessellation.PositionArray)
 			{
 				VertexDataSet.Emplace(GlobalVertexCount, Vertex * ImportParams.ScaleFactor);
 				++GlobalVertexCount;
@@ -126,7 +125,7 @@ namespace CADLibrary
 			const FVector& PositionA = VertexDataSet[i].Coordinates;
 			if (BBox.IsValid && !BBox.IsInside(PositionA))
 			{
-				VertexDataSet[i].Index = -1;
+				VertexDataSet[i].Index = INDEX_NONE;
 				continue;
 			}
 
@@ -148,7 +147,6 @@ namespace CADLibrary
 			VertexCount++;
 		}
 
-
 		// if Symmetric mesh, the symmetric side of the mesh have to be generated
 		FMatrix SymmetricMatrix;
 		bool bIsSymmetricMesh = MeshParameters.bIsSymmetric;
@@ -157,12 +155,13 @@ namespace CADLibrary
 			SymmetricMatrix = FDatasmithUtils::GetSymmetricMatrix(MeshParameters.SymmetricOrigin, MeshParameters.SymmetricNormal);
 		}
 
-
 		// Make MeshDescription.VertexPositions and VertexID
-		MeshDescription.ReserveNewVertices(VertexCount);
-		uint32 GlobalVertexIndex = 0;
+		MeshDescription.ReserveNewVertices(bIsSymmetricMesh ? VertexCount * 2 : VertexCount);
+		Body.VertexIds.Reserve(VertexCount);
 
-		for (uint32 VertexIndex = 0; VertexIndex < GlobalVertexCount; ++VertexIndex)
+		int32 GlobalVertexIndex = 0;
+
+		for (int32 VertexIndex = 0; VertexIndex < GlobalVertexCount; ++VertexIndex)
 		{
 			int32 RealIndex = VertexDataSet[VertexIndex].Index;
 
@@ -178,13 +177,15 @@ namespace CADLibrary
 			}
 
 			FVertexID VertexID = MeshDescription.CreateVertex();
+			Body.VertexIds.Add(VertexID);
 			VertexPositions[VertexID] = FDatasmithUtils::ConvertVector((FDatasmithUtils::EModelCoordSystem) ImportParams.ModelCoordSys, VertexDataSet[VertexIndex].Coordinates);
 			VertexDataSet[VertexIndex].VertexID = VertexID;
 		}
 
 		if (bIsSymmetricMesh)
 		{
-			for (uint32 VertexIndex = 0; VertexIndex < GlobalVertexCount; ++VertexIndex)
+			Body.SymmetricVertexIds.Reserve(Body.VertexIds.Num());
+			for (int32 VertexIndex = 0; VertexIndex < GlobalVertexCount; ++VertexIndex)
 			{
 				int32 RealIndex = VertexDataSet[VertexIndex].Index;
 
@@ -200,6 +201,7 @@ namespace CADLibrary
 				}
 
 				FVertexID VertexID = MeshDescription.CreateVertex();
+				Body.SymmetricVertexIds.Add(VertexID);
 				VertexPositions[VertexID] = FDatasmithUtils::ConvertVector((FDatasmithUtils::EModelCoordSystem) ImportParams.ModelCoordSys, VertexDataSet[VertexIndex].Coordinates);
 				VertexPositions[VertexID] = SymmetricMatrix.TransformPosition(VertexPositions[VertexID]);
 
@@ -211,31 +213,93 @@ namespace CADLibrary
 		GlobalVertexIndex = 0;
 		for (FTessellationData& CTTessellation : FaceTessellationSet)
 		{
-			CTTessellation.VertexIdSet.SetNum(CTTessellation.VertexArray.Num());
-			for (int32 VertexIndex = 0; VertexIndex < CTTessellation.VertexArray.Num(); ++VertexIndex, ++GlobalVertexIndex)
+			CTTessellation.PositionIndices.SetNum(CTTessellation.PositionArray.Num());
+			for (int32 VertexIndex = 0; VertexIndex < CTTessellation.PositionArray.Num(); ++VertexIndex, ++GlobalVertexIndex)
 			{
 				int32 NewIndex = NewIndexOf[IndexOfCoincidentNode[GlobalVertexIndex]];
-				CTTessellation.VertexIdSet[VertexIndex] = VertexDataSet[NewIndex].VertexID.GetValue();
+				CTTessellation.PositionIndices[VertexIndex] = VertexDataSet[NewIndex].VertexID.GetValue();
 			}
 		}
+	}
 
-		if (bIsSymmetricMesh)
+	void FillVertexPosition(const FImportParameters& ImportParams, const FMeshParameters& MeshParameters, FBodyMesh& Body, FMeshDescription& MeshDescription)
+	{
+		int32 TriangleCount = Body.TriangleCount;
+		TArray<FTessellationData>& FaceTessellationSet = Body.Faces;
+
+		// Add offset to the bounding box to avoid to remove good vertex
+		FVector Size = Body.BBox.GetSize();
+		FBox BBox = Body.BBox.ExpandBy(Size.Size());
+		BBox.IsValid = Body.BBox.IsValid;
+		BBox.Min *= ImportParams.ScaleFactor;
+		BBox.Max *= ImportParams.ScaleFactor;
+
+		TVertexAttributesRef<FVector> VertexPositions = MeshDescription.GetVertexPositions();
+
+		FVector Position;
+
+		const TArray<FVector>& VertexArray = Body.VertexArray;
+		TArray<int32>& VertexIdSet = Body.VertexIds;
+		int32 VertexCount = VertexArray.Num();
+		VertexIdSet.SetNumZeroed(VertexCount);
+
+		if (BBox.IsValid)
 		{
+			int32 VertexIndex = 0;
+			for (const FVector& Vertex : VertexArray)
+			{
+				if (!BBox.IsInside(Vertex))
+				{
+					VertexIdSet[VertexIndex] = INDEX_NONE;
+				}
+				VertexIndex++;
+			}
+		}
 
-			GlobalVertexIndex = 0;
-			for (FTessellationData& CTTessellation : FaceTessellationSet)
+		// Make MeshDescription.VertexPositions and VertexID
+		MeshDescription.ReserveNewVertices(MeshParameters.bIsSymmetric ? VertexCount * 2 : VertexCount);
+
+		int32 VertexIndex = -1;
+		for (const FVector& Vertex : VertexArray)
+		{
+			VertexIndex++;
+
+			// Vertex is outside bbox
+			if (VertexIdSet[VertexIndex] == INDEX_NONE)
 			{
-				CTTessellation.SymVertexIdSet.SetNum(CTTessellation.VertexArray.Num());
-				for (int32 VertexIndex = 0; VertexIndex < CTTessellation.VertexArray.Num(); ++VertexIndex, ++GlobalVertexIndex)
+				continue;
+			}
+
+			FVertexID VertexID = MeshDescription.CreateVertex();
+			VertexPositions[VertexID] = FDatasmithUtils::ConvertVector((FDatasmithUtils::EModelCoordSystem) ImportParams.ModelCoordSys, Vertex);
+			VertexIdSet[VertexIndex] = VertexID;
+		}
+
+		// if Symmetric mesh, the symmetric side of the mesh have to be generated
+		if (MeshParameters.bIsSymmetric)
+		{
+			FMatrix SymmetricMatrix = FDatasmithUtils::GetSymmetricMatrix(MeshParameters.SymmetricOrigin, MeshParameters.SymmetricNormal);
+
+			TArray<int32>& SymmetricVertexIds = Body.SymmetricVertexIds;
+			SymmetricVertexIds.SetNum(VertexArray.Num());
+
+			VertexIndex = 0;
+			for (const FVector& Vertex : VertexArray)
+			{
+				if (VertexIdSet[VertexIndex] == INDEX_NONE)
 				{
-					uint32 NewIndex = NewIndexOf[IndexOfCoincidentNode[GlobalVertexIndex]];
-					CTTessellation.SymVertexIdSet[VertexIndex] = VertexDataSet[NewIndex].SymVertexID.GetValue();
+					SymmetricVertexIds[VertexIndex++] = INDEX_NONE;
+					continue;
 				}
+
+				FVertexID VertexID = MeshDescription.CreateVertex();
+				VertexPositions[VertexID] = FDatasmithUtils::ConvertVector((FDatasmithUtils::EModelCoordSystem) ImportParams.ModelCoordSys, Vertex);
+				VertexPositions[VertexID] = SymmetricMatrix.TransformPosition(VertexPositions[VertexID]);
+				SymmetricVertexIds[VertexIndex++] = VertexID;
 			}
 		}
 	}
 
-
 	// PolygonAttributes name used into modeling tools (ExtendedMeshAttribute::PolyTriGroups)
 	const FName PolyTriGroups("PolyTriGroups");
 
@@ -289,27 +353,19 @@ namespace CADLibrary
 		for (const FPolygonGroupID PolygonID : MeshDestination.PolygonGroups().GetElementIDs())
 		{
 			int32 PatchId = PatchGroups[PolygonID];
-			if (PatchId > 0) 
+			if (PatchId > 0)
 			{
 				OutPatchIdSet.Add(PatchId);
 			}
 		}
 	}
 
-	bool FillMesh(const FMeshParameters& MeshParameters, const FImportParameters& ImportParams, TArray<FTessellationData>& FaceTessellations, FMeshDescription& MeshDescription)
+	bool FillMesh(const FMeshParameters& MeshParameters, const FImportParameters& ImportParams, FBodyMesh& BodyTessellation, FMeshDescription& MeshDescription)
 	{
 		const int32 UVChannel = 0;
-		const int32 TriangleCount = 3;
+		const int32 VerticesPerTriangle = 3;
 		const TriangleIndex Clockwise = { 0, 1, 2 };
 		const TriangleIndex CounterClockwise = { 0, 2, 1 };
-		const int32 InvalidID = INDEX_NONE;
-
-
-		TArray<FVertexInstanceID> TriangleVertexInstanceIDs;
-		TriangleVertexInstanceIDs.SetNum(TriangleCount);
-
-		TArray<FVertexInstanceID> MeshVertexInstanceIDs;
-		TArray<uint32> CTFaceIndex;  // new CT face index to remove degenerated face
 
 		// Gather all array data
 		FStaticMeshAttributes Attributes(MeshDescription);
@@ -327,7 +383,7 @@ namespace CADLibrary
 
 		// Find all the materials used
 		TMap<uint32, FPolygonGroupID> MaterialToPolygonGroupMapping;
-		for (const FTessellationData& FaceTessellation : FaceTessellations)
+		for (const FTessellationData& FaceTessellation : BodyTessellation.Faces)
 		{
 			// we assume that face has only color
 			MaterialToPolygonGroupMapping.Add(FaceTessellation.ColorName, INDEX_NONE);
@@ -352,22 +408,25 @@ namespace CADLibrary
 			NbStep = 2;
 		}
 
-		TSet<int32> PatchIdSet;
-		GetExistingPatches(MeshDescription, PatchIdSet);
-		bool bImportOnlyAlreadyPresent = (bool) PatchIdSet.Num();
+		TArray<FVertexInstanceID> TriangleVertexInstanceIDs;
+		TriangleVertexInstanceIDs.SetNum(VerticesPerTriangle);
+
+		TArray<FVertexInstanceID> MeshVertexInstanceIDs;
+		MeshVertexInstanceIDs.Reserve(BodyTessellation.VertexIds.Num() * NbStep);
+
+		TArray<uint32> NewFaceIndex;  // NewFaceIndex = Tessellation.TrianglesVerticesIndex without degenerated face
 
 		TPolygonAttributesRef<int32> PatchGroups = EnableCADPatchGroups(MeshDescription);
+		int32 PatchIndex = 0;
 		for (int32 Step = 0; Step < NbStep; ++Step)
 		{
 			// Swap mesh if needed
 			const TriangleIndex& Orientation = (!MeshParameters.bNeedSwapOrientation == (bool)Step) ? CounterClockwise : Clockwise;
+			TArray<int32>& VertexIdSet = (Step == 0) ? BodyTessellation.VertexIds : BodyTessellation.SymmetricVertexIds;
 
-			for (FTessellationData& Tessellation : FaceTessellations)
+			for (FTessellationData& Tessellation : BodyTessellation.Faces)
 			{
-				if (bImportOnlyAlreadyPresent && !PatchIdSet.Contains(Tessellation.PatchId))
-				{
-					continue;
-				}
+				int32 VertexIDs[3];
 
 				// Get the polygonGroup
 				const FPolygonGroupID* PolygonGroupID = MaterialToPolygonGroupMapping.Find(Tessellation.ColorName);
@@ -376,25 +435,18 @@ namespace CADLibrary
 					continue;
 				}
 
-				//int32 TriangleCount = IndicesCount / 3;
-				int32 VertexIDs[3];
-				FVector Temp3D = { 0, 0, 0 };
-				FVector2D TexCoord2D = { 0, 0 };
-
-				MeshVertexInstanceIDs.SetNum(Tessellation.IndexArray.Num());
-				CTFaceIndex.Reserve(Tessellation.IndexArray.Num());
-				CTFaceIndex.SetNum(0);
-
-				TArray<int32>& VertexIdSet = (Step == 0) ? Tessellation.VertexIdSet : Tessellation.SymVertexIdSet;
+				PatchIndex++;
+				MeshVertexInstanceIDs.Empty(Tessellation.VertexIndices.Num());
+				NewFaceIndex.Empty(Tessellation.VertexIndices.Num());
 
 				// build each valid face i.e. 3 different indexes
-				for (int32 Index = 0, NewIndex = 0; Index < Tessellation.IndexArray.Num(); Index += 3)
+				for (int32 Index = 0; Index < Tessellation.VertexIndices.Num(); Index += 3)
 				{
-					VertexIDs[Orientation[0]] = VertexIdSet[Tessellation.IndexArray[Index + 0]];
-					VertexIDs[Orientation[1]] = VertexIdSet[Tessellation.IndexArray[Index + 1]];
-					VertexIDs[Orientation[2]] = VertexIdSet[Tessellation.IndexArray[Index + 2]];
+					VertexIDs[Orientation[0]] = VertexIdSet[Tessellation.PositionIndices[Tessellation.VertexIndices[Index + 0]]];
+					VertexIDs[Orientation[1]] = VertexIdSet[Tessellation.PositionIndices[Tessellation.VertexIndices[Index + 1]]];
+					VertexIDs[Orientation[2]] = VertexIdSet[Tessellation.PositionIndices[Tessellation.VertexIndices[Index + 2]]];
 
-					if (VertexIDs[0] == InvalidID || VertexIDs[1] == InvalidID || VertexIDs[2] == InvalidID)
+					if (VertexIDs[0] == INDEX_NONE || VertexIDs[1] == INDEX_NONE || VertexIDs[2] == INDEX_NONE)
 					{
 						continue;
 					}
@@ -405,45 +457,34 @@ namespace CADLibrary
 						continue;
 					}
 
-					CTFaceIndex.Add(Tessellation.IndexArray[Index + 0]);
-					CTFaceIndex.Add(Tessellation.IndexArray[Index + 1]);
-					CTFaceIndex.Add(Tessellation.IndexArray[Index + 2]);
+					NewFaceIndex.Add(Tessellation.VertexIndices[Index + 0]);
+					NewFaceIndex.Add(Tessellation.VertexIndices[Index + 1]);
+					NewFaceIndex.Add(Tessellation.VertexIndices[Index + 2]);
 
-					TriangleVertexInstanceIDs[0] = MeshVertexInstanceIDs[NewIndex++] = MeshDescription.CreateVertexInstance((FVertexID)VertexIDs[0]);
-					TriangleVertexInstanceIDs[1] = MeshVertexInstanceIDs[NewIndex++] = MeshDescription.CreateVertexInstance((FVertexID)VertexIDs[1]);
-					TriangleVertexInstanceIDs[2] = MeshVertexInstanceIDs[NewIndex++] = MeshDescription.CreateVertexInstance((FVertexID)VertexIDs[2]);
+					MeshVertexInstanceIDs.Add(TriangleVertexInstanceIDs[0] = MeshDescription.CreateVertexInstance((FVertexID)VertexIDs[0]));
+					MeshVertexInstanceIDs.Add(TriangleVertexInstanceIDs[1] = MeshDescription.CreateVertexInstance((FVertexID)VertexIDs[1]));
+					MeshVertexInstanceIDs.Add(TriangleVertexInstanceIDs[2] = MeshDescription.CreateVertexInstance((FVertexID)VertexIDs[2]));
 
 					// Add the triangle as a polygon to the mesh description
 					const FPolygonID PolygonID = MeshDescription.CreatePolygon(*PolygonGroupID, TriangleVertexInstanceIDs);
+
 					// Set patch id attribute
 					PatchGroups[PolygonID] = Tessellation.PatchId;
 				}
 
-				// finalization of the mesh by setting colors, tangents, bi-normals, UV
-				for (int32 IndexFace = 0; IndexFace < CTFaceIndex.Num(); IndexFace += 3)
+				for (int32 IndexFace = 0; IndexFace < MeshVertexInstanceIDs.Num(); IndexFace += 3)
 				{
-					for (int32 Index = 0; Index < TriangleCount; Index++)
+					for (int32 Index = 0; Index < VerticesPerTriangle; Index++)
 					{
-						FVertexInstanceID VertexInstanceID = MeshVertexInstanceIDs[IndexFace + Orientation[Index]];
+						const FVertexInstanceID VertexInstanceID = MeshVertexInstanceIDs[IndexFace + Orientation[Index]];
+						VertexInstanceUVs.Set(VertexInstanceID, UVChannel, Tessellation.TexCoordArray[NewFaceIndex[IndexFace + Index]]);
 
 						VertexInstanceColors[VertexInstanceID] = FLinearColor::White;
-						VertexInstanceTangents[VertexInstanceID] = FVector(ForceInitToZero);
+						VertexInstanceTangents[VertexInstanceID] = FVector::ZeroVector;
 						VertexInstanceBinormalSigns[VertexInstanceID] = 0.0f;
 					}
 				}
 
-				if (Tessellation.TexCoordArray.Num())
-				{
-					for (int32 IndexFace = 0; IndexFace < CTFaceIndex.Num(); IndexFace += 3)
-					{
-						for (int32 Index = 0; Index < TriangleCount; Index++)
-						{
-							FVertexInstanceID VertexInstanceID = MeshVertexInstanceIDs[IndexFace + Orientation[Index]];
-							VertexInstanceUVs.Set(VertexInstanceID, UVChannel, Tessellation.TexCoordArray[CTFaceIndex[IndexFace + Index]]);
-						}
-					}
-				}
-
 				if (!Step)
 				{
 					FDatasmithUtils::ConvertVectorArray((FDatasmithUtils::EModelCoordSystem) ImportParams.ModelCoordSys, Tessellation.NormalArray);
@@ -455,67 +496,66 @@ namespace CADLibrary
 
 				if (Tessellation.NormalArray.Num() == 1)
 				{
-					Temp3D = Tessellation.NormalArray[0];
-					for (int32 Index = 0; Index < CTFaceIndex.Num(); Index++)
+					const FVector& Normal = Tessellation.NormalArray[0];
+					for (int32 Index = 0; Index < NewFaceIndex.Num(); Index++)
 					{
-						FVertexInstanceID VertexInstanceID = MeshVertexInstanceIDs[Index];
-						VertexInstanceNormals[VertexInstanceID] = Temp3D;
+						const FVertexInstanceID VertexInstanceID = MeshVertexInstanceIDs[Index];
+						VertexInstanceNormals[VertexInstanceID] = Normal;
 					}
 				}
 				else
 				{
-					for (FVector& Normal : Tessellation.NormalArray)
-					{
-						Normal = Normal.GetSafeNormal();
-					}
-
-					for (int32 IndexFace = 0; IndexFace < CTFaceIndex.Num(); IndexFace += 3)
+					for (int32 IndexFace = 0; IndexFace < NewFaceIndex.Num(); IndexFace += 3)
 					{
 						for (int32 Index = 0; Index < 3; Index++)
 						{
-							FVertexInstanceID VertexInstanceID = MeshVertexInstanceIDs[IndexFace + Orientation[Index]];
-							VertexInstanceNormals[VertexInstanceID] = Tessellation.NormalArray[CTFaceIndex[IndexFace + Index]];
+							const FVertexInstanceID VertexInstanceID = MeshVertexInstanceIDs[IndexFace + Orientation[Index]];
+							VertexInstanceNormals[VertexInstanceID] = Tessellation.NormalArray[NewFaceIndex[IndexFace + Index]];
 						}
 					}
 				}
 
-				// compute normals
 				if (Step)
 				{
+					// compute normals of Symmetric vertex
 					FMatrix SymmetricMatrix;
 					SymmetricMatrix = FDatasmithUtils::GetSymmetricMatrix(MeshParameters.SymmetricOrigin, MeshParameters.SymmetricNormal);
-					for (int32 Index = 0; Index < CTFaceIndex.Num(); Index++)
+					for (const FVertexInstanceID& VertexInstanceID : MeshVertexInstanceIDs)
 					{
-						VertexInstanceNormals[MeshVertexInstanceIDs[Index]] = SymmetricMatrix.TransformVector(VertexInstanceNormals[MeshVertexInstanceIDs[Index]]);;
+						VertexInstanceNormals[VertexInstanceID] = SymmetricMatrix.TransformVector(VertexInstanceNormals[VertexInstanceID]);;
 					}
 				}
 
 				if (MeshParameters.bNeedSwapOrientation)
 				{
-					for (int32 Index = 0; Index < CTFaceIndex.Num(); Index++)
+					for (int32 Index = 0; Index < MeshVertexInstanceIDs.Num(); Index++)
 					{
 						VertexInstanceNormals[MeshVertexInstanceIDs[Index]] = VertexInstanceNormals[MeshVertexInstanceIDs[Index]] * -1.f;
 					}
 				}
 			}
 		}
+
 		return true;
 	}
 
-	bool ConvertCTBodySetToMeshDescription(const FImportParameters& ImportParams, const FMeshParameters& MeshParameters, FBodyMesh& Body, FMeshDescription& MeshDescription)
+	bool ConvertBodyMeshToMeshDescription(const FImportParameters& ImportParams, const FMeshParameters& MeshParameters, FBodyMesh& Body, FMeshDescription& MeshDescription)
 	{
-		// Ref. CreateMesh(UDatasmithCADImportOptions* CADOptions, CTMesh& Mesh)
-
-		// in a closed big mesh VertexCount ~ TriangleCount / 2, EdgeCount ~ 1.5* TriangleCount
-		MeshDescription.ReserveNewVertexInstances(Body.TriangleCount * 3);
+		MeshDescription.ReserveNewVertexInstances(Body.VertexArray.Num());
 		MeshDescription.ReserveNewPolygons(Body.TriangleCount);
 		MeshDescription.ReserveNewEdges(Body.TriangleCount * 3);
 
 		// CoreTech is generating position duplicates. make sure to remove them before filling the mesh description
-		TArray<FVertexID> RemapVertexPosition;
-		FillVertexPosition(ImportParams, MeshParameters, Body, MeshDescription);
+		if (Body.VertexArray.Num() == 0)
+		{
+			FillKioVertexPosition(ImportParams, MeshParameters, Body, MeshDescription);
+		}
+		else
+		{
+			FillVertexPosition(ImportParams, MeshParameters, Body, MeshDescription);
+		}
 
-		if (!FillMesh(MeshParameters, ImportParams, Body.Faces, MeshDescription))
+		if (!FillMesh(MeshParameters, ImportParams, Body, MeshDescription))
 		{
 			return false;
 		}
@@ -641,10 +681,10 @@ namespace CADLibrary
 
 		while (CT_OBJECT_ID BodyId = Objects.IteratorIter())
 		{
-			Parser.GetBodyTessellation(BodyId, MainObjectId, BodyMesh, 0, /*bNeedRepair*/ true);
+			Parser.GetKioBodyTessellation(BodyId, MainObjectId, BodyMesh, 0, /*bNeedRepair*/ true);
 		}
 
-		bool bTessellated = ConvertCTBodySetToMeshDescription(ImportParams, MeshParameters, BodyMesh, MeshDesc);
+		bool bTessellated = ConvertBodyMeshToMeshDescription(ImportParams, MeshParameters, BodyMesh, MeshDesc);
 
 		CheckedCTError ConversionResult;
 		if (!bTessellated)
diff --git a/Engine/Plugins/Enterprise/DatasmithCADImporter/Source/CADLibrary/Public/CoreTechHelper.h b/Engine/Plugins/Enterprise/DatasmithCADImporter/Source/CADLibrary/Public/CoreTechHelper.h
--- a/Engine/Plugins/Enterprise/DatasmithCADImporter/Source/CADLibrary/Public/CoreTechHelper.h
+++ b/Engine/Plugins/Enterprise/DatasmithCADImporter/Source/CADLibrary/Public/CoreTechHelper.h
@@ -48,7 +48,7 @@ namespace CADLibrary
 		}
 	};
 
-	CADLIBRARY_API bool ConvertCTBodySetToMeshDescription(const FImportParameters& ImportParams, const FMeshParameters& MeshParameters, FBodyMesh& Body, FMeshDescription& MeshDescription);
+	CADLIBRARY_API bool ConvertBodyMeshToMeshDescription(const FImportParameters& ImportParams, const FMeshParameters& MeshParameters, FBodyMesh& Body, FMeshDescription& MeshDescription);
 
 	CADLIBRARY_API CT_IO_ERROR Tessellate(CT_OBJECT_ID MainObjectId, const FImportParameters& ImportParams, FMeshDescription& Mesh, FMeshParameters& MeshParameters);
 
diff --git a/Engine/Plugins/Enterprise/DatasmithCADImporter/Source/CADTools/Private/CADData.cpp b/Engine/Plugins/Enterprise/DatasmithCADImporter/Source/CADTools/Private/CADData.cpp
--- a/Engine/Plugins/Enterprise/DatasmithCADImporter/Source/CADTools/Private/CADData.cpp
+++ b/Engine/Plugins/Enterprise/DatasmithCADImporter/Source/CADTools/Private/CADData.cpp
@@ -6,6 +6,8 @@
 #include "Misc/FileHelper.h"
 #include "Misc/Paths.h"
 
+uint32 MeshArchiveMagic = 345612;
+
 namespace CADLibrary
 {
 uint32 BuildColorId(uint32 ColorId, uint8 Alpha)
@@ -72,29 +74,31 @@ FArchive& operator<<(FArchive& Ar, FFileDescription& File)
 
 FArchive& operator<<(FArchive& Ar, FTessellationData& TessellationData)
 {
-	Ar << TessellationData.PatchId;
+	Ar << TessellationData.PositionArray;
+
+	Ar << TessellationData.PositionIndices;
+	Ar << TessellationData.VertexIndices;
 
-	Ar << TessellationData.VertexArray;
 	Ar << TessellationData.NormalArray;
-	Ar << TessellationData.IndexArray;
 	Ar << TessellationData.TexCoordArray;
 
-	Ar << TessellationData.StartVertexIndex;
-
 	Ar << TessellationData.ColorName;
 	Ar << TessellationData.MaterialName;
 
+	Ar << TessellationData.PatchId;
+
 	return Ar;
 }
 
 FArchive& operator<<(FArchive& Ar, FBodyMesh& BodyMesh)
 {
+	Ar << BodyMesh.VertexArray;
+	Ar << BodyMesh.Faces;
+	Ar << BodyMesh.BBox;
+
 	Ar << BodyMesh.TriangleCount;
 	Ar << BodyMesh.BodyID;
 	Ar << BodyMesh.MeshActorName;
-	Ar << BodyMesh.Faces;
-
-	Ar << BodyMesh.BBox;
 
 	Ar << BodyMesh.MaterialSet;
 	Ar << BodyMesh.ColorSet;
@@ -106,8 +110,7 @@ void SerializeBodyMeshSet(const TCHAR* Filename, TArray<FBodyMesh>& InBodySet)
 {
 	TUniquePtr<FArchive> Archive(IFileManager::Get().CreateFileWriter(Filename));
 
-	uint32 type = 234561;
-	*Archive << type;
+	*Archive << MeshArchiveMagic;
 
 	*Archive << InBodySet;
 
@@ -118,15 +121,12 @@ void DeserializeBodyMeshFile(const TCHAR* Filename, TArray<FBodyMesh>& OutBodySe
 {
 	TUniquePtr<FArchive> Archive(IFileManager::Get().CreateFileReader(Filename));
 
-	uint32 type = 0;
-	*Archive << type;
-	if (type != 234561)
+	uint32 MagicNumber = 0;
+	*Archive << MagicNumber;
+	if (MagicNumber == MeshArchiveMagic)
 	{
-		Archive->Close();
-		return;
+		*Archive << OutBodySet;
 	}
-
-	*Archive << OutBodySet;
 	Archive->Close();
 }
 
diff --git a/Engine/Plugins/Enterprise/DatasmithCADImporter/Source/CADTools/Public/CADData.h b/Engine/Plugins/Enterprise/DatasmithCADImporter/Source/CADTools/Public/CADData.h
--- a/Engine/Plugins/Enterprise/DatasmithCADImporter/Source/CADTools/Public/CADData.h
+++ b/Engine/Plugins/Enterprise/DatasmithCADImporter/Source/CADTools/Public/CADData.h
@@ -83,26 +83,38 @@ struct CADTOOLS_API FFileDescription
 };
 
 /**
- * Helper struct to store tessellation data from CoreTech
+ * Helper struct to store tessellation data from CoreTech or CADKernel
+ *
+ * FBodyMesh and FTessellationData are design to manage mesh from CoreTech and CADKernel.
+ * FTessellationData is the mesh of a face
+ * FBodyMesh is the mesh of a body composed by an array of FTessellationData (one FTessellationData by body face)
+ *
+ * CoreTech mesh are defined surface by surface. The mesh is not connected
+ * CADKernel mesh is connected.
  */
 struct CADTOOLS_API FTessellationData
 {
 	friend CADTOOLS_API FArchive& operator<<(FArchive& Ar, FTessellationData& Tessellation);
 
-	int32    PatchId = 0;
+	/** Empty with CADKernel as set in FBodyMesh, Set by CoreTech (this is only the vertices of the face) */
+	TArray<FVector> PositionArray; 
 
-	TArray<FVector> VertexArray;
-	TArray<FVector> NormalArray;
-	TArray<int32> IndexArray;
-	TArray<FVector2D> TexCoordArray;
+	/** Index of each vertex in FBody::VertexArray. Empty with CoreTech and filled by FillKioVertexPosition */
+	TArray<int32> PositionIndices; 
+	
+	/** Index of Vertices of each face in the local Vertices set (i.e. VerticesBodyIndex for CADKernel, VertexArray for Coretech) */
+	TArray<int32> VertexIndices; 
 
-	uint32    StartVertexIndex = 0;
+	/** Normal of each vertex */
+	TArray<FVector> NormalArray; 
+
+	/** UV coordinates of each vertex */
+	TArray<FVector2D> TexCoordArray; 
 
 	CADUUID ColorName = 0;
 	CADUUID MaterialName = 0;
 
-	TArray<int32> VertexIdSet;  // StaticMesh FVertexID NO Serialize
-	TArray<int32> SymVertexIdSet; // StaticMesh FVertexID for sym part NO Serialize
+	int32 PatchId;
 };
 
 class CADTOOLS_API FBodyMesh
@@ -116,6 +128,7 @@ public:
 	friend FArchive& operator<<(FArchive& Ar, FBodyMesh& BodyMesh);
 
 public:
+	TArray<FVector> VertexArray; // set by CADKernel, filled by FillKioVertexPosition that merges coincident vertices (CoreTechHelper)
 	TArray<FTessellationData> Faces;
 	FBox BBox;
 
@@ -123,6 +136,9 @@ public:
 	CadId BodyID = 0;
 	CADUUID MeshActorName = 0;
 
+	TArray<int32> VertexIds;  // StaticMesh FVertexID NO Serialize, filled by FillKioVertexPosition or FillVertexPosition
+	TArray<int32> SymmetricVertexIds; // StaticMesh FVertexID for sym part NO Serialize, filled by FillKioVertexPosition or FillVertexPosition
+
 	TSet<uint32> MaterialSet;
 	TSet<uint32> ColorSet;
 };
diff --git a/Engine/Plugins/Enterprise/DatasmithCADImporter/Source/CADTools/Public/CADOptions.h b/Engine/Plugins/Enterprise/DatasmithCADImporter/Source/CADTools/Public/CADOptions.h
--- a/Engine/Plugins/Enterprise/DatasmithCADImporter/Source/CADTools/Public/CADOptions.h
+++ b/Engine/Plugins/Enterprise/DatasmithCADImporter/Source/CADTools/Public/CADOptions.h
@@ -37,11 +37,11 @@ namespace CADLibrary
 		BodyOnly,
 	};
 
-	//enum class ERetessellationRule : uint8
-	//{
-	//	All = 0,
-	//	SkipDeletedSurfaces,
-	//};
+	enum class ERetessellationRule : uint8
+	{
+		All = 0,
+		SkipDeletedSurfaces,
+	};
 
 	struct FImportParameters
 	{
@@ -56,7 +56,8 @@ namespace CADLibrary
 		EDisplayPreference DisplayPreference = EDisplayPreference::MaterialPrefered;
 		bool bScaleUVMap = true;
 		bool bEnableCacheUsage = true;
-		//ERetessellationRule RetessellationRule = ERetessellationRule::All;
+		bool bEnableKernelIOTessellation = true;
+		ERetessellationRule RetessellationRule = ERetessellationRule::All;
 
 		uint32 GetHash() const
 		{
diff --git a/Engine/Plugins/Enterprise/DatasmithCADImporter/Source/DatasmithCADTranslator/Private/DatasmithCADTranslator.cpp b/Engine/Plugins/Enterprise/DatasmithCADImporter/Source/DatasmithCADTranslator/Private/DatasmithCADTranslator.cpp
--- a/Engine/Plugins/Enterprise/DatasmithCADImporter/Source/DatasmithCADTranslator/Private/DatasmithCADTranslator.cpp
+++ b/Engine/Plugins/Enterprise/DatasmithCADImporter/Source/DatasmithCADTranslator/Private/DatasmithCADTranslator.cpp
@@ -19,6 +19,12 @@ static TAutoConsoleVariable<int32> CVarStaticCADTranslatorEnableThreadedImport(
 	TEXT("Activate to parallelise CAD file processing.\n"),
 	ECVF_Default);
 
+static TAutoConsoleVariable<int32> CVarStaticCADTranslatorEnableKernelIOTessellation(
+	TEXT("r.CADTranslator.EnableKernelIOTessellation"),
+	1,
+	TEXT("Activate to use KernelIO Tessellation.\n"),
+	ECVF_Default);
+
 void FDatasmithCADTranslator::Initialize(FDatasmithTranslatorCapabilities& OutCapabilities)
 {
 	if (ICADInterfacesModule::IsAvailable() == ECADInterfaceAvailability::Unavailable)
@@ -75,6 +81,11 @@ void FDatasmithCADTranslator::Initialize(FDatasmithTranslatorCapabilities& OutCa
 
 bool FDatasmithCADTranslator::LoadScene(TSharedRef<IDatasmithScene> DatasmithScene)
 {
+	ImportParameters.bEnableKernelIOTessellation = (CVarStaticCADTranslatorEnableKernelIOTessellation.GetValueOnAnyThread() != 0);;
+#ifdef CADKERNEL_DEBUG
+	ImportParameters.bEnableKernelIOTessellation = false;
+#endif
+
 	ImportParameters.MetricUnit = 0.001;
 	ImportParameters.ScaleFactor = 0.1;
 	const FDatasmithTessellationOptions& TesselationOptions = GetCommonTessellationOptions();
diff --git a/Engine/Plugins/Enterprise/DatasmithCADImporter/Source/DatasmithCADTranslator/Private/DatasmithMeshBuilder.cpp b/Engine/Plugins/Enterprise/DatasmithCADImporter/Source/DatasmithCADTranslator/Private/DatasmithMeshBuilder.cpp
--- a/Engine/Plugins/Enterprise/DatasmithCADImporter/Source/DatasmithCADTranslator/Private/DatasmithMeshBuilder.cpp
+++ b/Engine/Plugins/Enterprise/DatasmithCADImporter/Source/DatasmithCADTranslator/Private/DatasmithMeshBuilder.cpp
@@ -20,7 +20,7 @@ FDatasmithMeshBuilder::FDatasmithMeshBuilder(TMap<uint32, FString>& InCADFileToU
 	, CADFileToMeshFile(InCADFileToUE4GeomMap)
 	, ImportParameters(InImportParameters)
 {
-	LoadMeshFiles();
+		LoadMeshFiles();
 }
 
 void FDatasmithMeshBuilder::LoadMeshFiles()
@@ -34,9 +34,9 @@ void FDatasmithMeshBuilder::LoadMeshFiles()
 		{
 			continue;
 		}
-		TArray<CADLibrary::FBodyMesh>& BodyMeshSet = BodyMeshes.Emplace_GetRef();
-		DeserializeBodyMeshFile(*MeshFile, BodyMeshSet);
-		for (FBodyMesh& Body : BodyMeshSet)
+		TArray<CADLibrary::FBodyMesh>& BodyMeshArray = BodyMeshes.Emplace_GetRef();
+		DeserializeBodyMeshFile(*MeshFile, BodyMeshArray);
+		for (FBodyMesh& Body : BodyMeshArray)
 		{
 			MeshActorNameToBodyMesh.Emplace(Body.MeshActorName, &Body);
 		}
@@ -47,14 +47,14 @@ TOptional<FMeshDescription> FDatasmithMeshBuilder::GetMeshDescription(TSharedRef
 {
 #ifdef CAD_INTERFACE
 	const TCHAR* NameLabel = OutMeshElement->GetName();
-	CADUUID BodyUuid = (CADUUID) FCString::Atoi64(OutMeshElement->GetName()+2);  // +2 to remove 2 first char (Ox)
+	CADUUID BodyUuid = (CADUUID)FCString::Atoi64(OutMeshElement->GetName() + 2);  // +2 to remove 2 first char (Ox)
 	if (BodyUuid == 0)
 	{
 		return TOptional<FMeshDescription>();
 	}
 
 	FBodyMesh** PPBody = MeshActorNameToBodyMesh.Find(BodyUuid);
-	if(PPBody == nullptr || *PPBody == nullptr)
+	if (PPBody == nullptr || *PPBody == nullptr)
 	{
 		return TOptional<FMeshDescription>();
 	}
@@ -64,7 +64,7 @@ TOptional<FMeshDescription> FDatasmithMeshBuilder::GetMeshDescription(TSharedRef
 	FMeshDescription MeshDescription;
 	DatasmithMeshHelper::PrepareAttributeForStaticMesh(MeshDescription);
 
-	if (ConvertCTBodySetToMeshDescription(ImportParameters, OutMeshParameters, Body, MeshDescription))
+	if (ConvertBodyMeshToMeshDescription(ImportParameters, OutMeshParameters, Body, MeshDescription))
 	{
 		return MoveTemp(MeshDescription);
 	}
