Rename DatasmithCoreTechParametricSurfaceData into CoreTechSurface and DatasmithCoreTechExtension into CoreTechExtension

diff --git a/Engine/Plugins/Enterprise/DatasmithCADImporter/Source/CoreTechExtension/CoreTechExtension.Build.cs b/Engine/Plugins/Enterprise/DatasmithCADImporter/Source/CoreTechExtension/CoreTechExtension.Build.cs
new file mode 100644
--- /dev/null
+++ b/Engine/Plugins/Enterprise/DatasmithCADImporter/Source/CoreTechExtension/CoreTechExtension.Build.cs
@@ -0,0 +1,40 @@
+// Copyright Epic Games, Inc. All Rights Reserved.
+
+using System.IO;
+
+namespace UnrealBuildTool.Rules
+{
+	public class CoreTechExtension : ModuleRules
+	{
+		public CoreTechExtension(ReadOnlyTargetRules Target) : base(Target)
+		{
+			PublicDependencyModuleNames.AddRange(
+				new string[]
+				{
+					"Core",
+				}
+			);
+
+			PrivateDependencyModuleNames.AddRange(
+				new string[]
+				{
+					"Core",
+					"CoreUObject",
+					"MeshDescription",
+					"DataprepCore",
+					"DatasmithContent",
+					"CoreTechSurface",
+					"DatasmithImporter",
+					"Engine",
+					"StaticMeshEditor",
+					"Slate",
+					"SlateCore",
+					"StaticMeshDescription",
+					"EditorFramework",
+					"UnrealEd",
+					"CADLibrary",
+				}
+			);
+		}
+	}
+}
diff --git a/Engine/Plugins/Enterprise/DatasmithCADImporter/Source/CoreTechExtension/Private/CoreTechBlueprintLibrary.cpp b/Engine/Plugins/Enterprise/DatasmithCADImporter/Source/CoreTechExtension/Private/CoreTechBlueprintLibrary.cpp
new file mode 100644
--- /dev/null
+++ b/Engine/Plugins/Enterprise/DatasmithCADImporter/Source/CoreTechExtension/Private/CoreTechBlueprintLibrary.cpp
@@ -0,0 +1,108 @@
+// Copyright Epic Games, Inc. All Rights Reserved.
+
+#include "CoreTechBlueprintLibrary.h"
+
+#include "CoreTechRetessellateAction.h"
+#include "DatasmithAdditionalData.h"
+#include "DatasmithStaticMeshImporter.h" // Call to BuildStaticMesh
+#include "DatasmithUtils.h"
+
+#include "AssetData.h"
+#include "Engine/StaticMesh.h"
+#include "IStaticMeshEditor.h"
+#include "StaticMeshAttributes.h"
+#include "Toolkits/ToolkitManager.h"
+
+
+#define LOCTEXT_NAMESPACE "CoreTechRetessellateAction"
+
+
+bool UCoreTechBlueprintLibrary::RetessellateStaticMesh(UStaticMesh* StaticMesh, const FDatasmithRetessellationOptions& TessellationSettings, FText& FailureReason)
+{
+	return RetessellateStaticMeshWithNotification(StaticMesh, TessellationSettings, true, FailureReason);
+}
+
+bool UCoreTechBlueprintLibrary::RetessellateStaticMeshWithNotification(UStaticMesh* StaticMesh, const FDatasmithRetessellationOptions& TessellationSettings, bool bApplyChanges, FText& FailureReason)
+{
+	bool bTessellationOutcome = false;
+
+#ifdef CAD_LIBRARY
+	int32 LODIndex = 0;
+
+	FAssetData AssetData( StaticMesh );
+	if (UCoreTechParametricSurfaceData* CoreTechData = Datasmith::GetAdditionalData<UCoreTechParametricSurfaceData>(AssetData))
+	{
+		// Make sure MeshDescription exists
+		FMeshDescription* DestinationMeshDescription = StaticMesh->GetMeshDescription( LODIndex );
+		if (DestinationMeshDescription == nullptr)
+		{
+			DestinationMeshDescription = StaticMesh->CreateMeshDescription(0);
+		}
+
+		if (DestinationMeshDescription)
+		{
+			if (bApplyChanges)
+			{
+				StaticMesh->Modify();
+				StaticMesh->PreEditChange( nullptr );
+			}
+
+			const int32 OldNumberOfUVChannels = FStaticMeshAttributes(*DestinationMeshDescription).GetVertexInstanceUVs().GetNumChannels();
+			if (FCoreTechRetessellate_Impl::ApplyOnOneAsset( *StaticMesh, *CoreTechData, TessellationSettings ))
+			{
+				const int32 NumberOfUVChannels = FStaticMeshAttributes(*DestinationMeshDescription).GetVertexInstanceUVs().GetNumChannels();
+				if (NumberOfUVChannels < OldNumberOfUVChannels)
+				{
+					FailureReason = FText::Format(NSLOCTEXT("BlueprintRetessellation", "UVChannelsDestroyed", "Tessellation operation on Static Mesh {0} is destroying all UV channels above channel #{1}"), FText::FromString(StaticMesh->GetName()), NumberOfUVChannels - 1);
+				}
+
+				// Post static mesh has changed
+				if (bApplyChanges)
+				{
+					FDatasmithStaticMeshImporter::PreBuildStaticMesh(StaticMesh); // handle uvs stuff
+					FDatasmithStaticMeshImporter::BuildStaticMesh(StaticMesh);
+
+					StaticMesh->PostEditChange();
+					StaticMesh->MarkPackageDirty();
+
+					// Refresh associated editor
+					TSharedPtr<IToolkit> EditingToolkit = FToolkitManager::Get().FindEditorForAsset(StaticMesh);
+					if (IStaticMeshEditor* StaticMeshEditorInUse = StaticCastSharedPtr<IStaticMeshEditor>(EditingToolkit).Get())
+					{
+						StaticMeshEditorInUse->RefreshTool();
+					}
+				}
+				// No posting required, just make sure the new tessellation is committed
+				else
+				{
+					UStaticMesh::FCommitMeshDescriptionParams Params;
+					Params.bMarkPackageDirty = false;
+					Params.bUseHashAsGuid = true;
+					StaticMesh->CommitMeshDescription( LODIndex, Params );
+				}
+
+				// Save last tessellation settings
+				CoreTechData->LastTessellationOptions = TessellationSettings;
+
+				bTessellationOutcome = true;
+			}
+			else
+			{
+				FailureReason = NSLOCTEXT("BlueprintRetessellation", "LoadFailed", "Cannot generate mesh from parametric surface data");
+			}
+		}
+		else
+		{
+			FailureReason = NSLOCTEXT("BlueprintRetessellation", "MeshDescriptionMissing", "Cannot create mesh description");
+		}
+	}
+	else
+	{
+		FailureReason = NSLOCTEXT("BlueprintRetessellation", "MissingData", "No tessellation data attached to the static mesh");
+	}
+#endif
+
+	return bTessellationOutcome;
+}
+
+#undef LOCTEXT_NAMESPACE
\ No newline at end of file
diff --git a/Engine/Plugins/Enterprise/DatasmithCADImporter/Source/CoreTechExtension/Private/CoreTechExtensionModule.cpp b/Engine/Plugins/Enterprise/DatasmithCADImporter/Source/CoreTechExtension/Private/CoreTechExtensionModule.cpp
new file mode 100644
--- /dev/null
+++ b/Engine/Plugins/Enterprise/DatasmithCADImporter/Source/CoreTechExtension/Private/CoreTechExtensionModule.cpp
@@ -0,0 +1,101 @@
+// Copyright Epic Games, Inc. All Rights Reserved.
+
+#include "CoreTechExtensionModule.h"
+
+#include "CoreTechRetessellateAction.h"
+
+#include "Framework/MultiBox/MultiBoxBuilder.h"
+#include "IStaticMeshEditor.h"
+#include "Modules/ModuleInterface.h"
+#include "Modules/ModuleManager.h"
+#include "StaticMeshEditorModule.h"
+#include "UObject/StrongObjectPtr.h"
+
+
+#define LOCTEXT_NAMESPACE "CoreTechExtensionModule"
+
+
+/** UI extension that displays a Retessellate action in the StaticMeshEditor */
+namespace StaticMeshEditorExtenser
+{
+	bool CanExecute(UStaticMesh* Target)
+	{
+		TArray<FAssetData> AssetData;
+		AssetData.Emplace(Target);
+
+		FCoreTechRetessellate_Impl RetessellateAction;
+		return RetessellateAction.CanApplyOnAssets(AssetData);
+	}
+
+	void Execute(UStaticMesh* Target)
+	{
+		TArray<FAssetData> AssetData;
+		AssetData.Emplace(Target);
+
+		FCoreTechRetessellate_Impl RetessellateAction;
+		RetessellateAction.ApplyOnAssets(AssetData);
+	}
+
+	void ExtendAssetMenu(FMenuBuilder& MenuBuilder, UStaticMesh* Target)
+	{
+		MenuBuilder.AddMenuEntry(
+			FCoreTechRetessellate_Impl::Label,
+			FCoreTechRetessellate_Impl::Tooltip,
+			FSlateIcon(),
+			FUIAction(
+				FExecuteAction::CreateStatic(&StaticMeshEditorExtenser::Execute, Target),
+				FCanExecuteAction::CreateStatic(&StaticMeshEditorExtenser::CanExecute, Target)
+			)
+		);
+	}
+
+	TSharedRef<FExtender> CreateExtenderForObjects(const TSharedRef<FUICommandList> CommandList, const TArray<UObject*> Objects)
+	{
+		TSharedRef<FExtender> Extender = MakeShareable(new FExtender);
+
+		if (UStaticMesh* Target = Objects.Num() ? Cast<UStaticMesh>(Objects[0]) : nullptr)
+		{
+			Extender->AddMenuExtension(
+				"AssetEditorActions",
+				EExtensionHook::Position::After,
+				CommandList,
+				FMenuExtensionDelegate::CreateStatic(&StaticMeshEditorExtenser::ExtendAssetMenu, Target)
+			);
+		}
+
+		return Extender;
+	}
+
+	void Register()
+	{
+		if (!IsRunningCommandlet())
+		{
+			IStaticMeshEditorModule& StaticMeshEditorModule = FModuleManager::Get().LoadModuleChecked<IStaticMeshEditorModule>("StaticMeshEditor");
+			TArray<FAssetEditorExtender>& ExtenderDelegates = StaticMeshEditorModule.GetMenuExtensibilityManager()->GetExtenderDelegates();
+			ExtenderDelegates.Add(FAssetEditorExtender::CreateStatic(&StaticMeshEditorExtenser::CreateExtenderForObjects));
+		}
+	}
+};
+
+FCoreTechExtensionModule& FCoreTechExtensionModule::Get()
+{
+	return FModuleManager::LoadModuleChecked< FCoreTechExtensionModule >(CORETECHEXTENSION_MODULE_NAME);
+}
+
+bool FCoreTechExtensionModule::IsAvailable()
+{
+	return FModuleManager::Get().IsModuleLoaded(CORETECHEXTENSION_MODULE_NAME);
+}
+
+void FCoreTechExtensionModule::StartupModule()
+{
+	if (!IsRunningCommandlet())
+	{
+		StaticMeshEditorExtenser::Register();
+	}
+}
+
+IMPLEMENT_MODULE(FCoreTechExtensionModule, CoreTechExtension);
+
+#undef LOCTEXT_NAMESPACE // "DatasmithCoreTechExtensionModule"
+
diff --git a/Engine/Plugins/Enterprise/DatasmithCADImporter/Source/CoreTechExtension/Private/CoreTechRetessellateAction.cpp b/Engine/Plugins/Enterprise/DatasmithCADImporter/Source/CoreTechExtension/Private/CoreTechRetessellateAction.cpp
new file mode 100644
--- /dev/null
+++ b/Engine/Plugins/Enterprise/DatasmithCADImporter/Source/CoreTechExtension/Private/CoreTechRetessellateAction.cpp
@@ -0,0 +1,269 @@
+// Copyright Epic Games, Inc. All Rights Reserved.
+
+#include "CoreTechRetessellateAction.h"
+
+#ifdef CAD_LIBRARY
+#include "CoreTechHelper.h"
+#include "CoreTechMeshLoader.h"
+#include "CoreTechTypes.h"
+#endif
+
+#include "DatasmithAdditionalData.h"
+#include "DatasmithStaticMeshImporter.h" // Call to BuildStaticMesh
+#include "DatasmithUtils.h"
+#include "DatasmithTranslator.h"
+#include "UI/DatasmithDisplayHelper.h"
+
+#include "Algo/AnyOf.h"
+#include "AssetData.h"
+#include "Async/ParallelFor.h"
+#include "Engine/StaticMesh.h"
+#include "IStaticMeshEditor.h"
+#include "StaticMeshAttributes.h"
+#include "Misc/FileHelper.h"
+#include "Misc/ScopedSlowTask.h"
+#include "Toolkits/IToolkit.h"
+#include "Toolkits/ToolkitManager.h"
+#include "UObject/StrongObjectPtr.h"
+#include "Algo/Transform.h"
+
+
+#define LOCTEXT_NAMESPACE "CoreTechRetessellateAction"
+
+
+const FText FCoreTechRetessellate_Impl::Label = LOCTEXT("RetessellateActionLabel", "Retessellate");
+const FText FCoreTechRetessellate_Impl::Tooltip = LOCTEXT("RetessellateActionTooltip", "Tessellate the original NURBS surfaces to re-generate the mesh geometry");
+
+
+bool FCoreTechRetessellate_Impl::CanApplyOnAssets(const TArray<FAssetData>& SelectedAssets)
+{
+#ifdef CAD_LIBRARY
+	return Algo::AnyOf(SelectedAssets, [](const FAssetData& Asset){ return Datasmith::GetAdditionalData<UCoreTechParametricSurfaceData>(Asset) != nullptr; });
+#else
+	return false;
+#endif
+}
+
+void FCoreTechRetessellate_Impl::ApplyOnAssets(const TArray<FAssetData>& SelectedAssets)
+{
+#ifdef CAD_LIBRARY
+	TStrongObjectPtr<UCoreTechRetessellateActionOptions> RetessellateOptions = Datasmith::MakeOptions<UCoreTechRetessellateActionOptions>();
+
+	bool bSameOptionsForAll = false;
+	int32 NumAssetsToProcess = SelectedAssets.Num();
+	bool bAskForSameOption = NumAssetsToProcess > 1;
+
+	TArray<UStaticMesh*> TessellatedMeshes;
+	TessellatedMeshes.Reserve( SelectedAssets.Num() );
+
+	TUniquePtr<FScopedSlowTask> Progress;
+	int32 AssetIndex = -1;
+	for (const FAssetData& Asset : SelectedAssets)
+	{
+		AssetIndex++;
+		if (UCoreTechParametricSurfaceData* CoreTechData = Datasmith::GetAdditionalData<UCoreTechParametricSurfaceData>(Asset))
+		{
+			if (CoreTechData->RawData.Num())
+			{
+				if (UStaticMesh* StaticMesh = Cast<UStaticMesh>(Asset.GetAsset()))
+				{
+					if (!bSameOptionsForAll)
+					{
+						Datasmith::FDisplayParameters Parameters;
+						Parameters.bAskForSameOption = bAskForSameOption;
+						Parameters.WindowTitle = LOCTEXT("OptionWindow_WindowTitle", "Datasmith Retessellation Options");
+						Parameters.FileLabel = FText::Format(LOCTEXT("OptionWindow_AssetLabel", "Tessellate StaticMesh: {0}"), FText::FromString(StaticMesh->GetName()));
+						Parameters.FileTooltip = FText::FromString(StaticMesh->GetPathName());
+// 						Parameters.PackageLabel = FText::FromName(StaticMesh->GetOutermost()->FileName);
+						Parameters.ProceedButtonLabel = LOCTEXT("OptionWindow_ProceedButtonLabel", "Tessellate");
+						Parameters.ProceedButtonTooltip = LOCTEXT("OptionWindow_ProceedButtonTooltip", "Retessellate this mesh based on included nurbs data");
+						Parameters.CancelButtonLabel = LOCTEXT("OptionWindow_CancelButtonLabel", "Cancel");
+						Parameters.CancelButtonTooltip = LOCTEXT("OptionWindow_CancelButtonTooltip", "Cancel the retessellation operation");
+
+						bAskForSameOption = false; // ask only the fist time
+						RetessellateOptions->Options = CoreTechData->LastTessellationOptions;
+						Datasmith::FDisplayResult Result = Datasmith::DisplayOptions(RetessellateOptions, Parameters);
+						if (!Result.bValidated)
+						{
+							return;
+						}
+						bSameOptionsForAll |= Result.bUseSameOption;
+					}
+					CoreTechData->LastTessellationOptions = RetessellateOptions->Options;
+
+					int32 RemainingAssetsToProcess = NumAssetsToProcess - AssetIndex;
+					if (bSameOptionsForAll && !Progress.IsValid() && RemainingAssetsToProcess > 1)
+					{
+						Progress = MakeUnique<FScopedSlowTask>(RemainingAssetsToProcess);
+						Progress->MakeDialog(true);
+					}
+					if (Progress)
+					{
+						if (Progress->ShouldCancel())
+						{
+							return;
+						}
+
+						FText Text = FText::Format(LOCTEXT("RetessellateAssetMessage", "Tessellate StaticMesh ({0}/{1}): {2}"),
+							AssetIndex,
+							NumAssetsToProcess,
+							FText::FromString(StaticMesh->GetName())
+						);
+						Progress->EnterProgressFrame(1, Text);
+					}
+
+					if( StaticMesh->GetMeshDescription(0) == nullptr)
+					{
+						StaticMesh->CreateMeshDescription( 0 );
+					}
+
+					if(StaticMesh->GetMeshDescription(0) != nullptr)
+					{
+						StaticMesh->Modify();
+						StaticMesh->PreEditChange( nullptr );
+
+						if( ApplyOnOneAsset(*StaticMesh, *CoreTechData, RetessellateOptions->Options) )
+						{
+							TessellatedMeshes.Add(StaticMesh);
+						}
+					}
+				}
+			}
+		}
+	}
+
+	// Make sure lightmap settings are valid
+	if(TessellatedMeshes.Num() > 1)
+	{
+		ParallelFor(TessellatedMeshes.Num(), [&](int32 Index)
+		{
+			FDatasmithStaticMeshImporter::PreBuildStaticMesh(TessellatedMeshes[Index]); 
+		});
+	}
+	else if(TessellatedMeshes.Num() > 0)
+	{
+		FDatasmithStaticMeshImporter::PreBuildStaticMesh( TessellatedMeshes[0] );
+	}
+
+	FDatasmithStaticMeshImporter::BuildStaticMeshes( TessellatedMeshes );
+
+	for(UStaticMesh* StaticMesh : TessellatedMeshes)
+	{
+		StaticMesh->PostEditChange();
+		StaticMesh->MarkPackageDirty();
+
+		// Refresh associated editor
+		TSharedPtr<IToolkit> EditingToolkit = FToolkitManager::Get().FindEditorForAsset(StaticMesh);
+		if (IStaticMeshEditor* StaticMeshEditorInUse = StaticCastSharedPtr<IStaticMeshEditor>(EditingToolkit).Get())
+		{
+			StaticMeshEditorInUse->RefreshTool();
+		}
+	}
+
+#endif // CAD_LIBRARY
+}
+
+bool FCoreTechRetessellate_Impl::ApplyOnOneAsset(UStaticMesh& StaticMesh, UCoreTechParametricSurfaceData& CoreTechData, const FDatasmithRetessellationOptions& RetessellateOptions)
+{
+	bool bSuccessfulTessellation = false;
+
+#ifdef CAD_LIBRARY
+	// make a temporary file as GPure can only deal with files.
+	FString ResourceFile = CoreTechData.SourceFile.IsEmpty() ? FPaths::ProjectIntermediateDir() / "temp.ct" : CoreTechData.SourceFile;
+	FFileHelper::SaveArrayToFile(CoreTechData.RawData, *ResourceFile);
+
+	CADLibrary::CoreTechMeshLoader Loader;
+	CADLibrary::CTMesh Mesh;
+
+	CADLibrary::FImportParameters ImportParameters;
+	ImportParameters.MetricUnit = CoreTechData.SceneParameters.MetricUnit;
+	ImportParameters.ScaleFactor = CoreTechData.SceneParameters.ScaleFactor;
+	ImportParameters.ChordTolerance = RetessellateOptions.ChordTolerance;
+	ImportParameters.MaxEdgeLength = RetessellateOptions.MaxEdgeLength;
+	ImportParameters.MaxNormalAngle = RetessellateOptions.NormalTolerance;
+	ImportParameters.ModelCoordSys = CADLibrary::EModelCoordSystem(CoreTechData.SceneParameters.ModelCoordSys);
+	ImportParameters.StitchingTechnique = CADLibrary::EStitchingTechnique(RetessellateOptions.StitchingTechnique);
+
+	CADLibrary::FMeshParameters MeshParameters;
+	MeshParameters.bNeedSwapOrientation = CoreTechData.MeshParameters.bNeedSwapOrientation;
+	MeshParameters.bIsSymmetric = CoreTechData.MeshParameters.bIsSymmetric;
+	MeshParameters.SymmetricNormal = CoreTechData.MeshParameters.SymmetricNormal;
+	MeshParameters.SymmetricOrigin = CoreTechData.MeshParameters.SymmetricOrigin;
+
+	// Previous MeshDescription is get to be able to create a new one with the same order of PolygonGroup (the matching of color and partition is currently based on their order)
+	if(FMeshDescription* DestinationMeshDescription = StaticMesh.GetMeshDescription(0))
+	{
+		FMeshDescription MeshDescription;
+		FStaticMeshAttributes MeshDescriptionAttributes(MeshDescription);
+		MeshDescriptionAttributes.Register();
+
+		if (RetessellateOptions.RetessellationRule == EDatasmithCADRetessellationRule::SkipDeletedSurfaces)
+		{
+			CADLibrary::CopyPatchGroups(*DestinationMeshDescription, MeshDescription);
+		}
+
+		if ( Loader.LoadFile(ResourceFile, MeshDescription, ImportParameters, MeshParameters))
+		{
+			// To update the SectionInfoMap 
+			{
+				TPolygonGroupAttributesConstRef<FName> MaterialSlotNames = MeshDescriptionAttributes.GetPolygonGroupMaterialSlotNames();
+				FMeshSectionInfoMap& SectionInfoMap = StaticMesh.GetSectionInfoMap();
+
+				for (FPolygonGroupID PolygonGroupID : MeshDescription.PolygonGroups().GetElementIDs())
+				{
+					FMeshSectionInfo Section = SectionInfoMap.Get(0, PolygonGroupID.GetValue());
+					int32 MaterialIndex = StaticMesh.GetMaterialIndex(MaterialSlotNames[PolygonGroupID]);
+					if (MaterialIndex < 0)
+					{
+						MaterialIndex = 0;
+					}
+					Section.MaterialIndex = MaterialIndex;
+					SectionInfoMap.Set(0, PolygonGroupID.GetValue(), Section);
+				}
+			}
+			*DestinationMeshDescription = MoveTemp(MeshDescription);
+			bSuccessfulTessellation = true;
+		}
+	}
+#endif // CAD_LIBRARY
+
+	return bSuccessfulTessellation;
+}
+
+#undef LOCTEXT_NAMESPACE
+
+TSet<UStaticMesh*> GetReferencedStaticMeshes(const TArray<AActor*>& SelectedActors)
+{
+	TSet<UStaticMesh*> ReferencedStaticMeshes;
+
+#ifdef CAD_LIBRARY
+	for (const AActor* Actor : SelectedActors)
+	{
+		if (Actor)
+		{
+			for (const auto& Component : Actor->GetComponents())
+			{
+				if (UStaticMeshComponent* SMC = Cast<UStaticMeshComponent>(Component))
+				{
+					ReferencedStaticMeshes.Add(SMC->GetStaticMesh());
+				}
+			}
+		}
+	}
+#endif // CAD_LIBRARY
+
+	return ReferencedStaticMeshes;
+}
+
+bool UCoreTechRetessellateAction::CanApplyOnActors(const TArray<AActor*>& SelectedActors)
+{
+	const TSet<UStaticMesh*> ReferencedStaticMeshes = GetReferencedStaticMeshes(SelectedActors);
+	return Algo::AnyOf(ReferencedStaticMeshes, [](const UStaticMesh* Mesh){ return Datasmith::GetAdditionalData<UCoreTechParametricSurfaceData>(FAssetData(Mesh)); });
+}
+
+void UCoreTechRetessellateAction::ApplyOnActors(const TArray<AActor*>& SelectedActors)
+{
+	TArray<FAssetData> AssetData;
+	Algo::Transform(GetReferencedStaticMeshes(SelectedActors), AssetData, [](UStaticMesh* Mesh){ return FAssetData(Mesh);});
+	return ApplyOnAssets(AssetData);
+}
diff --git a/Engine/Plugins/Enterprise/DatasmithCADImporter/Source/CoreTechExtension/Private/DataprepTessellationOperation.cpp b/Engine/Plugins/Enterprise/DatasmithCADImporter/Source/CoreTechExtension/Private/DataprepTessellationOperation.cpp
new file mode 100644
--- /dev/null
+++ b/Engine/Plugins/Enterprise/DatasmithCADImporter/Source/CoreTechExtension/Private/DataprepTessellationOperation.cpp
@@ -0,0 +1,131 @@
+// Copyright Epic Games, Inc. All Rights Reserved.
+
+#include "DataprepTessellationOperation.h"
+
+#ifdef CAD_LIBRARY
+#include "CoreTechHelper.h"
+#include "CoreTechMeshLoader.h"
+#include "CoreTechBlueprintLibrary.h"
+#endif
+#include "DatasmithAdditionalData.h"
+#include "DatasmithAssetImportData.h"
+#include "DatasmithUtils.h"
+
+#include "IDataprepProgressReporter.h"
+
+#include "ActorEditorUtils.h"
+#include "Engine/StaticMesh.h"
+#include "StaticMeshAttributes.h"
+#include "GenericPlatform/GenericPlatformTime.h"
+#include "Logging/LogMacros.h"
+#include "Misc/FileHelper.h"
+
+#define LOCTEXT_NAMESPACE "DataprepTessellationOperation"
+
+DEFINE_LOG_CATEGORY_STATIC(LogCADLibrary, Log, All)
+
+void UDataprepTessellationOperation::PostLoad()
+{
+	if( HasAnyFlags(RF_WasLoaded) && TessellationSettings_DEPRECATED.ChordTolerance != -MAX_FLT)
+	{
+		ChordTolerance = TessellationSettings_DEPRECATED.ChordTolerance;
+		MaxEdgeLength = TessellationSettings_DEPRECATED.MaxEdgeLength;
+		NormalTolerance = TessellationSettings_DEPRECATED.NormalTolerance;
+		// Mark TessellationSettings_DEPRECATED as non usable
+		TessellationSettings_DEPRECATED.ChordTolerance = -MAX_FLT;
+
+		MarkPackageDirty();
+	}
+
+	Super::PostLoad();
+}
+
+void UDataprepTessellationOperation::OnExecution_Implementation(const FDataprepContext& InContext)
+{
+#ifdef CAD_LIBRARY
+	// Collect start time to log amount of time spent to import incoming file
+	uint64 StartTime = FPlatformTime::Cycles64();
+
+	TSet<UStaticMesh*> SelectedMeshes;
+
+	for (UObject* Object : InContext.Objects)
+	{
+		if ( UStaticMesh* StaticMesh = Cast<UStaticMesh>(Object) )
+		{
+			SelectedMeshes.Add( StaticMesh );
+		}
+		else if (AActor* Actor = Cast<AActor>(Object) )
+		{
+			TInlineComponentArray<UStaticMeshComponent*> StaticMeshComponents( Actor );
+			for (UStaticMeshComponent* StaticMeshComponent : StaticMeshComponents)
+			{
+				if((StaticMesh = StaticMeshComponent->GetStaticMesh()) != nullptr)
+				{
+					SelectedMeshes.Add( StaticMesh );
+				}
+			}
+		}
+	}
+
+	if(	!IsCancelled() && SelectedMeshes.Num() > 0)
+	{
+		FDatasmithRetessellationOptions TessellationSettings( ChordTolerance, MaxEdgeLength, NormalTolerance);
+
+		TSharedPtr<FDataprepWorkReporter> Task = CreateTask( LOCTEXT( "LogCADLibrary_Tessellating", "Tessellating meshes ..." ), (float)SelectedMeshes.Num() );
+
+		TArray<UObject*> ModifiedStaticMeshes;
+		ModifiedStaticMeshes.Reserve( SelectedMeshes.Num() );
+
+		for( UStaticMesh* StaticMesh : SelectedMeshes )
+		{
+			if( IsCancelled() )
+			{
+				break;
+			}
+
+			Task->ReportNextStep( FText::Format( LOCTEXT( "LogCADLibrary_Tessellating_One_Mesh", "Tessellating {0} ..."), FText::FromString( StaticMesh->GetName() ) ) );
+
+			if( StaticMesh->IsMeshDescriptionValid( 0 ) )
+			{
+				FText OutReason;
+				if( UCoreTechBlueprintLibrary::RetessellateStaticMeshWithNotification( StaticMesh, TessellationSettings, false, OutReason ) )
+				{
+					ModifiedStaticMeshes.Add( StaticMesh );
+					if (!OutReason.IsEmpty())
+					{
+						FText WarningMsg = FText::Format(LOCTEXT("DataprepTessellationOperation_TessellationCompletedWithWarning", "{0}"), OutReason);
+						LogInfo(WarningMsg);
+					}
+				}
+				else
+				{
+					FText WarningMsg = FText::Format( LOCTEXT( "DataprepTessellationOperation_TessellationFailed", "{0}" ), OutReason );
+					LogWarning( WarningMsg );
+				}
+			}
+			else
+			{
+				FText WarningMsg = FText::Format( LOCTEXT( "DataprepTessellationOperation_EmptyMesh", "No triangles in static mesh {0}" ), FText::FromString( StaticMesh->GetName() ) );
+				LogWarning( WarningMsg );
+			}
+		}
+
+		// Log time spent to import incoming file in minutes and seconds
+		double ElapsedSeconds = FPlatformTime::ToSeconds64(FPlatformTime::Cycles64() - StartTime);
+
+		int ElapsedMin = int(ElapsedSeconds / 60.0);
+		ElapsedSeconds -= 60.0 * (double)ElapsedMin;
+		UE_LOG( LogCADLibrary, Log, TEXT("Tessellation of %d out of %d static mesh(es) in [%d min %.3f s]"), ModifiedStaticMeshes.Num(), SelectedMeshes.Num(), ElapsedMin, ElapsedSeconds );
+
+		if(ModifiedStaticMeshes.Num() > 0)
+		{
+			AssetsModified( MoveTemp( ModifiedStaticMeshes ) );
+		}
+	}
+
+#else
+	UE_LOG( LogCADLibrary, Warning, TEXT("Tessellation not performed") );
+#endif
+}
+
+#undef LOCTEXT_NAMESPACE
diff --git a/Engine/Plugins/Enterprise/DatasmithCADImporter/Source/CoreTechExtension/Private/DataprepTessellationOperation.h b/Engine/Plugins/Enterprise/DatasmithCADImporter/Source/CoreTechExtension/Private/DataprepTessellationOperation.h
new file mode 100644
--- /dev/null
+++ b/Engine/Plugins/Enterprise/DatasmithCADImporter/Source/CoreTechExtension/Private/DataprepTessellationOperation.h
@@ -0,0 +1,80 @@
+// Copyright Epic Games, Inc. All Rights Reserved.
+
+#pragma once
+
+#include "CoreMinimal.h"
+
+#include "DataprepOperation.h"
+#include "DatasmithImportOptions.h"
+#include "DatasmithUtils.h"
+
+#include "DataprepTessellationOperation.generated.h"
+
+namespace CADLibrary
+{
+	struct FImportParameters;
+}
+struct FMeshDescription;
+class UCoreTechParametricSurfaceData;
+class UDatasmithStaticMeshCADImportData;
+
+UCLASS(Experimental, Category = MeshOperation, Meta = (DisplayName="Datasmith Tessellation", ToolTip = "For each static mesh to process, retessellate the mesh if the object contains the required data") )
+class UDataprepTessellationOperation : public UDataprepOperation
+{
+	GENERATED_BODY()
+
+	UDataprepTessellationOperation()
+		: ChordTolerance(0.2f)
+		, MaxEdgeLength(0.0f)
+		, NormalTolerance(20.0f)
+		// Mark TessellationSettings_DEPRECATED as non usable
+		, TessellationSettings_DEPRECATED(-MAX_FLT)
+	{
+	}
+
+	// UObject interface
+	/** Overridden for backward compatibility with 4.23. To be removed in 4.25 */
+	virtual void PostLoad() override;
+	// End of UObject interface
+
+public:
+	/**
+	 * Maximum distance between any point on a triangle generated by the tessellation process and the actual surface.
+	 * The lower the value the more triangles.
+	 * Default value is 0.2.
+	 */
+	UPROPERTY(EditAnywhere, BlueprintReadWrite, Category = "Tessellation Options", meta = (Units = cm, ToolTip = "Maximum distance between any generated triangle and the original surface. Smaller values make more triangles.", ClampMin = "0.0"))
+	float ChordTolerance;
+
+	/**
+	 * Maximum length of edges of triangles generated by the tessellation process.
+	 * The length is in scene/model unit. The smaller the more triangles are generated.
+	 * Value of 0 means no constraint on length of edges
+	 * Default value is 0.
+	 */
+	UPROPERTY(EditAnywhere, BlueprintReadWrite, AdvancedDisplay, Category = "Tessellation Options", meta = (Units = cm, DisplayName = "Max Edge Length", ToolTip = "Maximum length of any edge in the generated triangles. Smaller values make more triangles.", ClampMin = "0.0"))
+	float MaxEdgeLength;
+
+	/**
+	 * Maximum angle between the normal of two triangles generated by the tessellation process.
+	 * The angle is expressed in degree. The smaller the more triangles are generated.
+	 * Default value is 20 degrees.
+	 */
+	UPROPERTY(EditAnywhere, BlueprintReadWrite, AdvancedDisplay, Category = "Tessellation Options", meta = (Units = deg, ToolTip = "Maximum angle between adjacent triangles. Smaller values make more triangles.", ClampMin = "0.0", ClampMax = "90.0"))
+	float NormalTolerance;
+
+	/** Version 4.24 : Deprecated - not serialized anymore */
+	UPROPERTY()
+	FDatasmithTessellationOptions TessellationSettings_DEPRECATED;
+
+	//~ Begin UDataprepOperation Interface
+public:
+	virtual FText GetCategory_Implementation() const override
+	{
+		return FDataprepOperationCategories::MeshOperation;
+	}
+
+protected:
+	virtual void OnExecution_Implementation(const FDataprepContext& InContext) override;
+	//~ End UDataprepOperation Interface
+};
diff --git a/Engine/Plugins/Enterprise/DatasmithCADImporter/Source/CoreTechExtension/Public/CoreTechBlueprintLibrary.h b/Engine/Plugins/Enterprise/DatasmithCADImporter/Source/CoreTechExtension/Public/CoreTechBlueprintLibrary.h
new file mode 100644
--- /dev/null
+++ b/Engine/Plugins/Enterprise/DatasmithCADImporter/Source/CoreTechExtension/Public/CoreTechBlueprintLibrary.h
@@ -0,0 +1,44 @@
+// Copyright Epic Games, Inc. All Rights Reserved.
+
+#pragma once
+
+#include "CoreMinimal.h"
+
+#include "DatasmithImportOptions.h"
+
+#include "Kismet/BlueprintFunctionLibrary.h"
+#include "UObject/TextProperty.h"
+
+#include "CoreTechBlueprintLibrary.generated.h"
+
+class UStaticMesh;
+
+UCLASS(Blueprintable, BlueprintType, meta = (DisplayName = "CAD Surface Operations Library"))
+class CORETECHEXTENSION_API UCoreTechBlueprintLibrary : public UBlueprintFunctionLibrary
+{
+	GENERATED_BODY()
+
+public:
+	/**
+	 * Re-tessellate LOD 0 of a static mesh if it contains parametric surface data.
+	 * @param	StaticMesh				Static mesh to re-tessellate.
+	 * @param	TessellationSettings	Tessellation settings to use.
+	 * @param	FailureReason			Text describing the reason of failure.
+	 * @return True if successful, false otherwise
+	 */
+	UFUNCTION(BlueprintCallable, Category = "Datasmith | Surface Operations" )
+	static bool RetessellateStaticMesh(UStaticMesh* StaticMesh, const FDatasmithRetessellationOptions& TessellationSettings, FText& FailureReason);
+
+	/**
+	 * Re-tessellate LOD 0 of a static mesh if it contains parametric surface data.
+	 * This implementations allows to skip post edit operations (bApplyChanges=false),
+	 * the caller is then responsible to handle those operations.
+	 * @param	StaticMesh				Static mesh to re-tessellate.
+	 * @param	TessellationSettings	Tessellation settings to use.
+	 * @param	bApplyChanges			Indicates if change must be notified.
+	 * @param	FailureReason			Text describing the reason of failure, or a warning if the operation was successful.
+	 * @return True if successful, false otherwise
+	 */
+	UFUNCTION(Category = "Datasmith | Surface Operations" )
+	static bool RetessellateStaticMeshWithNotification(UStaticMesh* StaticMesh, const FDatasmithRetessellationOptions& TessellationSettings, bool bApplyChanges, FText& FailureReason);
+};
diff --git a/Engine/Plugins/Enterprise/DatasmithCADImporter/Source/CoreTechExtension/Public/CoreTechExtensionModule.h b/Engine/Plugins/Enterprise/DatasmithCADImporter/Source/CoreTechExtension/Public/CoreTechExtensionModule.h
new file mode 100644
--- /dev/null
+++ b/Engine/Plugins/Enterprise/DatasmithCADImporter/Source/CoreTechExtension/Public/CoreTechExtensionModule.h
@@ -0,0 +1,20 @@
+// Copyright Epic Games, Inc. All Rights Reserved.
+
+#pragma once
+
+#include "Modules/ModuleInterface.h"
+
+#define CORETECHEXTENSION_MODULE_NAME TEXT("CoreTechExtension")
+
+/**
+ * This module exposes additional features for assets containing CoreTech data.
+ */
+class FCoreTechExtensionModule : public IModuleInterface
+{
+public:
+    static FCoreTechExtensionModule& Get();
+    static bool IsAvailable();
+
+private:
+	virtual void StartupModule() override;
+};
diff --git a/Engine/Plugins/Enterprise/DatasmithCADImporter/Source/CoreTechExtension/Public/CoreTechRetessellateAction.h b/Engine/Plugins/Enterprise/DatasmithCADImporter/Source/CoreTechExtension/Public/CoreTechRetessellateAction.h
new file mode 100644
--- /dev/null
+++ b/Engine/Plugins/Enterprise/DatasmithCADImporter/Source/CoreTechExtension/Public/CoreTechRetessellateAction.h
@@ -0,0 +1,55 @@
+// Copyright Epic Games, Inc. All Rights Reserved.
+
+#pragma once
+
+#include "CoreMinimal.h"
+
+#include "CoreTechSurfaceExtension.h"
+#include "DatasmithCustomAction.h"
+#include "DatasmithAdditionalData.h"
+#include "DatasmithImportOptions.h"
+
+#include "CoreTechRetessellateAction.generated.h"
+
+class CORETECHEXTENSION_API FCoreTechRetessellate_Impl
+{
+public:
+	static const FText Label;
+	static const FText Tooltip;
+
+	static bool CanApplyOnAssets(const TArray<FAssetData>& SelectedAssets);
+	static void ApplyOnAssets(const TArray<FAssetData>& SelectedAssets);
+	static bool ApplyOnOneAsset(UStaticMesh& StaticMesh, UCoreTechParametricSurfaceData& CoreTechData, const FDatasmithRetessellationOptions& RetesselateOptions);
+};
+
+
+UCLASS()
+class CORETECHEXTENSION_API UCoreTechRetessellateAction : public UDatasmithCustomActionBase
+{
+	GENERATED_BODY()
+
+public:
+	virtual const FText& GetLabel() override { return FCoreTechRetessellate_Impl::Label; }
+	virtual const FText& GetTooltip() override { return FCoreTechRetessellate_Impl::Tooltip; }
+
+	virtual bool CanApplyOnAssets(const TArray<FAssetData>& SelectedAssets) override
+	{ return FCoreTechRetessellate_Impl::CanApplyOnAssets(SelectedAssets); }
+
+	virtual void ApplyOnAssets(const TArray<FAssetData>& SelectedAssets) override
+	{ return FCoreTechRetessellate_Impl::ApplyOnAssets(SelectedAssets); }
+
+	virtual bool CanApplyOnActors(const TArray<AActor*>& SelectedActors) override;
+
+	virtual void ApplyOnActors(const TArray<AActor*>& SelectedActors) override;
+};
+
+
+UCLASS(config = Editor, Transient)
+class CORETECHEXTENSION_API UCoreTechRetessellateActionOptions : public UObject
+{
+	GENERATED_BODY()
+
+public:
+	UPROPERTY(config, EditAnywhere, BlueprintReadWrite, Category = "NotVisible", meta = (ShowOnlyInnerProperties))
+	FDatasmithRetessellationOptions Options;
+};
diff --git a/Engine/Plugins/Enterprise/DatasmithCADImporter/Source/CoreTechSurface/CoreTechSurface.Build.cs b/Engine/Plugins/Enterprise/DatasmithCADImporter/Source/CoreTechSurface/CoreTechSurface.Build.cs
new file mode 100644
--- /dev/null
+++ b/Engine/Plugins/Enterprise/DatasmithCADImporter/Source/CoreTechSurface/CoreTechSurface.Build.cs
@@ -0,0 +1,33 @@
+// Copyright Epic Games, Inc. All Rights Reserved.
+
+using System.IO;
+
+namespace UnrealBuildTool.Rules
+{
+	public class CoreTechSurface : ModuleRules
+	{
+		public CoreTechSurface(ReadOnlyTargetRules Target) : base(Target)
+		{
+			PublicDependencyModuleNames.AddRange(
+				new string[]
+				{
+					"Core",
+				}
+			);
+
+			PrivateDependencyModuleNames.AddRange(
+				new string[]
+				{
+					"Core",
+					"CoreUObject",
+					"MeshDescription",
+					"DatasmithContent",
+					"DatasmithTranslator",
+					"Engine",
+					"StaticMeshDescription",
+					"CADLibrary",
+				}
+			);
+		}
+	}
+}
diff --git a/Engine/Plugins/Enterprise/DatasmithCADImporter/Source/CoreTechSurface/Private/CoreTechSurfaceExtension.cpp b/Engine/Plugins/Enterprise/DatasmithCADImporter/Source/CoreTechSurface/Private/CoreTechSurfaceExtension.cpp
new file mode 100644
--- /dev/null
+++ b/Engine/Plugins/Enterprise/DatasmithCADImporter/Source/CoreTechSurface/Private/CoreTechSurfaceExtension.cpp
@@ -0,0 +1,63 @@
+// Copyright Epic Games, Inc. All Rights Reserved.
+
+#include "CoreTechSurfaceExtension.h"
+#include "DatasmithPayload.h"
+#include "IDatasmithSceneElements.h"
+#include "Misc/FileHelper.h"
+#include "Misc/Paths.h"
+#include "UObject/EnterpriseObjectVersion.h"
+
+#ifdef CAD_LIBRARY
+#include "CADOptions.h"
+#endif // CAD_LIBRARY
+
+void UCoreTechParametricSurfaceData::Serialize(FArchive& Ar)
+{
+	Ar.UsingCustomVersion(FEnterpriseObjectVersion::GUID);
+
+	Super::Serialize(Ar);
+
+	if (Ar.IsSaving() || (Ar.IsLoading() && Ar.CustomVer(FEnterpriseObjectVersion::GUID) >= FEnterpriseObjectVersion::CoreTechParametricSurfaceOptim))
+	{
+		Ar << RawData;
+	}
+
+	if (RawData_DEPRECATED.Num() && RawData.Num() == 0)
+	{
+		RawData = MoveTemp(RawData_DEPRECATED);
+	}
+}
+
+namespace CoreTechSurface
+{
+	// TODO: convert to FCoreTechSceneParameters/FCoreTechMeshParameters ?
+	void AddCoreTechSurfaceDataForMesh(const TSharedRef<IDatasmithMeshElement>& InMeshElement, const CADLibrary::FImportParameters& InSceneParameters, const CADLibrary::FMeshParameters& InMeshParameters, const FDatasmithTessellationOptions& InTessellationOptions, FDatasmithMeshElementPayload& OutMeshPayload)
+	{
+#ifdef CAD_LIBRARY
+		// Store CoreTech additional data if provided
+		const TCHAR* CoretechFile = InMeshElement->GetFile();
+		if (FPaths::FileExists(CoretechFile))
+		{
+			TArray<uint8> ByteArray;
+			if (FFileHelper::LoadFileToArray(ByteArray, CoretechFile))
+			{
+				UCoreTechParametricSurfaceData* CoreTechData = Datasmith::MakeAdditionalData<UCoreTechParametricSurfaceData>();
+				CoreTechData->SourceFile = CoretechFile;
+				CoreTechData->RawData = MoveTemp(ByteArray);
+				CoreTechData->SceneParameters.ModelCoordSys = uint8(InSceneParameters.ModelCoordSys);
+				CoreTechData->SceneParameters.MetricUnit = InSceneParameters.MetricUnit;
+				CoreTechData->SceneParameters.ScaleFactor = InSceneParameters.ScaleFactor;
+
+				CoreTechData->MeshParameters.bNeedSwapOrientation = InMeshParameters.bNeedSwapOrientation;
+				CoreTechData->MeshParameters.bIsSymmetric = InMeshParameters.bIsSymmetric;
+				CoreTechData->MeshParameters.SymmetricNormal = InMeshParameters.SymmetricNormal;
+				CoreTechData->MeshParameters.SymmetricOrigin = InMeshParameters.SymmetricOrigin;
+
+				CoreTechData->LastTessellationOptions = InTessellationOptions;
+				OutMeshPayload.AdditionalData.Add(CoreTechData);
+			}
+		}
+#endif // CAD_LIBRARY
+	}
+
+}
diff --git a/Engine/Plugins/Enterprise/DatasmithCADImporter/Source/CoreTechSurface/Private/CoreTechSurfaceModule.cpp b/Engine/Plugins/Enterprise/DatasmithCADImporter/Source/CoreTechSurface/Private/CoreTechSurfaceModule.cpp
new file mode 100644
--- /dev/null
+++ b/Engine/Plugins/Enterprise/DatasmithCADImporter/Source/CoreTechSurface/Private/CoreTechSurfaceModule.cpp
@@ -0,0 +1,31 @@
+// Copyright Epic Games, Inc. All Rights Reserved.
+
+#include "CoreTechSurfaceModule.h"
+
+#include "Modules/ModuleInterface.h"
+#include "Modules/ModuleManager.h"
+#include "UObject/CoreRedirects.h"
+
+#define LOCTEXT_NAMESPACE "CoreTechSurfaceModule"
+
+void FCoreTechSurfaceModule::StartupModule()
+{
+	TArray<FCoreRedirect> Redirects;
+	Redirects.Emplace(ECoreRedirectFlags::Type_Property, TEXT("UCoreTechParametricSurfaceData.RawData"), TEXT("RawData_DEPRECATED"));
+	FCoreRedirects::AddRedirectList(Redirects, CORETECHSURFACE_MODULE_NAME);
+}
+
+FCoreTechSurfaceModule& FCoreTechSurfaceModule::Get()
+{
+	return FModuleManager::LoadModuleChecked< FCoreTechSurfaceModule >(CORETECHSURFACE_MODULE_NAME);
+}
+
+bool FCoreTechSurfaceModule::IsAvailable()
+{
+	return FModuleManager::Get().IsModuleLoaded(CORETECHSURFACE_MODULE_NAME);
+}
+
+IMPLEMENT_MODULE(FCoreTechSurfaceModule, CoreTechSurface);
+
+#undef LOCTEXT_NAMESPACE // "CoreTechSurfaceModule"
+
diff --git a/Engine/Plugins/Enterprise/DatasmithCADImporter/Source/CoreTechSurface/Private/CoreTechTranslator.cpp b/Engine/Plugins/Enterprise/DatasmithCADImporter/Source/CoreTechSurface/Private/CoreTechTranslator.cpp
new file mode 100644
--- /dev/null
+++ b/Engine/Plugins/Enterprise/DatasmithCADImporter/Source/CoreTechSurface/Private/CoreTechTranslator.cpp
@@ -0,0 +1,43 @@
+// Copyright Epic Games, Inc. All Rights Reserved.
+
+#include "CoreTechTranslator.h"
+
+#include "DatasmithImportOptions.h"
+
+void FCoreTechTranslator::GetSceneImportOptions(TArray<TStrongObjectPtr<UDatasmithOptionsBase>>& Options)
+{
+	FString Extension = GetSource().GetSourceFileExtension();
+	if (Extension == "cgr" || Extension == "3dxml")
+	{
+		return;
+	}
+
+	Options.Add(GetCommonTessellationOptionsPtr());
+}
+
+void FCoreTechTranslator::SetSceneImportOptions(TArray<TStrongObjectPtr<UDatasmithOptionsBase>>& Options)
+{
+	for (const TStrongObjectPtr<UDatasmithOptionsBase>& OptionPtr : Options)
+	{
+		if (UDatasmithCommonTessellationOptions* TessellationOptionsObject = Cast<UDatasmithCommonTessellationOptions>(OptionPtr.Get()))
+		{
+			CommonTessellationOptionsPtr.Reset(TessellationOptionsObject);
+		}
+	}
+}
+
+TStrongObjectPtr<UDatasmithCommonTessellationOptions>& FCoreTechTranslator::GetCommonTessellationOptionsPtr()
+{
+	if (!CommonTessellationOptionsPtr.IsValid())
+	{
+		CommonTessellationOptionsPtr = Datasmith::MakeOptions<UDatasmithCommonTessellationOptions>();
+		check(CommonTessellationOptionsPtr.IsValid());
+		InitCommonTessellationOptions(CommonTessellationOptionsPtr->Options);
+	}
+
+	return CommonTessellationOptionsPtr;
+}
+
+
+
+
diff --git a/Engine/Plugins/Enterprise/DatasmithCADImporter/Source/CoreTechSurface/Public/CoreTechSurfaceExtension.h b/Engine/Plugins/Enterprise/DatasmithCADImporter/Source/CoreTechSurface/Public/CoreTechSurfaceExtension.h
new file mode 100644
--- /dev/null
+++ b/Engine/Plugins/Enterprise/DatasmithCADImporter/Source/CoreTechSurface/Public/CoreTechSurfaceExtension.h
@@ -0,0 +1,89 @@
+// Copyright Epic Games, Inc. All Rights Reserved.
+
+#pragma once
+
+#include "CoreMinimal.h"
+
+#include "DatasmithCustomAction.h"
+#include "DatasmithAdditionalData.h"
+#include "DatasmithImportOptions.h"
+
+#include "CoreTechSurfaceExtension.generated.h"
+
+
+USTRUCT(BlueprintType)
+struct CORETECHSURFACE_API FCoreTechSceneParameters
+{
+	GENERATED_BODY()
+
+	// value from FDatasmithUtils::EModelCoordSystem
+	UPROPERTY()
+	uint8 ModelCoordSys;
+
+	UPROPERTY()
+	float MetricUnit;
+
+	UPROPERTY()
+	float ScaleFactor;
+};
+
+USTRUCT()
+struct CORETECHSURFACE_API FCoreTechMeshParameters
+{
+	GENERATED_BODY()
+
+	UPROPERTY()
+	bool bNeedSwapOrientation;
+
+	UPROPERTY()
+	bool bIsSymmetric;
+
+	UPROPERTY()
+	FVector SymmetricOrigin;
+
+	UPROPERTY()
+	FVector SymmetricNormal;
+};
+
+
+UCLASS(meta = (DisplayName = "Parametric Surface Data"))
+class CORETECHSURFACE_API UCoreTechParametricSurfaceData : public UDatasmithAdditionalData
+{
+	GENERATED_BODY()
+
+public:
+	UPROPERTY()
+	FString SourceFile;
+
+	// Too costly to serialize as a UPROPERTY, will use custom serialization.
+	TArray<uint8> RawData;
+
+	UPROPERTY()
+	FCoreTechSceneParameters SceneParameters;
+
+	UPROPERTY()
+	FCoreTechMeshParameters MeshParameters;
+
+	UPROPERTY(EditAnywhere, Category=NURBS)
+	FDatasmithTessellationOptions LastTessellationOptions;
+
+private:
+	UPROPERTY()
+	TArray<uint8> RawData_DEPRECATED;
+
+	virtual void Serialize(FArchive& Ar) override;
+};
+
+class IDatasmithMeshElement;
+struct FDatasmithMeshElementPayload;
+
+namespace CADLibrary
+{
+	struct FImportParameters;
+	struct FMeshParameters;
+}
+
+namespace CoreTechSurface
+{
+	void CORETECHSURFACE_API AddCoreTechSurfaceDataForMesh(const TSharedRef<IDatasmithMeshElement>& InMeshElement, const CADLibrary::FImportParameters& InSceneParameters, const CADLibrary::FMeshParameters&, const FDatasmithTessellationOptions& InTessellationOptions, FDatasmithMeshElementPayload& OutMeshPayload);
+}
diff --git a/Engine/Plugins/Enterprise/DatasmithCADImporter/Source/CoreTechSurface/Public/CoreTechSurfaceModule.h b/Engine/Plugins/Enterprise/DatasmithCADImporter/Source/CoreTechSurface/Public/CoreTechSurfaceModule.h
new file mode 100644
--- /dev/null
+++ b/Engine/Plugins/Enterprise/DatasmithCADImporter/Source/CoreTechSurface/Public/CoreTechSurfaceModule.h
@@ -0,0 +1,19 @@
+// Copyright Epic Games, Inc. All Rights Reserved.
+
+#pragma once
+
+#include "Modules/ModuleInterface.h"
+
+#define CORETECHSURFACE_MODULE_NAME TEXT("CoreTechSurface")
+
+/**
+ * This module exposes additional features for assets containing CoreTech data.
+ */
+class FCoreTechSurfaceModule : public IModuleInterface
+{
+public:
+	virtual void StartupModule() override;
+
+	static FCoreTechSurfaceModule& Get();
+	static bool IsAvailable();
+};
diff --git a/Engine/Plugins/Enterprise/DatasmithCADImporter/Source/CoreTechSurface/Public/CoreTechTranslator.h b/Engine/Plugins/Enterprise/DatasmithCADImporter/Source/CoreTechSurface/Public/CoreTechTranslator.h
new file mode 100644
--- /dev/null
+++ b/Engine/Plugins/Enterprise/DatasmithCADImporter/Source/CoreTechSurface/Public/CoreTechTranslator.h
@@ -0,0 +1,41 @@
+// Copyright Epic Games, Inc. All Rights Reserved.
+
+#pragma once
+
+#include "CoreMinimal.h"
+
+#include "DatasmithImportOptions.h"
+#include "DatasmithTranslator.h"
+#include "UObject/ObjectMacros.h"
+
+class UDatasmithCommonTessellationOptions;
+
+
+class CORETECHSURFACE_API FCoreTechTranslator : public IDatasmithTranslator
+{
+public:
+	virtual ~FCoreTechTranslator() {}
+
+	// Begin IDatasmithTranslator overrides
+	virtual void GetSceneImportOptions(TArray<TStrongObjectPtr<UDatasmithOptionsBase>>& Options) override;
+	virtual void SetSceneImportOptions(TArray<TStrongObjectPtr<UDatasmithOptionsBase>>& Options) override;
+	// End IDatasmithTranslator overrides
+
+protected:
+	const FDatasmithTessellationOptions& GetCommonTessellationOptions()
+	{
+		return GetCommonTessellationOptionsPtr()->Options;
+	}
+
+	/** 
+	 * Call when the UDatasmithCommonTessellationOptions object is created. This is the unique opportunity for
+	 * child class to overwrite some values
+	 * @param TessellationOptions Reference on member of UDatasmithCommonTessellationOptions
+	 */
+	virtual void InitCommonTessellationOptions(FDatasmithTessellationOptions& TessellationOptions) {}
+
+private:
+	TStrongObjectPtr<UDatasmithCommonTessellationOptions>& GetCommonTessellationOptionsPtr();
+	TStrongObjectPtr<UDatasmithCommonTessellationOptions> CommonTessellationOptionsPtr;
+};
+
diff --git a/Engine/Plugins/Enterprise/DatasmithCADImporter/Source/DatasmithCADTranslator/Private/DatasmithCADTranslator.cpp b/Engine/Plugins/Enterprise/DatasmithCADImporter/Source/DatasmithCADTranslator/Private/DatasmithCADTranslator.cpp
--- a/Engine/Plugins/Enterprise/DatasmithCADImporter/Source/DatasmithCADTranslator/Private/DatasmithCADTranslator.cpp
+++ b/Engine/Plugins/Enterprise/DatasmithCADImporter/Source/DatasmithCADTranslator/Private/DatasmithCADTranslator.cpp
@@ -5,7 +5,7 @@
 #ifdef CAD_LIBRARY
 
 #include "CADInterfacesModule.h"
-#include "CoreTechParametricSurfaceExtension.h"
+#include "CoreTechSurfaceExtension.h"
 #include "DatasmithCADTranslatorModule.h"
 #include "DatasmithDispatcher.h"
 #include "DatasmithMeshBuilder.h"
@@ -171,14 +171,14 @@ bool FDatasmithCADTranslator::LoadStaticMesh(const TSharedRef<IDatasmithMeshElem
 	{
 		OutMeshPayload.LodMeshes.Add(MoveTemp(Mesh.GetValue()));
 
-		DatasmithCoreTechParametricSurfaceData::AddCoreTechSurfaceDataForMesh(MeshElement, ImportParameters, MeshParameters, GetCommonTessellationOptions(), OutMeshPayload);
+		CoreTechSurface::AddCoreTechSurfaceDataForMesh(MeshElement, ImportParameters, MeshParameters, GetCommonTessellationOptions(), OutMeshPayload);
 	}
 	return OutMeshPayload.LodMeshes.Num() > 0;
 }
 
 void FDatasmithCADTranslator::SetSceneImportOptions(TArray<TStrongObjectPtr<UDatasmithOptionsBase>>& Options)
 {
-	FDatasmithCoreTechTranslator::SetSceneImportOptions(Options);
+	FCoreTechTranslator::SetSceneImportOptions(Options);
 }
 #endif
 
diff --git a/Engine/Plugins/Enterprise/DatasmithCADImporter/Source/DatasmithCADTranslator/Private/DatasmithCADTranslator.h b/Engine/Plugins/Enterprise/DatasmithCADImporter/Source/DatasmithCADTranslator/Private/DatasmithCADTranslator.h
--- a/Engine/Plugins/Enterprise/DatasmithCADImporter/Source/DatasmithCADTranslator/Private/DatasmithCADTranslator.h
+++ b/Engine/Plugins/Enterprise/DatasmithCADImporter/Source/DatasmithCADTranslator/Private/DatasmithCADTranslator.h
@@ -4,7 +4,7 @@
 
 #include "CoreMinimal.h"
 
-#include "DatasmithCoreTechTranslator.h"
+#include "CoreTechTranslator.h"
 
 #ifdef CAD_LIBRARY
 #include "CADOptions.h"
@@ -17,7 +17,7 @@
 
 class IDatasmithMeshElement;
 
-class FDatasmithCADTranslator : public FDatasmithCoreTechTranslator
+class FDatasmithCADTranslator : public FCoreTechTranslator
 {
 public:
 	virtual FName GetFName() const override { return "DatasmithCADTranslator"; };
diff --git a/Engine/Plugins/Enterprise/DatasmithCADImporter/Source/DatasmithCoreTechExtension/DatasmithCoreTechExtension.Build.cs b/Engine/Plugins/Enterprise/DatasmithCADImporter/Source/DatasmithCoreTechExtension/DatasmithCoreTechExtension.Build.cs
deleted file mode 100644
diff --git a/Engine/Plugins/Enterprise/DatasmithCADImporter/Source/DatasmithCoreTechExtension/Private/CoreTechBlueprintLibrary.cpp b/Engine/Plugins/Enterprise/DatasmithCADImporter/Source/DatasmithCoreTechExtension/Private/CoreTechBlueprintLibrary.cpp
deleted file mode 100644
diff --git a/Engine/Plugins/Enterprise/DatasmithCADImporter/Source/DatasmithCoreTechExtension/Private/CoreTechRetessellateAction.cpp b/Engine/Plugins/Enterprise/DatasmithCADImporter/Source/DatasmithCoreTechExtension/Private/CoreTechRetessellateAction.cpp
deleted file mode 100644
diff --git a/Engine/Plugins/Enterprise/DatasmithCADImporter/Source/DatasmithCoreTechExtension/Private/DataprepTessellationOperation.cpp b/Engine/Plugins/Enterprise/DatasmithCADImporter/Source/DatasmithCoreTechExtension/Private/DataprepTessellationOperation.cpp
deleted file mode 100644
diff --git a/Engine/Plugins/Enterprise/DatasmithCADImporter/Source/DatasmithCoreTechExtension/Private/DataprepTessellationOperation.h b/Engine/Plugins/Enterprise/DatasmithCADImporter/Source/DatasmithCoreTechExtension/Private/DataprepTessellationOperation.h
deleted file mode 100644
diff --git a/Engine/Plugins/Enterprise/DatasmithCADImporter/Source/DatasmithCoreTechExtension/Private/DatasmithCoreTechExtensionModule.cpp b/Engine/Plugins/Enterprise/DatasmithCADImporter/Source/DatasmithCoreTechExtension/Private/DatasmithCoreTechExtensionModule.cpp
deleted file mode 100644
diff --git a/Engine/Plugins/Enterprise/DatasmithCADImporter/Source/DatasmithCoreTechExtension/Public/CoreTechBlueprintLibrary.h b/Engine/Plugins/Enterprise/DatasmithCADImporter/Source/DatasmithCoreTechExtension/Public/CoreTechBlueprintLibrary.h
deleted file mode 100644
diff --git a/Engine/Plugins/Enterprise/DatasmithCADImporter/Source/DatasmithCoreTechExtension/Public/CoreTechRetessellateAction.h b/Engine/Plugins/Enterprise/DatasmithCADImporter/Source/DatasmithCoreTechExtension/Public/CoreTechRetessellateAction.h
deleted file mode 100644
diff --git a/Engine/Plugins/Enterprise/DatasmithCADImporter/Source/DatasmithCoreTechExtension/Public/DatasmithCoreTechExtensionModule.h b/Engine/Plugins/Enterprise/DatasmithCADImporter/Source/DatasmithCoreTechExtension/Public/DatasmithCoreTechExtensionModule.h
deleted file mode 100644
diff --git a/Engine/Plugins/Enterprise/DatasmithCADImporter/Source/DatasmithCoreTechParametricSurfaceData/DatasmithCoreTechParametricSurfaceData.Build.cs b/Engine/Plugins/Enterprise/DatasmithCADImporter/Source/DatasmithCoreTechParametricSurfaceData/DatasmithCoreTechParametricSurfaceData.Build.cs
deleted file mode 100644
diff --git a/Engine/Plugins/Enterprise/DatasmithCADImporter/Source/DatasmithCoreTechParametricSurfaceData/Private/CoreTechParametricSurfaceExtension.cpp b/Engine/Plugins/Enterprise/DatasmithCADImporter/Source/DatasmithCoreTechParametricSurfaceData/Private/CoreTechParametricSurfaceExtension.cpp
deleted file mode 100644
diff --git a/Engine/Plugins/Enterprise/DatasmithCADImporter/Source/DatasmithCoreTechParametricSurfaceData/Private/DatasmithCoreTechParametricSurfaceDataModule.cpp b/Engine/Plugins/Enterprise/DatasmithCADImporter/Source/DatasmithCoreTechParametricSurfaceData/Private/DatasmithCoreTechParametricSurfaceDataModule.cpp
deleted file mode 100644
diff --git a/Engine/Plugins/Enterprise/DatasmithCADImporter/Source/DatasmithCoreTechParametricSurfaceData/Private/DatasmithCoreTechTranslator.cpp b/Engine/Plugins/Enterprise/DatasmithCADImporter/Source/DatasmithCoreTechParametricSurfaceData/Private/DatasmithCoreTechTranslator.cpp
deleted file mode 100644
diff --git a/Engine/Plugins/Enterprise/DatasmithCADImporter/Source/DatasmithCoreTechParametricSurfaceData/Public/CoreTechParametricSurfaceExtension.h b/Engine/Plugins/Enterprise/DatasmithCADImporter/Source/DatasmithCoreTechParametricSurfaceData/Public/CoreTechParametricSurfaceExtension.h
deleted file mode 100644
diff --git a/Engine/Plugins/Enterprise/DatasmithCADImporter/Source/DatasmithCoreTechParametricSurfaceData/Public/DatasmithCoreTechParametricSurfaceDataModule.h b/Engine/Plugins/Enterprise/DatasmithCADImporter/Source/DatasmithCoreTechParametricSurfaceData/Public/DatasmithCoreTechParametricSurfaceDataModule.h
deleted file mode 100644
diff --git a/Engine/Plugins/Enterprise/DatasmithCADImporter/Source/DatasmithCoreTechParametricSurfaceData/Public/DatasmithCoreTechTranslator.h b/Engine/Plugins/Enterprise/DatasmithCADImporter/Source/DatasmithCoreTechParametricSurfaceData/Public/DatasmithCoreTechTranslator.h
deleted file mode 100644
diff --git a/Engine/Plugins/Enterprise/DatasmithCADImporter/Source/DatasmithWireTranslator/Private/DatasmithWireTranslator.cpp b/Engine/Plugins/Enterprise/DatasmithCADImporter/Source/DatasmithWireTranslator/Private/DatasmithWireTranslator.cpp
--- a/Engine/Plugins/Enterprise/DatasmithCADImporter/Source/DatasmithWireTranslator/Private/DatasmithWireTranslator.cpp
+++ b/Engine/Plugins/Enterprise/DatasmithCADImporter/Source/DatasmithWireTranslator/Private/DatasmithWireTranslator.cpp
@@ -27,7 +27,7 @@
 #ifdef CAD_LIBRARY
 #include "AliasCoretechWrapper.h" // requires CoreTech as public dependency
 #include "CADInterfacesModule.h"
-#include "CoreTechParametricSurfaceExtension.h"
+#include "CoreTechSurfaceExtension.h"
 #endif
 
 #ifdef USE_OPENMODEL
@@ -2151,7 +2151,7 @@ bool FDatasmithWireTranslator::LoadStaticMesh(const TSharedRef<IDatasmithMeshEle
 	{
 		OutMeshPayload.LodMeshes.Add(MoveTemp(Mesh.GetValue()));
 #ifdef CAD_LIBRARY
-		DatasmithCoreTechParametricSurfaceData::AddCoreTechSurfaceDataForMesh(MeshElement, ImportParameters, MeshParameters, GetCommonTessellationOptions(), OutMeshPayload);
+		CoreTechSurface::AddCoreTechSurfaceDataForMesh(MeshElement, ImportParameters, MeshParameters, GetCommonTessellationOptions(), OutMeshPayload);
 #endif //CAD_LIBRARY
 	}
 	return OutMeshPayload.LodMeshes.Num() > 0;
@@ -2163,7 +2163,7 @@ bool FDatasmithWireTranslator::LoadStaticMesh(const TSharedRef<IDatasmithMeshEle
 void FDatasmithWireTranslator::SetSceneImportOptions(TArray<TStrongObjectPtr<UDatasmithOptionsBase>>& Options)
 {
 #ifdef USE_OPENMODEL
-	FDatasmithCoreTechTranslator::SetSceneImportOptions(Options);
+	FCoreTechTranslator::SetSceneImportOptions(Options);
 
 	if (Translator)
 	{
diff --git a/Engine/Plugins/Enterprise/DatasmithCADImporter/Source/DatasmithWireTranslator/Private/DatasmithWireTranslator.h b/Engine/Plugins/Enterprise/DatasmithCADImporter/Source/DatasmithWireTranslator/Private/DatasmithWireTranslator.h
--- a/Engine/Plugins/Enterprise/DatasmithCADImporter/Source/DatasmithWireTranslator/Private/DatasmithWireTranslator.h
+++ b/Engine/Plugins/Enterprise/DatasmithCADImporter/Source/DatasmithWireTranslator/Private/DatasmithWireTranslator.h
@@ -8,14 +8,14 @@
 #undef USE_OPENMODEL
 #endif
 
-#include "DatasmithCoreTechTranslator.h"
+#include "CoreTechTranslator.h"
 #include "UObject/ObjectMacros.h"
 
 class FWireTranslatorImpl;
 class AlDagNode;
 
 
-class FDatasmithWireTranslator : public FDatasmithCoreTechTranslator
+class FDatasmithWireTranslator : public FCoreTechTranslator
 {
 public:
 	FDatasmithWireTranslator();
