Rename DatasmithCoreTechParametricSurfaceData into CoreTechSurface and DatasmithCoreTechExtension into CoreTechExtension

diff --git a/Engine/Plugins/Enterprise/DatasmithCADImporter/DatasmithCADImporter.uplugin b/Engine/Plugins/Enterprise/DatasmithCADImporter/DatasmithCADImporter.uplugin
--- a/Engine/Plugins/Enterprise/DatasmithCADImporter/DatasmithCADImporter.uplugin
+++ b/Engine/Plugins/Enterprise/DatasmithCADImporter/DatasmithCADImporter.uplugin
@@ -74,7 +74,7 @@
 			]
 		},
 		{
-			"Name": "DatasmithCoreTechExtension",
+			"Name": "CoreTechExtension",
 			"Type": "Editor",
 			"LoadingPhase": "Default",
 			"WhitelistPlatforms": [
@@ -82,7 +82,7 @@
 			]
 		},
 		{
-			"Name": "DatasmithCoreTechParametricSurfaceData",
+			"Name": "CoreTechSurface",
 			"Type": "Runtime",
 			"LoadingPhase": "Default",
 			"WhitelistPlatforms": [
diff --git a/Engine/Plugins/Enterprise/DatasmithCADImporter/Source/DatasmithCADTranslator/DatasmithCADTranslator.Build.cs b/Engine/Plugins/Enterprise/DatasmithCADImporter/Source/DatasmithCADTranslator/DatasmithCADTranslator.Build.cs
--- a/Engine/Plugins/Enterprise/DatasmithCADImporter/Source/DatasmithCADTranslator/DatasmithCADTranslator.Build.cs
+++ b/Engine/Plugins/Enterprise/DatasmithCADImporter/Source/DatasmithCADTranslator/DatasmithCADTranslator.Build.cs
@@ -8,7 +8,7 @@ namespace UnrealBuildTool.Rules
 	{
 		public DatasmithCADTranslator(ReadOnlyTargetRules Target) : base(Target)
 		{
-            PublicDependencyModuleNames.AddRange(
+			PublicDependencyModuleNames.AddRange(
 				new string[]
 				{
 					"Core",
@@ -22,8 +22,8 @@ namespace UnrealBuildTool.Rules
 					"CADInterfaces",
 					"CADLibrary",
 					"CADTools",
+					"CoreTechSurface",
 					"DatasmithCore",
-					"DatasmithCoreTechParametricSurfaceData",
 					"DatasmithContent",
 					"DatasmithDispatcher",
 					"DatasmithTranslator",
diff --git a/Engine/Plugins/Enterprise/DatasmithCADImporter/Source/DatasmithOpenNurbsTranslator/DatasmithOpenNurbsTranslator.Build.cs b/Engine/Plugins/Enterprise/DatasmithCADImporter/Source/DatasmithOpenNurbsTranslator/DatasmithOpenNurbsTranslator.Build.cs
--- a/Engine/Plugins/Enterprise/DatasmithCADImporter/Source/DatasmithOpenNurbsTranslator/DatasmithOpenNurbsTranslator.Build.cs
+++ b/Engine/Plugins/Enterprise/DatasmithCADImporter/Source/DatasmithOpenNurbsTranslator/DatasmithOpenNurbsTranslator.Build.cs
@@ -22,11 +22,11 @@ namespace UnrealBuildTool.Rules
 				{
                     "CADInterfaces",
                     "CADLibrary",
-					"Engine",
+                    "CoreTechSurface",
 					"DatasmithCore",
-                    "DatasmithCoreTechParametricSurfaceData",
 					"DatasmithContent",
 					"DatasmithTranslator",
+					"Engine",
 					"MeshDescription",
 					"StaticMeshDescription",
 				}
diff --git a/Engine/Plugins/Enterprise/DatasmithCADImporter/Source/DatasmithOpenNurbsTranslator/Private/DatasmithOpenNurbsTranslator.cpp b/Engine/Plugins/Enterprise/DatasmithCADImporter/Source/DatasmithOpenNurbsTranslator/Private/DatasmithOpenNurbsTranslator.cpp
--- a/Engine/Plugins/Enterprise/DatasmithCADImporter/Source/DatasmithOpenNurbsTranslator/Private/DatasmithOpenNurbsTranslator.cpp
+++ b/Engine/Plugins/Enterprise/DatasmithCADImporter/Source/DatasmithOpenNurbsTranslator/Private/DatasmithOpenNurbsTranslator.cpp
@@ -6,7 +6,7 @@
 #ifdef USE_OPENNURBS // The whole translation unit is skipped without OpenNurbs TPS library
 
 #ifdef CAD_LIBRARY
-#include "CoreTechParametricSurfaceExtension.h"
+#include "CoreTechSurfaceExtension.h"
 #include "RhinoCoretechWrapper.h"
 #endif // CAD_LIBRARY
 
@@ -3514,7 +3514,7 @@ bool FDatasmithOpenNurbsTranslator::LoadStaticMesh(const TSharedRef<IDatasmithMe
 
 		CADLibrary::FMeshParameters MeshParameters;
 
-		DatasmithCoreTechParametricSurfaceData::AddCoreTechSurfaceDataForMesh(MeshElement, ImportParameters, MeshParameters, OpenNurbsOptions, OutMeshPayload);
+		CoreTechSurface::AddCoreTechSurfaceDataForMesh(MeshElement, ImportParameters, MeshParameters, OpenNurbsOptions, OutMeshPayload);
 #endif // CAD_LIBRARY
 
 	}
diff --git a/Engine/Plugins/Enterprise/DatasmithCADImporter/Source/DatasmithPLMXMLTranslator/DatasmithPLMXMLTranslator.Build.cs b/Engine/Plugins/Enterprise/DatasmithCADImporter/Source/DatasmithPLMXMLTranslator/DatasmithPLMXMLTranslator.Build.cs
--- a/Engine/Plugins/Enterprise/DatasmithCADImporter/Source/DatasmithPLMXMLTranslator/DatasmithPLMXMLTranslator.Build.cs
+++ b/Engine/Plugins/Enterprise/DatasmithCADImporter/Source/DatasmithPLMXMLTranslator/DatasmithPLMXMLTranslator.Build.cs
@@ -30,7 +30,7 @@ namespace UnrealBuildTool.Rules
 					"CADInterfaces",
 					"CADLibrary",
 					"CADTools",
-					"DatasmithCoreTechParametricSurfaceData",
+					"CoreTechSurface",
 					"DatasmithCADTranslator", // for DatasmithMeshBuilder
                     "DatasmithTranslator",
                     "DatasmithDispatcher",
diff --git a/Engine/Plugins/Enterprise/DatasmithCADImporter/Source/DatasmithPLMXMLTranslator/Private/DatasmithPlmXmlImporter.cpp b/Engine/Plugins/Enterprise/DatasmithCADImporter/Source/DatasmithPLMXMLTranslator/Private/DatasmithPlmXmlImporter.cpp
--- a/Engine/Plugins/Enterprise/DatasmithCADImporter/Source/DatasmithPLMXMLTranslator/Private/DatasmithPlmXmlImporter.cpp
+++ b/Engine/Plugins/Enterprise/DatasmithCADImporter/Source/DatasmithPLMXMLTranslator/Private/DatasmithPlmXmlImporter.cpp
@@ -20,7 +20,7 @@
 #ifdef CAD_LIBRARY
 #include "CADData.h"
 #include "CADToolsModule.h"
-#include "CoreTechParametricSurfaceExtension.h"
+#include "CoreTechSurfaceExtension.h"
 #endif // CAD_LIBRARY
 #include "DatasmithAdditionalData.h"
 #include "DatasmithDispatcher.h"
@@ -306,7 +306,7 @@ namespace PlmXml
 			{
 				OutMeshPayload.LodMeshes.Add(MoveTemp(Mesh.GetValue()));
 
-				DatasmithCoreTechParametricSurfaceData::AddCoreTechSurfaceDataForMesh(MeshElement, ImportParameters, MeshParameters, TessellationOptions, OutMeshPayload);
+				CoreTechSurface::AddCoreTechSurfaceDataForMesh(MeshElement, ImportParameters, MeshParameters, TessellationOptions, OutMeshPayload);
 				
 				return true;
 			}
diff --git a/Engine/Plugins/Enterprise/DatasmithCADImporter/Source/DatasmithWireTranslator/DatasmithWireTranslator.Build.cs b/Engine/Plugins/Enterprise/DatasmithCADImporter/Source/DatasmithWireTranslator/DatasmithWireTranslator.Build.cs
--- a/Engine/Plugins/Enterprise/DatasmithCADImporter/Source/DatasmithWireTranslator/DatasmithWireTranslator.Build.cs
+++ b/Engine/Plugins/Enterprise/DatasmithCADImporter/Source/DatasmithWireTranslator/DatasmithWireTranslator.Build.cs
@@ -21,9 +21,9 @@ public class DatasmithWireTranslator : ModuleRules
                 "CADInterfaces",
                 "CADLibrary",
                 "CADTools",
+				"CoreTechSurface",
                 "DatasmithContent",
 				"DatasmithCore",
-				"DatasmithCoreTechParametricSurfaceData",
 				"DatasmithTranslator",
 				"Engine",
 				"MeshDescription",
