Disabling linkage to CADKernel library to fix linkage errors
    - The CADKernel functionalities will be restored when David is back from vacations

diff --git a/Engine/Plugins/Enterprise/DatasmithCADImporter/Source/CADInterfaces/Private/CoreTechFileParser.cpp b/Engine/Plugins/Enterprise/DatasmithCADImporter/Source/CADInterfaces/Private/CoreTechFileParser.cpp
--- a/Engine/Plugins/Enterprise/DatasmithCADImporter/Source/CADInterfaces/Private/CoreTechFileParser.cpp
+++ b/Engine/Plugins/Enterprise/DatasmithCADImporter/Source/CADInterfaces/Private/CoreTechFileParser.cpp
@@ -16,8 +16,10 @@
 #include "Misc/Paths.h"
 #include "Templates/TypeHash.h"
 
+#ifdef CADKERNELLIB
 #include "CADKernel/Mesh/Structure/SurfaceMesh.h"
 #include "CADKernel/Math/Point.h"
+#endif
 
 namespace CADLibrary 
 {
@@ -572,10 +574,14 @@ void FCoreTechFileParser::ReadMaterials()
 FCoreTechFileParser::FCoreTechFileParser(const FImportParameters& ImportParams, const FString& EnginePluginsPath, const FString& InCachePath)
 	: CachePath(InCachePath)
 	, ImportParameters(ImportParams)
+#ifdef CADKERNELLIB
 	, CoreTechBridge()
+#endif
 {
 	CTKIO_InitializeKernel(ImportParameters.MetricUnit, *EnginePluginsPath);
+#ifdef CADKERNELLIB
 	CoreTechBridge.InitializeCADKernel(0.00001 / ImportParams.MetricUnit);
+#endif
 }
 
 bool FCoreTechFileParser::FindFile(FFileDescription& File)
@@ -818,7 +824,7 @@ FCoreTechFileParser::EProcessResult FCoreTechFileParser::ReadFileWithKernelIO()
 	return EProcessResult::ProcessOk;
 }
 
-	CT_FLAGS FCoreTechFileParser::SetCoreTechImportOption()
+CT_FLAGS FCoreTechFileParser::SetCoreTechImportOption()
 {
 	// Set import option
 	CT_FLAGS Flags = CT_LOAD_FLAGS_USE_DEFAULT;
@@ -1131,6 +1137,7 @@ uint32 GetBodiesFaceSetNum(TArray<CT_OBJECT_ID>& BodySet)
 	return size;
 }
 
+#ifdef CADKERNELLIB
 void FCoreTechFileParser::DefineMeshCriteria(FIdent MeshModelId)
 {
 	TArray<FIdent> Criteria;
@@ -1254,6 +1261,7 @@ void FCoreTechFileParser::GetBodyTessellation(CT_OBJECT_ID CTBodyId, CT_OBJECT_I
 		FaceIndex++;
 	}
 }
+#endif
 
 void FCoreTechFileParser::GetKioBodyTessellation(CT_OBJECT_ID BodyId, CT_OBJECT_ID ParentId, FBodyMesh& OutBodyMesh, uint32 DefaultMaterialHash, bool bNeedRepair)
 {
@@ -1394,6 +1402,7 @@ bool FCoreTechFileParser::ReadBody(CT_OBJECT_ID BodyId, CT_OBJECT_ID ParentId, u
 	CT_FLAGS BodyProperties;
 	CT_BODY_IO::AskProperties(BodyId, BodyProperties);
 
+#ifdef CADKERNELLIB
 	if (ImportParameters.bEnableKernelIOTessellation || !(BodyProperties & CT_BODY_PROP_EXACT))
 	{
 		GetKioBodyTessellation(BodyId, ParentId, BodyMeshes[BodyMeshIndex], DefaultMaterialHash, bNeedRepair);
@@ -1402,6 +1411,9 @@ bool FCoreTechFileParser::ReadBody(CT_OBJECT_ID BodyId, CT_OBJECT_ID ParentId, u
 	{
 		GetBodyTessellation(BodyId, ParentId, BodyMeshes[BodyMeshIndex], DefaultMaterialHash, bNeedRepair, BodyFile);
 	}
+#else
+	GetKioBodyTessellation(BodyId, ParentId, BodyMeshes[BodyMeshIndex], DefaultMaterialHash, bNeedRepair);
+#endif
 	SceneGraphArchive.BodySet[Index].ColorFaceSet = BodyMeshes[BodyMeshIndex].ColorSet;
 	SceneGraphArchive.BodySet[Index].MaterialFaceSet = BodyMeshes[BodyMeshIndex].MaterialSet;
 	
diff --git a/Engine/Plugins/Enterprise/DatasmithCADImporter/Source/CADInterfaces/Public/CoreTechFileParser.h b/Engine/Plugins/Enterprise/DatasmithCADImporter/Source/CADInterfaces/Public/CoreTechFileParser.h
--- a/Engine/Plugins/Enterprise/DatasmithCADImporter/Source/CADInterfaces/Public/CoreTechFileParser.h
+++ b/Engine/Plugins/Enterprise/DatasmithCADImporter/Source/CADInterfaces/Public/CoreTechFileParser.h
@@ -19,7 +19,9 @@
 #endif // CAD_INTERFACE
 
 #ifdef CAD_INTERFACE // temporary in the wait of a full integration in UE
+#ifdef CADKERNELLIB
 #include "CADKernel/CADInterface/CoreTech/CoreTechBridge.h"
+#endif
 #endif // CAD_INTERFACE
 
 struct FFileStatData;
@@ -69,9 +71,11 @@ public:
 
 	void GetKioBodyTessellation(CT_OBJECT_ID BodyId, CT_OBJECT_ID ParentId, FBodyMesh& OutBodyMesh, uint32 ParentMaterialHash, bool bNeedRepair);
 
+#ifdef CADKERNELLIB
 	void GetBodyTessellation(CT_OBJECT_ID BodyId, CT_OBJECT_ID ParentId, FBodyMesh& OutBodyMesh, uint32 ParentMaterialHash, bool bNeedRepair, FString BodyFile);
 	uint32 GetFaceTessellation(FIdent FaceID, TArray<FTessellationData>& FaceTessellationSet, const FImportParameters& ImportParams);
 	void DefineMeshCriteria(FIdent MeshModelId);
+#endif	
 
 	TSet<FFileDescription>& GetExternalRefSet()
 	{
@@ -143,7 +147,9 @@ protected:
 
 	const FImportParameters& ImportParameters;
 
+#ifdef CADKERNELLIB
 	CADKernel::FCoreTechBridge CoreTechBridge;
+#endif
 };
 
 #endif // CAD_INTERFACE
diff --git a/Engine/Plugins/Enterprise/DatasmithCADImporter/Source/DatasmithCADTranslator/Private/DatasmithCADTranslator.cpp b/Engine/Plugins/Enterprise/DatasmithCADImporter/Source/DatasmithCADTranslator/Private/DatasmithCADTranslator.cpp
--- a/Engine/Plugins/Enterprise/DatasmithCADImporter/Source/DatasmithCADTranslator/Private/DatasmithCADTranslator.cpp
+++ b/Engine/Plugins/Enterprise/DatasmithCADImporter/Source/DatasmithCADTranslator/Private/DatasmithCADTranslator.cpp
@@ -20,11 +20,13 @@ static TAutoConsoleVariable<int32> CVarStaticCADTranslatorEnableThreadedImport(
 	TEXT("Activate to parallelise CAD file processing.\n"),
 	ECVF_Default);
 
+#ifdef CADKERNELLIB
 static TAutoConsoleVariable<int32> CVarStaticCADTranslatorEnableKernelIOTessellation(
 	TEXT("r.CADTranslator.EnableKernelIOTessellation"),
 	1,
 	TEXT("Activate to use KernelIO Tessellation.\n"),
 	ECVF_Default);
+#endif
 
 void FDatasmithCADTranslator::Initialize(FDatasmithTranslatorCapabilities& OutCapabilities)
 {
@@ -82,9 +84,11 @@ void FDatasmithCADTranslator::Initialize(FDatasmithTranslatorCapabilities& OutCa
 
 bool FDatasmithCADTranslator::LoadScene(TSharedRef<IDatasmithScene> DatasmithScene)
 {
+#ifdef CADKERNELLIB
 	ImportParameters.bEnableKernelIOTessellation = (CVarStaticCADTranslatorEnableKernelIOTessellation.GetValueOnAnyThread() != 0);;
 #ifdef CADKERNEL_DEBUG
 	ImportParameters.bEnableKernelIOTessellation = false;
+#endif
 #endif
 
 	ImportParameters.MetricUnit = 0.001;
