Tessellation.PositionIndices has INDEX_NONE index for bad vertices. Triangle indices have to be tested before affecting VertexIds.

diff --git a/Engine/Plugins/Enterprise/DatasmithCADImporter/Source/CADLibrary/Private/CoreTechHelper.cpp b/Engine/Plugins/Enterprise/DatasmithCADImporter/Source/CADLibrary/Private/CoreTechHelper.cpp
--- a/Engine/Plugins/Enterprise/DatasmithCADImporter/Source/CADLibrary/Private/CoreTechHelper.cpp
+++ b/Engine/Plugins/Enterprise/DatasmithCADImporter/Source/CADLibrary/Private/CoreTechHelper.cpp
@@ -452,15 +452,19 @@ namespace CADLibrary
 				// build each valid face i.e. 3 different indexes
 				for (int32 Index = 0; Index < Tessellation.VertexIndices.Num(); Index += 3)
 				{
-					VertexIDs[Orientation[0]] = VertexIdSet[Tessellation.PositionIndices[Tessellation.VertexIndices[Index + 0]]];
-					VertexIDs[Orientation[1]] = VertexIdSet[Tessellation.PositionIndices[Tessellation.VertexIndices[Index + 1]]];
-					VertexIDs[Orientation[2]] = VertexIdSet[Tessellation.PositionIndices[Tessellation.VertexIndices[Index + 2]]];
+					VertexIDs[Orientation[0]] = Tessellation.PositionIndices[Tessellation.VertexIndices[Index + 0]];
+					VertexIDs[Orientation[1]] = Tessellation.PositionIndices[Tessellation.VertexIndices[Index + 1]];
+					VertexIDs[Orientation[2]] = Tessellation.PositionIndices[Tessellation.VertexIndices[Index + 2]];
 
 					if (VertexIDs[0] == INDEX_NONE || VertexIDs[1] == INDEX_NONE || VertexIDs[2] == INDEX_NONE)
 					{
 						continue;
 					}
 
+					VertexIDs[0] = VertexIdSet[VertexIDs[0]];
+					VertexIDs[1] = VertexIdSet[VertexIDs[1]];
+					VertexIDs[2] = VertexIdSet[VertexIDs[2]];
+
 					// Verify the 3 input indices are not defining a degenerated triangle
 					if (VertexIDs[0] == VertexIDs[1] || VertexIDs[0] == VertexIDs[2] || VertexIDs[1] == VertexIDs[2])
 					{
