Add polygroups management in retesselate

diff --git a/Engine/Plugins/Enterprise/DatasmithCADImporter/Source/CADLibrary/Private/CoreTechHelper.cpp b/Engine/Plugins/Enterprise/DatasmithCADImporter/Source/CADLibrary/Private/CoreTechHelper.cpp
--- a/Engine/Plugins/Enterprise/DatasmithCADImporter/Source/CADLibrary/Private/CoreTechHelper.cpp
+++ b/Engine/Plugins/Enterprise/DatasmithCADImporter/Source/CADLibrary/Private/CoreTechHelper.cpp
@@ -418,6 +418,10 @@ namespace CADLibrary
 			NbStep = 2;
 		}
 
+		TSet<int32> PatchIdSet;
+		GetExistingPatches(MeshDescription, PatchIdSet);
+		bool bImportOnlyAlreadyPresent = (bool)PatchIdSet.Num();
+
 		TArray<FVertexInstanceID> TriangleVertexInstanceIDs;
 		TriangleVertexInstanceIDs.SetNum(VerticesPerTriangle);
 
@@ -436,7 +440,11 @@ namespace CADLibrary
 
 			for (FTessellationData& Tessellation : BodyTessellation.Faces)
 			{
-				int32 VertexIDs[3];
+
+				if (bImportOnlyAlreadyPresent && !PatchIdSet.Contains(Tessellation.PatchId))
+				{
+					continue;
+				}
 
 				// Get the polygonGroup
 				const FPolygonGroupID* PolygonGroupID = MaterialToPolygonGroupMapping.Find(Tessellation.ColorName);
@@ -445,6 +453,8 @@ namespace CADLibrary
 					continue;
 				}
 
+				int32 VertexIDs[3];
+
 				PatchIndex++;
 				MeshVertexInstanceIDs.Empty(Tessellation.VertexIndices.Num());
 				NewFaceIndex.Empty(Tessellation.VertexIndices.Num());
