In case of file with many body, the face id are re-indent in the sub CT file.
During a re-tessellation with skip deleted option, the face id of the new mesh does not match with the polygroup id.
To avoid this probleme, the initial id is saved in each face as an integer attribut named "DatasmithFaceId"

diff --git a/Engine/Plugins/Enterprise/DatasmithCADImporter/Source/CADInterfaces/Private/CoreTechFileParser.cpp b/Engine/Plugins/Enterprise/DatasmithCADImporter/Source/CADInterfaces/Private/CoreTechFileParser.cpp
--- a/Engine/Plugins/Enterprise/DatasmithCADImporter/Source/CADInterfaces/Private/CoreTechFileParser.cpp
+++ b/Engine/Plugins/Enterprise/DatasmithCADImporter/Source/CADInterfaces/Private/CoreTechFileParser.cpp
@@ -24,6 +24,8 @@
 namespace CADLibrary 
 {
 
+int32 GetIntegerParameterDataValue(CT_OBJECT_ID NodeId, const TCHAR* InMetaDataName);
+
 namespace {
 
 	double Distance(const CT_COORDINATE& Point1, const CT_COORDINATE& Point2)
@@ -161,11 +163,11 @@ FString AsFString(CT_STR CtName)
 	return CtName.IsEmpty() ? FString() : CtName.toUnicode();
 };
 
-	uint32 FCoreTechFileParser::GetFileHash()
+uint32 FCoreTechFileParser::GetFileHash()
 {
 	FFileStatData FileStatData = IFileManager::Get().GetStatData(*FileDescription.Path);
 
-		FileSize = FileStatData.FileSize;
+	FileSize = FileStatData.FileSize;
 	FDateTime ModificationTime = FileStatData.ModificationTime;
 
 	uint32 FileHash = GetTypeHash(FileDescription);
@@ -248,10 +250,12 @@ uint32 GetKioFaceTessellation(CT_OBJECT_ID FaceID, TArray<FTessellationData>& Fa
 		return 0;
 	}
 
+	int32 PatchId = GetIntegerParameterDataValue(FaceID, TEXT("DatasmithFaceId"));
+
 	FTessellationData& Tessellation = FaceTessellationSet.Emplace_GetRef();
 	Tessellation.VertexIndices.SetNum(IndexCount);
-
-	Tessellation.PatchId = (int32) FaceID;
+	
+	Tessellation.PatchId = PatchId;
 
 	switch (IndexType)
 	{
@@ -692,6 +696,65 @@ FCoreTechFileParser::EProcessResult FCoreTechFileParser::ProcessFile(const FFile
 	return ReadFileWithKernelIO();
 }
 
+// For each face, adds an integer parameter representing the id of the face to avoid re-identation of faces in sub CT file
+// Used by Re-tessellation Rule to Skip Deleted Surfaces
+void AddFaceIdAttribut(CT_OBJECT_ID NodeId)
+{
+	CT_OBJECT_TYPE Type;
+	CT_OBJECT_IO::AskType(NodeId, Type);
+
+	switch (Type)
+	{
+	case CT_INSTANCE_TYPE:
+	{
+		CT_OBJECT_ID ReferenceNodeId;
+		if (CT_INSTANCE_IO::AskChild(NodeId, ReferenceNodeId) == IO_OK)
+		{
+			AddFaceIdAttribut(ReferenceNodeId);
+		}
+		break;
+	}
+
+	case CT_ASSEMBLY_TYPE:
+	case CT_PART_TYPE:
+	case CT_COMPONENT_TYPE:
+	{
+		CT_LIST_IO Children;
+		if (CT_COMPONENT_IO::AskChildren(NodeId, Children) == IO_OK)
+		{
+			Children.IteratorInitialize();
+			CT_OBJECT_ID ChildId;
+			while ((ChildId = Children.IteratorIter()) != 0)
+			{
+				AddFaceIdAttribut(ChildId);
+			}
+		}
+
+		break;
+	}
+
+	case CT_BODY_TYPE:
+	{
+		CT_LIST_IO FaceList;
+		CT_BODY_IO::AskFaces(NodeId, FaceList);
+
+		CT_OBJECT_ID FaceID;
+		FaceList.IteratorInitialize();
+		while ((FaceID = FaceList.IteratorIter()) != 0)
+		{
+			CT_OBJECT_IO::AddAttribute(FaceID, CT_ATTRIB_INTEGER_PARAMETER);
+
+			ensure(CT_CURRENT_ATTRIB_IO::SetStrField(ITH_INTEGER_PARAMETER_NAME, "DatasmithFaceId") == IO_OK);
+			ensure(CT_CURRENT_ATTRIB_IO::SetIntField(ITH_INTEGER_PARAMETER_VALUE, FaceID) == IO_OK);
+		}
+		break;
+	}
+
+	default:
+		break;
+	}
+}
+
 FCoreTechFileParser::EProcessResult FCoreTechFileParser::ReadFileWithKernelIO()
 {
 	CT_IO_ERROR Result = IO_OK;
@@ -702,46 +765,46 @@ FCoreTechFileParser::EProcessResult FCoreTechFileParser::ReadFileWithKernelIO()
 	SceneGraphArchive.FullPath = FileDescription.Path;
 	SceneGraphArchive.CADFileName = FileDescription.Name;
 
-		// the parallelization of monolithic Jt file is set in SetCoreTechImportOption. Then it's processed as the other exploded formats
-		CT_FLAGS CTImportOption = SetCoreTechImportOption();
+	// the parallelization of monolithic Jt file is set in SetCoreTechImportOption. Then it's processed as the other exploded formats
+	CT_FLAGS CTImportOption = SetCoreTechImportOption();
 
 	FString LoadOption;
 	CT_UINT32 NumberOfIds = 1;
 
 	if (!FileDescription.Configuration.IsEmpty())
 	{
-			if (FileDescription.Extension == "jt")
-			{
-				LoadOption = FileDescription.Configuration;
-			}
-			else
-			{
-		NumberOfIds = CT_KERNEL_IO::AskFileNbOfIds(*FileDescription.Path);
-		if (NumberOfIds > 1)
+		if (FileDescription.Extension == "jt")
 		{
-			CT_UINT32 ActiveConfig = CT_KERNEL_IO::AskFileActiveConfig(*FileDescription.Path);
-			for (CT_UINT32 i = 0; i < NumberOfIds; i++)
+			LoadOption = FileDescription.Configuration;
+		}
+		else
+		{
+			NumberOfIds = CT_KERNEL_IO::AskFileNbOfIds(*FileDescription.Path);
+			if (NumberOfIds > 1)
 			{
-				CT_STR ConfValue = CT_KERNEL_IO::AskFileIdIthName(*FileDescription.Path, i);
+				CT_UINT32 ActiveConfig = CT_KERNEL_IO::AskFileActiveConfig(*FileDescription.Path);
+				for (CT_UINT32 i = 0; i < NumberOfIds; i++)
+				{
+					CT_STR ConfValue = CT_KERNEL_IO::AskFileIdIthName(*FileDescription.Path, i);
 				if (FileDescription.Configuration == AsFString(ConfValue)) {
-					ActiveConfig = i;
-					break;
+						ActiveConfig = i;
+						break;
+					}
 				}
-			}
 
-			CTImportOption |= CT_LOAD_FLAGS_READ_SPECIFIC_OBJECT;
-			LoadOption = FString::FromInt((int32) ActiveConfig);
+				CTImportOption |= CT_LOAD_FLAGS_READ_SPECIFIC_OBJECT;
+				LoadOption = FString::FromInt((int32)ActiveConfig);
+			}
 		}
 	}
-		}
 
 	Result = CT_KERNEL_IO::LoadFile(*FileDescription.Path, MainId, CTImportOption, 0, *LoadOption);
 	if (Result == IO_ERROR_EMPTY_ASSEMBLY)
 	{
 		CT_KERNEL_IO::UnloadModel();
-			CT_FLAGS CTReImportOption = CTImportOption | CT_LOAD_FLAGS_LOAD_EXTERNAL_REF;
-			CTReImportOption &= ~CT_LOAD_FLAGS_READ_ASM_STRUCT_ONLY;  // BUG CT -> Ticket 11685
-			Result = CT_KERNEL_IO::LoadFile(*FileDescription.Path, MainId, CTReImportOption, 0, *LoadOption);
+		CT_FLAGS CTReImportOption = CTImportOption | CT_LOAD_FLAGS_LOAD_EXTERNAL_REF;
+		CTReImportOption &= ~CT_LOAD_FLAGS_READ_ASM_STRUCT_ONLY;  // BUG CT -> Ticket 11685
+		Result = CT_KERNEL_IO::LoadFile(*FileDescription.Path, MainId, CTReImportOption, 0, *LoadOption);
 	}
 
 	// the file is loaded but it's empty, so no data is generate
@@ -768,6 +831,12 @@ FCoreTechFileParser::EProcessResult FCoreTechFileParser::ReadFileWithKernelIO()
 		CT_KERNEL_IO::SaveFile(ObjectList, *FPaths::Combine(CachePath, TEXT("cad"), SceneGraphArchive.ArchiveFileName + TEXT(".ct")), L"Ct");
 	}
 
+
+	if (ImportParameters.StitchingTechnique != StitchingNone)
+	{
+		CADLibrary::Repair(MainId, ImportParameters.StitchingTechnique, 10.);
+	}
+
 	SetCoreTechTessellationState(ImportParameters);
 
 	SceneGraphArchive.FullPath = FileDescription.Path;
@@ -1266,7 +1335,7 @@ void FCoreTechFileParser::GetBodyTessellation(CT_OBJECT_ID CTBodyId, CT_OBJECT_I
 void FCoreTechFileParser::GetKioBodyTessellation(CT_OBJECT_ID BodyId, CT_OBJECT_ID ParentId, FBodyMesh& OutBodyMesh, uint32 DefaultMaterialHash, bool bNeedRepair)
 {
 	TArray<CT_OBJECT_ID> BodyFaces;
-	{
+{
 		CT_LIST_IO FaceList;
 		CT_BODY_IO::AskFaces(BodyId, FaceList);
 		BodyFaces.Reserve((int32)(1.3 * FaceList.Count()));
@@ -1316,18 +1385,18 @@ void FCoreTechFileParser::GetKioBodyTessellation(CT_OBJECT_ID BodyId, CT_OBJECT_
 		CT_LIST_IO FaceList;
 		CT_BODY_IO::AskFaces(Body, FaceList);
 
-		// Compute Body BBox based on CAD data
-		uint32 VerticesSize;
+	// Compute Body BBox based on CAD data
+	uint32 VerticesSize;
 		CT_BODY_IO::AskVerticesSizeArray(Body, VerticesSize);
 
-		TArray<CT_COORDINATE> VerticesArray;
-		VerticesArray.SetNum(VerticesSize);
+	TArray<CT_COORDINATE> VerticesArray;
+	VerticesArray.SetNum(VerticesSize);
 		CT_BODY_IO::AskVerticesArray(Body, VerticesArray.GetData());
 
-		for (const CT_COORDINATE& Point : VerticesArray)
-		{
-			BBox += FVector((float)Point.xyz[0], (float)Point.xyz[1], (float)Point.xyz[2]);
-		}
+	for (const CT_COORDINATE& Point : VerticesArray)
+	{
+		BBox += FVector((float)Point.xyz[0], (float)Point.xyz[1], (float)Point.xyz[2]);
+	}
 
 		CT_OBJECT_ID FaceID;
 		FaceList.IteratorInitialize();
@@ -1482,21 +1551,38 @@ void FCoreTechFileParser::GetStringMetaDataValue(CT_OBJECT_ID NodeId, const TCHA
 	{
 		if (CT_CURRENT_ATTRIB_IO::AskStrField(ITH_STRING_METADATA_NAME, FieldName) != IO_OK)
 		{
-			break;
+			continue;
 		}
 		if (!FCString::Strcmp(InMetaDataName, *AsFString(FieldName)))
 		{
 			CT_STR FieldStrValue;
-			if (CT_CURRENT_ATTRIB_IO::AskStrField(ITH_STRING_METADATA_VALUE, FieldStrValue) != IO_OK)
-			{
-				break;
-			}
+			CT_CURRENT_ATTRIB_IO::AskStrField(ITH_STRING_METADATA_VALUE, FieldStrValue);
 			OutMetaDataValue = AsFString(FieldStrValue);
 			return;
 		}
 	}
 }
 
+int32 GetIntegerParameterDataValue(CT_OBJECT_ID NodeId, const TCHAR* InMetaDataName)
+{
+	CT_STR FieldName;
+	CT_UINT32 IthAttrib = 0;
+	int32 IntegerParameterValue = 0;
+	while (CT_OBJECT_IO::SearchAttribute(NodeId, CT_ATTRIB_INTEGER_PARAMETER, IthAttrib++) == IO_OK)
+	{
+		if (CT_CURRENT_ATTRIB_IO::AskStrField(ITH_INTEGER_PARAMETER_NAME, FieldName) != IO_OK)
+		{
+			continue;
+		}
+		if (!FCString::Strcmp(InMetaDataName, *AsFString(FieldName)))
+		{
+			CT_CURRENT_ATTRIB_IO::AskIntField(ITH_INTEGER_PARAMETER_VALUE, IntegerParameterValue);
+			break;
+		}
+	}
+	return IntegerParameterValue;
+}
+
 void FCoreTechFileParser::ReadNodeMetaData(CT_OBJECT_ID NodeId, TMap<FString, FString>& OutMetaData)
 {
 	if (CT_COMPONENT_IO::IsA(NodeId, CT_COMPONENT_TYPE))
