Replace UE4 by UnrealEditor in comment

diff --git a/Engine/Plugins/Enterprise/DatasmithCADImporter/Source/DatasmithOpenNurbsTranslator/Private/RhinoCoretechWrapper.h b/Engine/Plugins/Enterprise/DatasmithCADImporter/Source/DatasmithOpenNurbsTranslator/Private/RhinoCoretechWrapper.h
--- a/Engine/Plugins/Enterprise/DatasmithCADImporter/Source/DatasmithOpenNurbsTranslator/Private/RhinoCoretechWrapper.h
+++ b/Engine/Plugins/Enterprise/DatasmithCADImporter/Source/DatasmithOpenNurbsTranslator/Private/RhinoCoretechWrapper.h
@@ -43,7 +43,7 @@ public:
 	 * Handle input file unit and an output unit
 	 * @param InOwner
 	 * @param FileMetricUnit number of meters per file unit.
-	 * @param ScaleFactor scale factor to apply to the mesh to be in UE4 unit (cm).
+	 * @param ScaleFactor scale factor to apply to the mesh to be in UnrealEditor unit (cm).
 	 * eg. For a file in inches, arg should be 0.0254
 	 */
 	FRhinoCoretechWrapper(const TCHAR* InOwner, double FileMetricUnit, double ScaleFactor)
