Merging CL 15730704 from Main (//UE4/Main) to Main (//UE5/Main)

OpenNurbsImporter - DatasmithAreaLights have wrong transfrom on import.

diff --git a/Engine/Plugins/Enterprise/DatasmithCADImporter/Source/DatasmithOpenNurbsTranslator/Private/DatasmithOpenNurbsTranslator.cpp b/Engine/Plugins/Enterprise/DatasmithCADImporter/Source/DatasmithOpenNurbsTranslator/Private/DatasmithOpenNurbsTranslator.cpp
--- a/Engine/Plugins/Enterprise/DatasmithCADImporter/Source/DatasmithOpenNurbsTranslator/Private/DatasmithOpenNurbsTranslator.cpp
+++ b/Engine/Plugins/Enterprise/DatasmithCADImporter/Source/DatasmithOpenNurbsTranslator/Private/DatasmithOpenNurbsTranslator.cpp
@@ -1339,17 +1339,16 @@ void FOpenNurbsTranslatorImpl::TranslateLightTable(const ON_ClassArray<FOpenNurb
 				AreaLightElement->SetAttenuationRadius(1800.f);
 			}
 
-			ON_3dVector InvLengthAxis = -LightObj.Length();
+			ON_3dVector LengthAxis = LightObj.Length();
 			ON_3dVector WidthAxis = LightObj.Width();
-			ON_3dVector InvLightAxis = ON_CrossProduct(WidthAxis, InvLengthAxis);
-			ON_Xform Xform(Center, InvLightAxis.UnitVector(), WidthAxis.UnitVector(), InvLengthAxis.UnitVector());
+			ON_3dVector InvLightAxis = ON_CrossProduct(WidthAxis, LengthAxis);
+			ON_Xform Xform(Center, InvLightAxis.UnitVector(), WidthAxis.UnitVector(), LengthAxis.UnitVector());
 
 			FMatrix Matrix;
 			DatasmithOpenNurbsTranslatorUtils::XFormToMatrix(Xform, Matrix);
 
 			FTransform Transform(Matrix);
-			const FTransform RightHanded = FTransform(FRotator(0.0f, 0.0f, 0.0f), FVector(0.0f, 0.0f, 0.0f), FVector(-1.0f, 1.0f, 1.0f));
-			FTransform CorrectedTransform = RightHanded * Transform * RightHanded;
+			FTransform CorrectedTransform = FDatasmithUtils::ConvertTransform(FDatasmithUtils::EModelCoordSystem::ZUp_RightHanded_FBXLegacy, Transform);
 
 			AreaLightElement->SetTranslation(CorrectedTransform.GetTranslation() * ScalingFactor);
 			AreaLightElement->SetScale(CorrectedTransform.GetScale3D());
