Replace UE4 by UE in variable names

diff --git a/Engine/Plugins/Enterprise/DatasmithCADImporter/Source/DatasmithCADTranslator/Private/DatasmithCADTranslator.cpp b/Engine/Plugins/Enterprise/DatasmithCADImporter/Source/DatasmithCADTranslator/Private/DatasmithCADTranslator.cpp
--- a/Engine/Plugins/Enterprise/DatasmithCADImporter/Source/DatasmithCADTranslator/Private/DatasmithCADTranslator.cpp
+++ b/Engine/Plugins/Enterprise/DatasmithCADImporter/Source/DatasmithCADTranslator/Private/DatasmithCADTranslator.cpp
@@ -125,10 +125,10 @@ bool FDatasmithCADTranslator::LoadScene(TSharedRef<IDatasmithScene> DatasmithSce
 		ImportParameters.bEnableCacheUsage = true;
 	}
 
-	bool bWithProcessor = (CVarStaticCADTranslatorEnableThreadedImport.GetValueOnAnyThread() != 0);
+		bool bWithProcessor = (CVarStaticCADTranslatorEnableThreadedImport.GetValueOnAnyThread() != 0);
 
 #ifdef CAD_TRANSLATOR_DEBUG
-	bWithProcessor = false;
+		bWithProcessor = false;
 #endif //CAD_TRANSLATOR_DEBUG
 
 	// Only use multi-processed translation if it is required and cache's usage is enabled
@@ -140,16 +140,16 @@ bool FDatasmithCADTranslator::LoadScene(TSharedRef<IDatasmithScene> DatasmithSce
 			DatasmithDispatcher::FDatasmithDispatcher Dispatcher(ImportParameters, CachePath, NumCores, CADFileToUE4FileMap, CADFileToUE4GeomMap);
 			Dispatcher.AddTask(FileDescription);
 
-			Dispatcher.Process(bWithProcessor);
-		}
+		Dispatcher.Process(bWithProcessor);
+	}
 
-		FDatasmithSceneGraphBuilder SceneGraphBuilder(CADFileToUE4FileMap, CachePath, DatasmithScene, GetSource(), ImportParameters);
-		SceneGraphBuilder.Build();
+	FDatasmithSceneGraphBuilder SceneGraphBuilder(CADFileToUEFileMap, CachePath, DatasmithScene, GetSource(), ImportParameters);
+	SceneGraphBuilder.Build();
 
-		MeshBuilderPtr = MakeUnique<FDatasmithMeshBuilder>(CADFileToUE4GeomMap, CachePath, ImportParameters);
+	MeshBuilderPtr = MakeUnique<FDatasmithMeshBuilder>(CADFileToUEGeomMap, CachePath, ImportParameters);
 
-		return true;
-	}
+	return true;
+}
 
 	ImportParameters.bEnableCacheUsage = false;
 
@@ -171,7 +171,7 @@ void FDatasmithCADTranslator::UnloadScene()
 {
 	MeshBuilderPtr = nullptr;
 
-	CADFileToUE4GeomMap.Empty();
+	CADFileToUEGeomMap.Empty();
 }
 
 bool FDatasmithCADTranslator::LoadStaticMesh(const TSharedRef<IDatasmithMeshElement> MeshElement, FDatasmithMeshElementPayload& OutMeshPayload)
diff --git a/Engine/Plugins/Enterprise/DatasmithCADImporter/Source/DatasmithCADTranslator/Private/DatasmithCADTranslator.h b/Engine/Plugins/Enterprise/DatasmithCADImporter/Source/DatasmithCADTranslator/Private/DatasmithCADTranslator.h
--- a/Engine/Plugins/Enterprise/DatasmithCADImporter/Source/DatasmithCADTranslator/Private/DatasmithCADTranslator.h
+++ b/Engine/Plugins/Enterprise/DatasmithCADImporter/Source/DatasmithCADTranslator/Private/DatasmithCADTranslator.h
@@ -31,7 +31,7 @@ public:
 	virtual void SetSceneImportOptions(TArray<TStrongObjectPtr<UDatasmithOptionsBase>>& Options) override;
 
 private:
-	TMap<uint32, FString> CADFileToUE4GeomMap;
+	TMap<uint32, FString> CADFileToUEGeomMap;
 
 	CADLibrary::FImportParameters ImportParameters;
 
