Fix compilation

diff --git a/Engine/Plugins/Enterprise/DatasmithCADImporter/Source/DatasmithCADTranslator/Private/DatasmithCADTranslator.cpp b/Engine/Plugins/Enterprise/DatasmithCADImporter/Source/DatasmithCADTranslator/Private/DatasmithCADTranslator.cpp
--- a/Engine/Plugins/Enterprise/DatasmithCADImporter/Source/DatasmithCADTranslator/Private/DatasmithCADTranslator.cpp
+++ b/Engine/Plugins/Enterprise/DatasmithCADImporter/Source/DatasmithCADTranslator/Private/DatasmithCADTranslator.cpp
@@ -134,10 +134,10 @@ bool FDatasmithCADTranslator::LoadScene(TSharedRef<IDatasmithScene> DatasmithSce
 	// Only use multi-processed translation if it is required and cache's usage is enabled
 	if (bWithProcessor && ImportParameters.bEnableCacheUsage)
 	{
-		TMap<uint32, FString> CADFileToUE4FileMap;
+		TMap<uint32, FString> CADFileToUEFileMap;
 		int32 NumCores = FPlatformMisc::NumberOfCores();
 		{
-			DatasmithDispatcher::FDatasmithDispatcher Dispatcher(ImportParameters, CachePath, NumCores, CADFileToUE4FileMap, CADFileToUE4GeomMap);
+			DatasmithDispatcher::FDatasmithDispatcher Dispatcher(ImportParameters, CachePath, NumCores, CADFileToUEFileMap, CADFileToUEGeomMap);
 			Dispatcher.AddTask(FileDescription);
 
 		Dispatcher.Process(bWithProcessor);
diff --git a/Engine/Plugins/Enterprise/DatasmithCADImporter/Source/DatasmithPLMXMLTranslator/Private/DatasmithPlmXmlImporter.cpp b/Engine/Plugins/Enterprise/DatasmithCADImporter/Source/DatasmithPLMXMLTranslator/Private/DatasmithPlmXmlImporter.cpp
--- a/Engine/Plugins/Enterprise/DatasmithCADImporter/Source/DatasmithPLMXMLTranslator/Private/DatasmithPlmXmlImporter.cpp
+++ b/Engine/Plugins/Enterprise/DatasmithCADImporter/Source/DatasmithPLMXMLTranslator/Private/DatasmithPlmXmlImporter.cpp
@@ -228,8 +228,8 @@ namespace PlmXml
 		TUniquePtr<FDatasmithMeshBuilder> MeshBuilderPtr;
 
 		CADLibrary::FImportParameters ImportParameters;
-		TMap<uint32, FString> CADFileToUE4GeomMap;
-		TMap<uint32, FString> CADFileToUE4FileMap;
+		TMap<uint32, FString> CADFileToUEGeomMap;
+		TMap<uint32, FString> CADFileToUEFileMap;
 
 		TArray<FString> FilePaths;
 
@@ -253,7 +253,7 @@ namespace PlmXml
 			ImportParameters.MaxNormalAngle = TessellationOptions.NormalTolerance;
 			ImportParameters.StitchingTechnique = (CADLibrary::EStitchingTechnique) TessellationOptions.StitchingTechnique;
 
-			DatasmithDispatcher = MakeUnique<DatasmithDispatcher::FDatasmithDispatcher>(ImportParameters, CacheDir, FPlatformMisc::NumberOfCores(), CADFileToUE4FileMap, CADFileToUE4GeomMap);
+			DatasmithDispatcher = MakeUnique<DatasmithDispatcher::FDatasmithDispatcher>(ImportParameters, CacheDir, FPlatformMisc::NumberOfCores(), CADFileToUEFileMap, CADFileToUEGeomMap);
 		}
 
 		// Adds geom file to load and returns Id to use in InstantiateMesh later(after all is loaded)
@@ -284,9 +284,9 @@ namespace PlmXml
 		void Process(const FDatasmithSceneSource& Source)
 		{
 			DatasmithDispatcher->Process(true);
-			SceneGraphBuilder = MakeUnique<FDatasmithSceneGraphBuilder>(CADFileToUE4FileMap, CacheDir, DatasmithScene, Source, ImportParameters);
+			SceneGraphBuilder = MakeUnique<FDatasmithSceneGraphBuilder>(CADFileToUEFileMap, CacheDir, DatasmithScene, Source, ImportParameters);
 			SceneGraphBuilder->LoadSceneGraphDescriptionFiles();
-			MeshBuilderPtr = MakeUnique<FDatasmithMeshBuilder>(CADFileToUE4GeomMap, CacheDir, ImportParameters);
+			MeshBuilderPtr = MakeUnique<FDatasmithMeshBuilder>(CADFileToUEGeomMap, CacheDir, ImportParameters);
 		}
 
 		void UnloadScene()
