Rename of DatasmithCoreTechExtension module into CoreTechExtension is the root cause.
The new class's names of this module have been added in CoreRedirects group of BaseDatasmithCadImporter.ini.

diff --git a/Engine/Plugins/Enterprise/DatasmithCADImporter/Config/BaseDatasmithCADImporter.ini b/Engine/Plugins/Enterprise/DatasmithCADImporter/Config/BaseDatasmithCADImporter.ini
new file mode 100644
--- /dev/null
+++ b/Engine/Plugins/Enterprise/DatasmithCADImporter/Config/BaseDatasmithCADImporter.ini
@@ -0,0 +1,8 @@
+[CoreRedirects]
+
++ClassRedirects=(OldName="/Script/DatasmithCoreTechExtension.CoreTechBlueprintLibrary", NewName="/Script/CoreTechExtension.CoreTechBlueprintLibrary")
++ClassRedirects=(OldName="/Script/DatasmithCoreTechExtension.CoreTechRetessellateAction", NewName="/Script/CoreTechExtension.CoreTechRetessellateAction")
++ClassRedirects=(OldName="/Script/DatasmithCoreTechExtension.CoreTechRetessellateActionOptions", NewName="/Script/CoreTechExtension.CoreTechRetessellateActionOptions")
++ClassRedirects=(OldName="/Script/DatasmithCoreTechExtension.DataprepTessellationOperation", NewName="/Script/CoreTechExtension.DataprepTessellationOperation")
+
+
