Fix the transfer of FImportParameters to CADWorker:
The implementation of FImportParameters archiving to its transfer to CADWorker was done in DatasmtihCommands.
The modification of FImportParameters was unknown by DatasmithCommands.

diff --git a/Engine/Plugins/Enterprise/DatasmithCADImporter/Source/CADTools/Public/CADOptions.h b/Engine/Plugins/Enterprise/DatasmithCADImporter/Source/CADTools/Public/CADOptions.h
--- a/Engine/Plugins/Enterprise/DatasmithCADImporter/Source/CADTools/Public/CADOptions.h
+++ b/Engine/Plugins/Enterprise/DatasmithCADImporter/Source/CADTools/Public/CADOptions.h
@@ -43,6 +43,7 @@ namespace CADLibrary
 		EDisplayPreference DisplayPreference = EDisplayPreference::MaterialPrefered;
 		bool bScaleUVMap = true;
 		bool bEnableCacheUsage = true;
+		bool bEnableKernelIOTessellation = true;
 
 		uint32 GetHash() const
 		{
@@ -58,6 +59,24 @@ namespace CADLibrary
 			Hash = HashCombine(Hash, GetTypeHash(bScaleUVMap));
 			return Hash;
 		}
+
+		friend FArchive& operator<<(FArchive& Ar, FImportParameters& ImportParameters)
+		{
+			Ar << ImportParameters.MetricUnit;
+			Ar << ImportParameters.ScaleFactor;
+			Ar << ImportParameters.ChordTolerance;
+			Ar << ImportParameters.MaxEdgeLength;
+			Ar << ImportParameters.MaxNormalAngle;
+			Ar << (uint32&) ImportParameters.StitchingTechnique;
+			Ar << (uint8&) ImportParameters.ModelCoordSys;
+			Ar << (uint8&) ImportParameters.Propagation;
+			Ar << (uint8&) ImportParameters.DisplayPreference;
+			Ar << ImportParameters.bScaleUVMap;
+			Ar << ImportParameters.bEnableCacheUsage;
+			Ar << ImportParameters.bEnableKernelIOTessellation;
+			return Ar;
+		}
+
 	};
 
 	struct FMeshParameters
diff --git a/Engine/Plugins/Enterprise/DatasmithCADImporter/Source/DatasmithDispatcher/Private/DatasmithCommands.cpp b/Engine/Plugins/Enterprise/DatasmithCADImporter/Source/DatasmithDispatcher/Private/DatasmithCommands.cpp
--- a/Engine/Plugins/Enterprise/DatasmithCADImporter/Source/DatasmithDispatcher/Private/DatasmithCommands.cpp
+++ b/Engine/Plugins/Enterprise/DatasmithCADImporter/Source/DatasmithDispatcher/Private/DatasmithCommands.cpp
@@ -61,12 +61,7 @@ namespace DatasmithDispatcher
 
 	void FImportParametersCommand::SerializeImpl(FArchive& Ar)
 	{
-		Ar << ImportParameters.ChordTolerance;
-		Ar << ImportParameters.MaxEdgeLength;
-		Ar << ImportParameters.MaxNormalAngle;
-		Ar << ImportParameters.MetricUnit;
-		Ar << ImportParameters.ScaleFactor;
-		Ar << (uint32&)ImportParameters.StitchingTechnique;
+		Ar << ImportParameters;
 	}
 
 }
