Preparation of CADKernel integration

diff --git a/Engine/Plugins/Enterprise/DatasmithCADImporter/Source/CADInterfaces/CADInterfaces.Build.cs b/Engine/Plugins/Enterprise/DatasmithCADImporter/Source/CADInterfaces/CADInterfaces.Build.cs
--- a/Engine/Plugins/Enterprise/DatasmithCADImporter/Source/CADInterfaces/CADInterfaces.Build.cs
+++ b/Engine/Plugins/Enterprise/DatasmithCADImporter/Source/CADInterfaces/CADInterfaces.Build.cs
@@ -14,6 +14,7 @@ namespace UnrealBuildTool.Rules
 				new string[]
 				{
 					"Core",
+					"CADKernel",
 					"CADTools",
 					"DatasmithCore"
 				}
diff --git a/Engine/Plugins/Enterprise/DatasmithCADImporter/Source/CADInterfaces/Private/CoreTechFileReader.cpp b/Engine/Plugins/Enterprise/DatasmithCADImporter/Source/CADInterfaces/Private/CoreTechFileReader.cpp
--- a/Engine/Plugins/Enterprise/DatasmithCADImporter/Source/CADInterfaces/Private/CoreTechFileReader.cpp
+++ b/Engine/Plugins/Enterprise/DatasmithCADImporter/Source/CADInterfaces/Private/CoreTechFileReader.cpp
@@ -21,7 +21,6 @@
 
 #include "CADData.h"
 #include "CADOptions.h"
-
 #include "DatasmithUtils.h"
 #include "HAL/FileManager.h"
 #include "Internationalization/Text.h"
@@ -329,7 +328,7 @@ namespace CADLibrary
 
 		CoreTechFileReaderUtils::AddFaceIdAttribut(MainId);
 
-		if (Context.ImportParameters.StitchingTechnique != StitchingNone)
+		if (Context.ImportParameters.StitchingTechnique != StitchingNone && Context.ImportParameters.bEnableKernelIOTessellation)
 		{
 			CADLibrary::CTKIO_Repair(MainId, Context.ImportParameters.StitchingTechnique, 10.);
 		}
@@ -536,9 +535,22 @@ namespace CADLibrary
 
 		for (CT_OBJECT_ID BodyId : Bodies)
 		{
-			if (ReadBody(BodyId, ComponentId, DefaultMaterialHash, false))
+			CT_FLAGS BodyProperties;
+			CT_BODY_IO::AskProperties(BodyId, BodyProperties);
+
+			if (Context.ImportParameters.bEnableKernelIOTessellation || !(BodyProperties & CT_BODY_PROP_EXACT))
+			{
+				if (ReadKioBody(BodyId, ComponentId, DefaultMaterialHash, false))
+				{
+					Context.SceneGraphArchive.ComponentSet[Index].Children.Add(BodyId);
+				}
+			}
+			else
 			{
-				Context.SceneGraphArchive.ComponentSet[Index].Children.Add(BodyId);
+				//if (ReadBody(BodyId, ComponentId, DefaultMaterialHash, false))
+				//{
+				//	Context.SceneGraphArchive.ComponentSet[Index].Children.Add(BodyId);
+				//}
 			}
 		}
 
@@ -630,7 +642,7 @@ namespace CADLibrary
 		return ReadNode(ReferenceNodeId, DefaultMaterialHash);
 	}
 
-	bool FCoreTechFileReader::ReadBody(CT_OBJECT_ID BodyId, CT_OBJECT_ID ParentId, uint32 DefaultMaterialHash, bool bNeedRepair)
+	bool FCoreTechFileReader::ReadKioBody(CT_OBJECT_ID BodyId, CT_OBJECT_ID ParentId, uint32 DefaultMaterialHash, bool bNeedRepair)
 	{
 		// Is this body a constructive geometry ?
 		CT_LIST_IO FaceList;
@@ -1273,7 +1285,6 @@ namespace CADLibrary
 		}
 	}
 
-
 	namespace CoreTechFileReaderUtils
 	{
 		template<typename ValueType>
@@ -1467,29 +1478,30 @@ namespace CADLibrary
 			}
 
 			Tessellation.PatchId = CoreTechFileReaderUtils::GetIntegerParameterDataValue(FaceID, TEXT("DatasmithFaceId"));
-			Tessellation.IndexArray.SetNum(IndexCount);
+
+			Tessellation.VertexIndices.SetNum(IndexCount);
 
 			switch (IndexType)
 			{
 			case CT_TESS_UBYTE:
-				FillArrayOfInt<uint8>(IndexCount, IndexArray, Tessellation.IndexArray.GetData());
+				FillArrayOfInt<uint8>(IndexCount, IndexArray, Tessellation.VertexIndices.GetData());
 				break;
 			case CT_TESS_USHORT:
-				FillArrayOfInt<uint16>(IndexCount, IndexArray, Tessellation.IndexArray.GetData());
+				FillArrayOfInt<uint16>(IndexCount, IndexArray, Tessellation.VertexIndices.GetData());
 				break;
 			case CT_TESS_UINT:
-				FillArrayOfInt<uint32>(IndexCount, IndexArray, Tessellation.IndexArray.GetData());
+				FillArrayOfInt<uint32>(IndexCount, IndexArray, Tessellation.VertexIndices.GetData());
 				break;
 			}
 
-			Tessellation.VertexArray.SetNum(VertexCount);
+			Tessellation.PositionArray.SetNum(VertexCount);
 			switch (VertexType)
 			{
 			case CT_TESS_FLOAT:
-				FillArrayOfVector<float>(VertexCount, VertexArray, Tessellation.VertexArray.GetData());
+				FillArrayOfVector<float>(VertexCount, VertexArray, Tessellation.PositionArray.GetData());
 				break;
 			case CT_TESS_DOUBLE:
-				FillArrayOfVector<double>(VertexCount, VertexArray, Tessellation.VertexArray.GetData());
+				FillArrayOfVector<double>(VertexCount, VertexArray, Tessellation.PositionArray.GetData());
 				break;
 			}
 
@@ -1536,7 +1548,7 @@ namespace CADLibrary
 				}
 			}
 
-			return Tessellation.IndexArray.Num() / 3;
+			return Tessellation.VertexIndices.Num() / 3;
 		}
 
 
@@ -1868,6 +1880,7 @@ namespace CADLibrary
 				FaceIndex++;
 			}
 		}
+
 	}
 }
 
diff --git a/Engine/Plugins/Enterprise/DatasmithCADImporter/Source/CADInterfaces/Private/CoreTechFileReader.h b/Engine/Plugins/Enterprise/DatasmithCADImporter/Source/CADInterfaces/Private/CoreTechFileReader.h
--- a/Engine/Plugins/Enterprise/DatasmithCADImporter/Source/CADInterfaces/Private/CoreTechFileReader.h
+++ b/Engine/Plugins/Enterprise/DatasmithCADImporter/Source/CADInterfaces/Private/CoreTechFileReader.h
@@ -81,7 +81,8 @@ namespace CADLibrary
 		bool ReadInstance(CT_OBJECT_ID NodeId, uint32 ParentMaterialHash);
 		bool ReadComponent(CT_OBJECT_ID NodeId, uint32 ParentMaterialHash);
 		bool ReadUnloadedComponent(CT_OBJECT_ID NodeId);
-		bool ReadBody(CT_OBJECT_ID NodeId, CT_OBJECT_ID ParentId, uint32 ParentMaterialHash, bool bNeedRepair);
+
+		bool ReadKioBody(CT_OBJECT_ID NodeId, CT_OBJECT_ID ParentId, uint32 ParentMaterialHash, bool bNeedRepair);
 
 		bool FindFile(FFileDescription& FileDescription);
 
diff --git a/Engine/Plugins/Enterprise/DatasmithCADImporter/Source/CADInterfaces/Public/CADSceneGraph.h b/Engine/Plugins/Enterprise/DatasmithCADImporter/Source/CADInterfaces/Public/CADSceneGraph.h
--- a/Engine/Plugins/Enterprise/DatasmithCADImporter/Source/CADInterfaces/Public/CADSceneGraph.h
+++ b/Engine/Plugins/Enterprise/DatasmithCADImporter/Source/CADInterfaces/Public/CADSceneGraph.h
@@ -13,7 +13,7 @@ namespace CADLibrary
 class CADINTERFACES_API ICADArchiveObject
 {
 public:
-	ICADArchiveObject(CadId Id = 0)
+	ICADArchiveObject(FCadId Id = 0)
 		: ObjectId(Id)
 	{
 	}
@@ -28,7 +28,7 @@ public:
 class CADINTERFACES_API FArchiveInstance : public ICADArchiveObject
 {
 public:
-	FArchiveInstance(CadId Id = 0)
+	FArchiveInstance(FCadId Id = 0)
 		: ICADArchiveObject(Id)
 	{
 	}
@@ -37,7 +37,7 @@ public:
 
 public:
 	FMatrix TransformMatrix = FMatrix::Identity;
-	CadId ReferenceNodeId = 0;
+	FCadId ReferenceNodeId = 0;
 	bool bIsExternalRef = false;
 	FFileDescription ExternalRef;
 };
@@ -45,7 +45,7 @@ public:
 class CADINTERFACES_API FArchiveComponent : public ICADArchiveObject
 {
 public:
-	FArchiveComponent(CadId Id = 0)
+	FArchiveComponent(FCadId Id = 0)
 		: ICADArchiveObject(Id)
 	{
 	}
@@ -53,13 +53,13 @@ public:
 	friend FArchive& operator<<(FArchive& Ar, FArchiveComponent& C);
 
 public:
-	TArray<CadId> Children;
+	TArray<FCadId> Children;
 };
 
 class CADINTERFACES_API FArchiveUnloadedComponent : public FArchiveComponent
 {
 public:
-	FArchiveUnloadedComponent(CadId Id = 0)
+	FArchiveUnloadedComponent(FCadId Id = 0)
 		: FArchiveComponent(Id)
 	{
 	}
@@ -74,7 +74,7 @@ public:
 class CADINTERFACES_API FArchiveBody : public ICADArchiveObject
 {
 public:
-	FArchiveBody(CadId Id = 0)
+	FArchiveBody(FCadId Id = 0)
 		: ICADArchiveObject(Id)
 	{
 	}
@@ -82,17 +82,17 @@ public:
 	friend FArchive& operator<<(FArchive& Ar, FArchiveBody& C);
 
 public:
-	CADUUID MeshActorName;
+	FCADUUID MeshActorName;
 
-	TSet<MaterialId> MaterialFaceSet;
-	TSet<ColorId> ColorFaceSet;
+	TSet<FMaterialId> MaterialFaceSet;
+	TSet<FColorId> ColorFaceSet;
 
 };
 
 class CADINTERFACES_API FArchiveColor
 {
 public:
-	FArchiveColor(ColorId Id = 0)
+	FArchiveColor(FColorId Id = 0)
 		: ObjectId(Id)
 	{
 	}
@@ -100,15 +100,15 @@ public:
 	friend FArchive& operator<<(FArchive& Ar, FArchiveColor& C);
 
 public:
-	ColorId ObjectId;
+	FColorId ObjectId;
 	FColor Color;
-	CADUUID UEMaterialName;
+	FCADUUID UEMaterialName;
 };
 
 class CADINTERFACES_API FArchiveMaterial
 {
 public:
-	FArchiveMaterial(MaterialId Id = 0)
+	FArchiveMaterial(FMaterialId Id = 0)
 		: ObjectId(Id)
 	{
 	}
@@ -116,8 +116,8 @@ public:
 	friend FArchive& operator<<(FArchive& Ar, FArchiveMaterial& C);
 
 public:
-	MaterialId ObjectId;
-	CADUUID UEMaterialName;
+	FMaterialId ObjectId;
+	FCADUUID UEMaterialName;
 	FCADMaterial Material;
 };
 
@@ -135,18 +135,18 @@ public:
 	FString FullPath;
 	TSet<FFileDescription> ExternalRefSet;
 
-	TMap<ColorId, FArchiveColor> ColorHIdToColor;
-	TMap<MaterialId, FArchiveMaterial> MaterialHIdToMaterial;
+	TMap<FColorId, FArchiveColor> ColorHIdToColor;
+	TMap<FMaterialId, FArchiveMaterial> MaterialHIdToMaterial;
 
 	TArray<FArchiveBody> BodySet;
 	TArray<FArchiveComponent> ComponentSet;
 	TArray<FArchiveUnloadedComponent> UnloadedComponentSet;
 	TArray<FArchiveInstance> Instances;
 
-	TMap<CadId, int32> CADIdToBodyIndex;
-	TMap<CadId, int32> CADIdToComponentIndex;
-	TMap<CadId, int32> CADIdToUnloadedComponentIndex;
-	TMap<CadId, int32> CADIdToInstanceIndex;
+	TMap<FCadId, int32> CADIdToBodyIndex;
+	TMap<FCadId, int32> CADIdToComponentIndex;
+	TMap<FCadId, int32> CADIdToUnloadedComponentIndex;
+	TMap<FCadId, int32> CADIdToInstanceIndex;
 };
 
 
diff --git a/Engine/Plugins/Enterprise/DatasmithCADImporter/Source/CADLibrary/Private/CoreTechHelper.cpp b/Engine/Plugins/Enterprise/DatasmithCADImporter/Source/CADLibrary/Private/CoreTechHelper.cpp
--- a/Engine/Plugins/Enterprise/DatasmithCADImporter/Source/CADLibrary/Private/CoreTechHelper.cpp
+++ b/Engine/Plugins/Enterprise/DatasmithCADImporter/Source/CADLibrary/Private/CoreTechHelper.cpp
@@ -21,8 +21,6 @@ typedef uint32 TriangleIndex[3];
 
 namespace CADLibrary
 {
-	// Ref. GPureMeshInterface
-
 	struct FVertexData
 	{
 		float Z;
@@ -38,6 +36,7 @@ namespace CADLibrary
 		/** Initialization constructor. */
 		FVertexData(int32 InIndex, const FVector& V)
 		{
+			//Z = V.X + V.Y + V.Z;
 			Z = 0.30f * V.X + 0.33f * V.Y + 0.37f * V.Z;
 			Index = InIndex;
 			Coordinates = V;
@@ -67,7 +66,7 @@ namespace CADLibrary
 		return (VertexIDs[0] == VertexIDs[1] || VertexIDs[0] == VertexIDs[2] || VertexIDs[1] == VertexIDs[2]);
 	}
 
-	void FillVertexPosition(const FImportParameters& ImportParams, const FMeshParameters& MeshParameters, FBodyMesh& Body, FMeshDescription& MeshDescription)
+	void FillKioVertexPosition(const FImportParameters& ImportParams, const FMeshParameters& MeshParameters, FBodyMesh& Body, FMeshDescription& MeshDescription)
 	{
 		int32 TriangleCount = Body.TriangleCount;
 		TArray<FTessellationData>& FaceTessellationSet = Body.Faces;
@@ -79,18 +78,17 @@ namespace CADLibrary
 		BBox.Min *= ImportParams.ScaleFactor;
 		BBox.Max *= ImportParams.ScaleFactor;
 
-		TVertexAttributesRef<FVector> VertexPositions = MeshDescription.VertexAttributes().GetAttributesRef<FVector>(MeshAttribute::Vertex::Position);
+		TVertexAttributesRef<FVector> VertexPositions = MeshDescription.GetVertexPositions();
 
 		// Create a list of vertex Z/index pairs
 		TArray<FVertexData> VertexDataSet;
 		VertexDataSet.Reserve(TriangleCount * 3);
 
 		FVector Position;
-		uint32 GlobalVertexCount = 0;
+		int32 GlobalVertexCount = 0;
 		for (FTessellationData& CTTessellation : FaceTessellationSet)
 		{
-			CTTessellation.StartVertexIndex = GlobalVertexCount;
-			for (const FVector& Vertex : CTTessellation.VertexArray)
+			for (const FVector& Vertex : CTTessellation.PositionArray)
 			{
 				VertexDataSet.Emplace(GlobalVertexCount, Vertex * ImportParams.ScaleFactor);
 				++GlobalVertexCount;
@@ -125,7 +123,7 @@ namespace CADLibrary
 			const FVector& PositionA = VertexDataSet[i].Coordinates;
 			if (BBox.IsValid && !BBox.IsInside(PositionA))
 			{
-				VertexDataSet[i].Index = -1;
+				VertexDataSet[i].Index = INDEX_NONE;
 				continue;
 			}
 
@@ -147,7 +145,6 @@ namespace CADLibrary
 			VertexCount++;
 		}
 
-
 		// if Symmetric mesh, the symmetric side of the mesh have to be generated
 		FMatrix SymmetricMatrix;
 		bool bIsSymmetricMesh = MeshParameters.bIsSymmetric;
@@ -156,12 +153,13 @@ namespace CADLibrary
 			SymmetricMatrix = FDatasmithUtils::GetSymmetricMatrix(MeshParameters.SymmetricOrigin, MeshParameters.SymmetricNormal);
 		}
 
-
 		// Make MeshDescription.VertexPositions and VertexID
-		MeshDescription.ReserveNewVertices(VertexCount);
-		uint32 GlobalVertexIndex = 0;
+		MeshDescription.ReserveNewVertices(bIsSymmetricMesh ? VertexCount * 2 : VertexCount);
+		Body.VertexIds.Reserve(VertexCount);
+
+		int32 GlobalVertexIndex = 0;
 
-		for (uint32 VertexIndex = 0; VertexIndex < GlobalVertexCount; ++VertexIndex)
+		for (int32 VertexIndex = 0; VertexIndex < GlobalVertexCount; ++VertexIndex)
 		{
 			int32 RealIndex = VertexDataSet[VertexIndex].Index;
 
@@ -177,13 +175,15 @@ namespace CADLibrary
 			}
 
 			FVertexID VertexID = MeshDescription.CreateVertex();
-			VertexPositions[VertexID] = FDatasmithUtils::ConvertVector(ImportParams.ModelCoordSys, VertexDataSet[VertexIndex].Coordinates);
+			Body.VertexIds.Add(VertexID);
+			VertexPositions[VertexID] = FDatasmithUtils::ConvertVector((FDatasmithUtils::EModelCoordSystem) ImportParams.ModelCoordSys, VertexDataSet[VertexIndex].Coordinates);
 			VertexDataSet[VertexIndex].VertexID = VertexID;
 		}
 
 		if (bIsSymmetricMesh)
 		{
-			for (uint32 VertexIndex = 0; VertexIndex < GlobalVertexCount; ++VertexIndex)
+			Body.SymmetricVertexIds.Reserve(Body.VertexIds.Num());
+			for (int32 VertexIndex = 0; VertexIndex < GlobalVertexCount; ++VertexIndex)
 			{
 				int32 RealIndex = VertexDataSet[VertexIndex].Index;
 
@@ -199,7 +199,8 @@ namespace CADLibrary
 				}
 
 				FVertexID VertexID = MeshDescription.CreateVertex();
-				VertexPositions[VertexID] = FDatasmithUtils::ConvertVector(ImportParams.ModelCoordSys, VertexDataSet[VertexIndex].Coordinates);
+				Body.SymmetricVertexIds.Add(VertexID);
+				VertexPositions[VertexID] = FDatasmithUtils::ConvertVector((FDatasmithUtils::EModelCoordSystem) ImportParams.ModelCoordSys, VertexDataSet[VertexIndex].Coordinates);
 				VertexPositions[VertexID] = SymmetricMatrix.TransformPosition(VertexPositions[VertexID]);
 
 				VertexDataSet[VertexIndex].SymVertexID = VertexID;
@@ -210,31 +211,17 @@ namespace CADLibrary
 		GlobalVertexIndex = 0;
 		for (FTessellationData& CTTessellation : FaceTessellationSet)
 		{
-			CTTessellation.VertexIdSet.SetNum(CTTessellation.VertexArray.Num());
-			for (int32 VertexIndex = 0; VertexIndex < CTTessellation.VertexArray.Num(); ++VertexIndex, ++GlobalVertexIndex)
+			CTTessellation.PositionIndices.SetNum(CTTessellation.PositionArray.Num());
+			for (int32 VertexIndex = 0; VertexIndex < CTTessellation.PositionArray.Num(); ++VertexIndex, ++GlobalVertexIndex)
 			{
 				int32 NewIndex = NewIndexOf[IndexOfCoincidentNode[GlobalVertexIndex]];
-				CTTessellation.VertexIdSet[VertexIndex] = VertexDataSet[NewIndex].VertexID.GetValue();
-			}
-		}
-
-		if (bIsSymmetricMesh)
-		{
-
-			GlobalVertexIndex = 0;
-			for (FTessellationData& CTTessellation : FaceTessellationSet)
-			{
-				CTTessellation.SymVertexIdSet.SetNum(CTTessellation.VertexArray.Num());
-				for (int32 VertexIndex = 0; VertexIndex < CTTessellation.VertexArray.Num(); ++VertexIndex, ++GlobalVertexIndex)
-				{
-					uint32 NewIndex = NewIndexOf[IndexOfCoincidentNode[GlobalVertexIndex]];
-					CTTessellation.SymVertexIdSet[VertexIndex] = VertexDataSet[NewIndex].SymVertexID.GetValue();
-				}
+				FVertexData& Vertex = VertexDataSet[NewIndex];
+				int32 Index = Vertex.VertexID.GetValue();
+				CTTessellation.PositionIndices[VertexIndex] = VertexDataSet[NewIndex].VertexID.GetValue();
 			}
 		}
 	}
 
-
 	// PolygonAttributes name used into modeling tools (ExtendedMeshAttribute::PolyTriGroups)
 	const FName PolyTriGroups("PolyTriGroups");
 
@@ -251,13 +238,13 @@ namespace CADLibrary
 		return PatchGroups;
 	}
 
-	 /**
-	  * Polygon group is an attribute of polygons.As long as the mesh description is empty(no polygon), polygon group cannot be defined.
-	  * The work around used is to create PolygonGroups and to set them a PolygonGroupAttributes.
-	  * To get the existing polygon groups, for each created polygonGroups the PolygonGroupId is got(see GetExistingPatches)
-	  * Warning: CopyPatchGroups is call in FCoreTechRetessellate_Impl::ApplyOnOneAsset(CoreTechRetessellateAction.cpp) only if the option 
-	  * RetessellateOptions.RetessellationRule equal EDatasmithCADRetessellationRule::SkipDeletedSurfaces
-	  */
+	/**
+	 * Polygon group is an attribute of polygons.As long as the mesh description is empty(no polygon), polygon group cannot be defined.
+	 * The work around used is to create PolygonGroups and to set them a PolygonGroupAttributes.
+	 * To get the existing polygon groups, for each created polygonGroups the PolygonGroupId is got(see GetExistingPatches)
+	 * Warning: CopyPatchGroups is call in FCoreTechRetessellate_Impl::ApplyOnOneAsset(CoreTechRetessellateAction.cpp) only if the option 
+	 * RetessellateOptions.RetessellationRule equal EDatasmithCADRetessellationRule::SkipDeletedSurfaces
+	 */
 	void CopyPatchGroups(FMeshDescription& MeshSource, FMeshDescription& MeshDestination)
 	{
 		TPolygonGroupAttributesRef<int32> PatchGroups = MeshDestination.PolygonGroupAttributes().GetAttributesRef<int32>(PolyTriGroups);
@@ -305,20 +292,18 @@ namespace CADLibrary
 		}
 	}
 
-	bool FillMesh(const FMeshParameters& MeshParameters, const FImportParameters& ImportParams, TArray<FTessellationData>& FaceTessellations, FMeshDescription& MeshDescription)
+	bool FillMesh(const FMeshParameters& MeshParameters, const FImportParameters& ImportParams, FBodyMesh& BodyTessellation, FMeshDescription& MeshDescription)
 	{
 		const int32 UVChannel = 0;
 		const int32 TriangleCount = 3;
 		const TriangleIndex Clockwise = { 0, 1, 2 };
 		const TriangleIndex CounterClockwise = { 0, 2, 1 };
-		const int32 InvalidID = INDEX_NONE;
-
 
 		TArray<FVertexInstanceID> TriangleVertexInstanceIDs;
 		TriangleVertexInstanceIDs.SetNum(TriangleCount);
 
 		TArray<FVertexInstanceID> MeshVertexInstanceIDs;
-		TArray<uint32> CTFaceIndex;  // new CT face index to remove degenerated face
+		TArray<uint32> NewFaceIndex;  // new CT face index to remove degenerated face
 
 		// Gather all array data
 		FStaticMeshAttributes Attributes(MeshDescription);
@@ -336,7 +321,7 @@ namespace CADLibrary
 
 		// Find all the materials used
 		TMap<uint32, FPolygonGroupID> MaterialToPolygonGroupMapping;
-		for (const FTessellationData& FaceTessellation : FaceTessellations)
+		for (const FTessellationData& FaceTessellation : BodyTessellation.Faces)
 		{
 			// we assume that face has only color
 			MaterialToPolygonGroupMapping.Add(FaceTessellation.ColorName, INDEX_NONE);
@@ -366,12 +351,14 @@ namespace CADLibrary
 		bool bImportOnlyAlreadyPresent = (bool) PatchIdSet.Num();
 
 		TPolygonAttributesRef<int32> PatchGroups = EnableCADPatchGroups(MeshDescription);
+		int32 PatchIndex = 0;
 		for (int32 Step = 0; Step < NbStep; ++Step)
 		{
 			// Swap mesh if needed
 			const TriangleIndex& Orientation = (!MeshParameters.bNeedSwapOrientation == (bool)Step) ? CounterClockwise : Clockwise;
+			TArray<int32>& VertexIdSet = (Step == 0) ? BodyTessellation.VertexIds : BodyTessellation.SymmetricVertexIds;
 
-			for (FTessellationData& Tessellation : FaceTessellations)
+			for (FTessellationData& Tessellation : BodyTessellation.Faces)
 			{
 				if (bImportOnlyAlreadyPresent && !PatchIdSet.Contains(Tessellation.PatchId))
 				{
@@ -385,55 +372,54 @@ namespace CADLibrary
 					continue;
 				}
 
-				//int32 TriangleCount = IndicesCount / 3;
 				int32 VertexIDs[3];
+				int32 VertexIndex[3];
 				FVector Temp3D = { 0, 0, 0 };
 				FVector2D TexCoord2D = { 0, 0 };
 
-				MeshVertexInstanceIDs.SetNum(Tessellation.IndexArray.Num());
-				CTFaceIndex.Reserve(Tessellation.IndexArray.Num());
-				CTFaceIndex.SetNum(0);
+				MeshVertexInstanceIDs.Empty(Tessellation.VertexIndices.Num());
+				NewFaceIndex.Empty(Tessellation.VertexIndices.Num());
 
-				TArray<int32>& VertexIdSet = (Step == 0) ? Tessellation.VertexIdSet : Tessellation.SymVertexIdSet;
+				PatchIndex++;
 
 				// build each valid face i.e. 3 different indexes
-				for (int32 Index = 0, NewIndex = 0; Index < Tessellation.IndexArray.Num(); Index += 3)
+				for (int32 Index = 0; Index < Tessellation.VertexIndices.Num(); Index += 3)
 				{
-					VertexIDs[Orientation[0]] = VertexIdSet[Tessellation.IndexArray[Index + 0]];
-					VertexIDs[Orientation[1]] = VertexIdSet[Tessellation.IndexArray[Index + 1]];
-					VertexIDs[Orientation[2]] = VertexIdSet[Tessellation.IndexArray[Index + 2]];
-
-					if (VertexIDs[0] == InvalidID || VertexIDs[1] == InvalidID || VertexIDs[2] == InvalidID)
-					{
-						continue;
-					}
+					VertexIndex[0] = Tessellation.PositionIndices[Tessellation.VertexIndices[Index + Orientation[0]]];
+					VertexIndex[1] = Tessellation.PositionIndices[Tessellation.VertexIndices[Index + Orientation[1]]];
+					VertexIndex[2] = Tessellation.PositionIndices[Tessellation.VertexIndices[Index + Orientation[2]]];
 
 					// Verify the 3 input indices are not defining a degenerated triangle
-					if (VertexIDs[0] == VertexIDs[1] || VertexIDs[0] == VertexIDs[2] || VertexIDs[1] == VertexIDs[2])
+					if (VertexIndex[0] == INDEX_NONE || VertexIndex[1] == INDEX_NONE || VertexIndex[2] == INDEX_NONE)
 					{
 						continue;
 					}
 
-					CTFaceIndex.Add(Tessellation.IndexArray[Index + 0]);
-					CTFaceIndex.Add(Tessellation.IndexArray[Index + 1]);
-					CTFaceIndex.Add(Tessellation.IndexArray[Index + 2]);
+					VertexIDs[0] = VertexIdSet[VertexIndex[0]];
+					VertexIDs[1] = VertexIdSet[VertexIndex[1]];
+					VertexIDs[2] = VertexIdSet[VertexIndex[2]];
 
-					TriangleVertexInstanceIDs[0] = MeshVertexInstanceIDs[NewIndex++] = MeshDescription.CreateVertexInstance((FVertexID)VertexIDs[0]);
-					TriangleVertexInstanceIDs[1] = MeshVertexInstanceIDs[NewIndex++] = MeshDescription.CreateVertexInstance((FVertexID)VertexIDs[1]);
-					TriangleVertexInstanceIDs[2] = MeshVertexInstanceIDs[NewIndex++] = MeshDescription.CreateVertexInstance((FVertexID)VertexIDs[2]);
+					NewFaceIndex.Add(Tessellation.VertexIndices[Index + Orientation[0]]);
+					NewFaceIndex.Add(Tessellation.VertexIndices[Index + Orientation[1]]);
+					NewFaceIndex.Add(Tessellation.VertexIndices[Index + Orientation[2]]);
+
+					MeshVertexInstanceIDs.Add(TriangleVertexInstanceIDs[0] = MeshDescription.CreateVertexInstance((FVertexID)VertexIDs[0]));
+					MeshVertexInstanceIDs.Add(TriangleVertexInstanceIDs[1] = MeshDescription.CreateVertexInstance((FVertexID)VertexIDs[1]));
+					MeshVertexInstanceIDs.Add(TriangleVertexInstanceIDs[2] = MeshDescription.CreateVertexInstance((FVertexID)VertexIDs[2]));
 
 					// Add the triangle as a polygon to the mesh description
 					const FPolygonID PolygonID = MeshDescription.CreatePolygon(*PolygonGroupID, TriangleVertexInstanceIDs);
+
 					// Set patch id attribute
 					PatchGroups[PolygonID] = Tessellation.PatchId;
 				}
 
-				// finalization of the mesh by setting colors, tangents, bi-normals, UV
-				for (int32 IndexFace = 0; IndexFace < CTFaceIndex.Num(); IndexFace += 3)
+				for (int32 IndexFace = 0; IndexFace < MeshVertexInstanceIDs.Num(); IndexFace += 3)
 				{
 					for (int32 Index = 0; Index < TriangleCount; Index++)
 					{
-						FVertexInstanceID VertexInstanceID = MeshVertexInstanceIDs[IndexFace + Orientation[Index]];
+						const FVertexInstanceID VertexInstanceID = MeshVertexInstanceIDs[IndexFace + Orientation[Index]];
+						VertexInstanceUVs.Set(VertexInstanceID, UVChannel, Tessellation.TexCoordArray[NewFaceIndex[IndexFace + Index]]);
 
 						VertexInstanceColors[VertexInstanceID] = FLinearColor::White;
 						VertexInstanceTangents[VertexInstanceID] = FVector(ForceInitToZero);
@@ -443,12 +429,12 @@ namespace CADLibrary
 
 				if (Tessellation.TexCoordArray.Num())
 				{
-					for (int32 IndexFace = 0; IndexFace < CTFaceIndex.Num(); IndexFace += 3)
+					for (int32 IndexFace = 0; IndexFace < NewFaceIndex.Num(); IndexFace += 3)
 					{
 						for (int32 Index = 0; Index < TriangleCount; Index++)
 						{
 							FVertexInstanceID VertexInstanceID = MeshVertexInstanceIDs[IndexFace + Orientation[Index]];
-							VertexInstanceUVs.Set(VertexInstanceID, UVChannel, Tessellation.TexCoordArray[CTFaceIndex[IndexFace + Index]]);
+							VertexInstanceUVs.Set(VertexInstanceID, UVChannel, Tessellation.TexCoordArray[NewFaceIndex[IndexFace + Index]]);
 						}
 					}
 				}
@@ -464,11 +450,11 @@ namespace CADLibrary
 
 				if (Tessellation.NormalArray.Num() == 1)
 				{
-					Temp3D = Tessellation.NormalArray[0];
-					for (int32 Index = 0; Index < CTFaceIndex.Num(); Index++)
+					const FVector& Normal = Tessellation.NormalArray[0];
+					for (int32 Index = 0; Index < NewFaceIndex.Num(); Index++)
 					{
-						FVertexInstanceID VertexInstanceID = MeshVertexInstanceIDs[Index];
-						VertexInstanceNormals[VertexInstanceID] = Temp3D;
+						const FVertexInstanceID VertexInstanceID = MeshVertexInstanceIDs[Index];
+						VertexInstanceNormals[VertexInstanceID] = Normal;
 					}
 				}
 				else
@@ -478,12 +464,12 @@ namespace CADLibrary
 						Normal = Normal.GetSafeNormal();
 					}
 
-					for (int32 IndexFace = 0; IndexFace < CTFaceIndex.Num(); IndexFace += 3)
+					for (int32 IndexFace = 0; IndexFace < NewFaceIndex.Num(); IndexFace += 3)
 					{
 						for (int32 Index = 0; Index < 3; Index++)
 						{
-							FVertexInstanceID VertexInstanceID = MeshVertexInstanceIDs[IndexFace + Orientation[Index]];
-							VertexInstanceNormals[VertexInstanceID] = Tessellation.NormalArray[CTFaceIndex[IndexFace + Index]];
+							const FVertexInstanceID VertexInstanceID = MeshVertexInstanceIDs[IndexFace + Orientation[Index]];
+							VertexInstanceNormals[VertexInstanceID] = Tessellation.NormalArray[NewFaceIndex[IndexFace + Index]];
 						}
 					}
 				}
@@ -491,17 +477,18 @@ namespace CADLibrary
 				// compute normals
 				if (Step)
 				{
+					// compute normals of Symmetric vertex
 					FMatrix SymmetricMatrix;
 					SymmetricMatrix = FDatasmithUtils::GetSymmetricMatrix(MeshParameters.SymmetricOrigin, MeshParameters.SymmetricNormal);
-					for (int32 Index = 0; Index < CTFaceIndex.Num(); Index++)
+					for (const FVertexInstanceID& VertexInstanceID : MeshVertexInstanceIDs)
 					{
-						VertexInstanceNormals[MeshVertexInstanceIDs[Index]] = SymmetricMatrix.TransformVector(VertexInstanceNormals[MeshVertexInstanceIDs[Index]]);;
+						VertexInstanceNormals[VertexInstanceID] = SymmetricMatrix.TransformVector(VertexInstanceNormals[VertexInstanceID]);;
 					}
 				}
 
 				if (MeshParameters.bNeedSwapOrientation)
 				{
-					for (int32 Index = 0; Index < CTFaceIndex.Num(); Index++)
+					for (int32 Index = 0; Index < MeshVertexInstanceIDs.Num(); Index++)
 					{
 						VertexInstanceNormals[MeshVertexInstanceIDs[Index]] = VertexInstanceNormals[MeshVertexInstanceIDs[Index]] * -1.f;
 					}
@@ -514,15 +501,21 @@ namespace CADLibrary
 	bool ConvertBodyMeshToMeshDescription(const FImportParameters& ImportParams, const FMeshParameters& MeshParameters, FBodyMesh& Body, FMeshDescription& MeshDescription)
 	{
 		// in a closed big mesh VertexCount ~ TriangleCount / 2, EdgeCount ~ 1.5* TriangleCount
-		MeshDescription.ReserveNewVertexInstances(Body.TriangleCount * 3);
+		MeshDescription.ReserveNewVertexInstances(Body.VertexArray.Num());
 		MeshDescription.ReserveNewPolygons(Body.TriangleCount);
 		MeshDescription.ReserveNewEdges(Body.TriangleCount * 3);
 
 		// CoreTech is generating position duplicates. make sure to remove them before filling the mesh description
-		TArray<FVertexID> RemapVertexPosition;
-		FillVertexPosition(ImportParams, MeshParameters, Body, MeshDescription);
-
-		if (!FillMesh(MeshParameters, ImportParams, Body.Faces, MeshDescription))
+		if (Body.VertexArray.Num() == 0)
+		{
+			FillKioVertexPosition(ImportParams, MeshParameters, Body, MeshDescription);
+		}
+		else
+		{
+			//FillVertexPosition(ImportParams, MeshParameters, Body, MeshDescription);
+		}
+		
+		if (!FillMesh(MeshParameters, ImportParams, Body, MeshDescription))
 		{
 			return false;
 		}
diff --git a/Engine/Plugins/Enterprise/DatasmithCADImporter/Source/CADTools/Private/CADData.cpp b/Engine/Plugins/Enterprise/DatasmithCADImporter/Source/CADTools/Private/CADData.cpp
--- a/Engine/Plugins/Enterprise/DatasmithCADImporter/Source/CADTools/Private/CADData.cpp
+++ b/Engine/Plugins/Enterprise/DatasmithCADImporter/Source/CADTools/Private/CADData.cpp
@@ -7,8 +7,11 @@
 #include "Misc/FileHelper.h"
 #include "Misc/Paths.h"
 
+uint32 MeshArchiveMagic = 345612;
+
 namespace CADLibrary
 {
+
 uint32 BuildColorId(uint32 ColorId, uint8 Alpha)
 {
 	if (Alpha == 0)
@@ -18,7 +21,7 @@ uint32 BuildColorId(uint32 ColorId, uint8 Alpha)
 	return ColorId | Alpha << 24;
 }
 
-void GetCTColorIdAlpha(ColorId ColorId, uint32& CTColorId, uint8& Alpha)
+void GetCTColorIdAlpha(FColorId ColorId, uint32& CTColorId, uint8& Alpha)
 {
 	CTColorId = ColorId & 0x00ffffff;
 	Alpha = (uint8)((ColorId & 0xff000000) >> 24);
@@ -73,29 +76,31 @@ FArchive& operator<<(FArchive& Ar, FFileDescription& File)
 
 FArchive& operator<<(FArchive& Ar, FTessellationData& TessellationData)
 {
-	Ar << TessellationData.PatchId;
+	Ar << TessellationData.PositionArray;
+
+	Ar << TessellationData.PositionIndices;
+	Ar << TessellationData.VertexIndices;
 
-	Ar << TessellationData.VertexArray;
 	Ar << TessellationData.NormalArray;
-	Ar << TessellationData.IndexArray;
 	Ar << TessellationData.TexCoordArray;
 
-	Ar << TessellationData.StartVertexIndex;
-
 	Ar << TessellationData.ColorName;
 	Ar << TessellationData.MaterialName;
 
+	Ar << TessellationData.PatchId;
+
 	return Ar;
 }
 
 FArchive& operator<<(FArchive& Ar, FBodyMesh& BodyMesh)
 {
+	Ar << BodyMesh.VertexArray;
+	Ar << BodyMesh.Faces;
+	Ar << BodyMesh.BBox;
+
 	Ar << BodyMesh.TriangleCount;
 	Ar << BodyMesh.BodyID;
 	Ar << BodyMesh.MeshActorName;
-	Ar << BodyMesh.Faces;
-
-	Ar << BodyMesh.BBox;
 
 	Ar << BodyMesh.MaterialSet;
 	Ar << BodyMesh.ColorSet;
@@ -107,8 +112,8 @@ void SerializeBodyMeshSet(const TCHAR* Filename, TArray<FBodyMesh>& InBodySet)
 {
 	TUniquePtr<FArchive> Archive(IFileManager::Get().CreateFileWriter(Filename));
 
-	uint32 type = 234561;
-	*Archive << type;
+	uint32 MagicNumber = MeshArchiveMagic;
+	*Archive << MagicNumber;
 
 	*Archive << InBodySet;
 
@@ -119,9 +124,9 @@ void DeserializeBodyMeshFile(const TCHAR* Filename, TArray<FBodyMesh>& OutBodySe
 {
 	TUniquePtr<FArchive> Archive(IFileManager::Get().CreateFileReader(Filename));
 
-	uint32 type = 0;
-	*Archive << type;
-	if (type != 234561)
+	uint32 MagicNumber = 0;
+	*Archive << MagicNumber;
+	if (MagicNumber != MeshArchiveMagic)
 	{
 		Archive->Close();
 		return;
diff --git a/Engine/Plugins/Enterprise/DatasmithCADImporter/Source/CADTools/Public/CADData.h b/Engine/Plugins/Enterprise/DatasmithCADImporter/Source/CADTools/Public/CADData.h
--- a/Engine/Plugins/Enterprise/DatasmithCADImporter/Source/CADTools/Public/CADData.h
+++ b/Engine/Plugins/Enterprise/DatasmithCADImporter/Source/CADTools/Public/CADData.h
@@ -9,10 +9,10 @@
 
 class FArchive;
 
-using CadId = uint32; // Identifier defined in the input CAD file
-using ColorId = uint32; // Identifier defined in the input CAD file
-using MaterialId = uint32; // Identifier defined in the input CAD file
-using CADUUID = uint32;  // Universal unique identifier that be used for the unreal asset name (Actor, Material)
+using FCadId = uint32; // Identifier defined in the input CAD file
+using FColorId = uint32; // Identifier defined in the input CAD file
+using FMaterialId = uint32; // Identifier defined in the input CAD file
+using FCADUUID = uint32;  // Universal unique identifier that be used for the unreal asset name (Actor, Material)
 
 
 namespace CADLibrary
@@ -39,9 +39,9 @@ public:
 
 struct CADTOOLS_API FObjectDisplayDataId
 {
-	CADUUID DefaultMaterialName = 0;
-	MaterialId Material = 0;
-	ColorId Color = 0; // => FastHash == ColorId+Transparency
+	FCADUUID DefaultMaterialName = 0;
+	FMaterialId Material = 0;
+	FColorId Color = 0; // => FastHash == ColorId+Transparency
 };
 
 struct CADTOOLS_API FFileDescription
@@ -87,32 +87,44 @@ struct CADTOOLS_API FFileDescription
 };
 
 /**
- * Helper struct to store tessellation data from CoreTech
+ * Helper struct to store tessellation data from CoreTech or CADKernel
+ *
+ * FBodyMesh and FTessellationData are design to manage mesh from CoreTech and CADKernel.
+ * FTessellationData is the mesh of a face
+ * FBodyMesh is the mesh of a body composed by an array of FTessellationData (one FTessellationData by body face)
+ *
+ * CoreTech mesh are defined surface by surface. The mesh is not connected
+ * CADKernel mesh is connected.
  */
 struct CADTOOLS_API FTessellationData
 {
 	friend CADTOOLS_API FArchive& operator<<(FArchive& Ar, FTessellationData& Tessellation);
 
-	int32    PatchId = 0;
+	/** Empty with CADKernel as set in FBodyMesh, Set by CoreTech (this is only the vertices of the face) */
+	TArray<FVector> PositionArray;
 
-	TArray<FVector> VertexArray;
+	/** Index of each vertex in FBody::VertexArray. Empty with CoreTech and filled by FillKioVertexPosition */
+	TArray<int32> PositionIndices;
+
+	/** Index of Vertices of each face in the local Vertices set (i.e. VerticesBodyIndex for CADKernel, VertexArray for Coretech) */
+	TArray<int32> VertexIndices;
+
+	/** Normal of each vertex */
 	TArray<FVector> NormalArray;
-	TArray<int32> IndexArray;
-	TArray<FVector2D> TexCoordArray;
 
-	uint32    StartVertexIndex = 0;
+	/** UV coordinates of each vertex */
+	TArray<FVector2D> TexCoordArray;
 
-	CADUUID ColorName = 0;
-	CADUUID MaterialName = 0;
+	FCADUUID ColorName = 0;
+	FCADUUID MaterialName = 0;
 
-	TArray<int32> VertexIdSet;  // StaticMesh FVertexID NO Serialize
-	TArray<int32> SymVertexIdSet; // StaticMesh FVertexID for sym part NO Serialize
+	int32 PatchId;
 };
 
 class CADTOOLS_API FBodyMesh
 {
 public:
-	FBodyMesh(CadId InBodyID = 0) : BodyID(InBodyID)
+	FBodyMesh(FCadId InBodyID = 0) : BodyID(InBodyID)
 	{
 		BBox.Init();
 	}
@@ -120,12 +132,16 @@ public:
 	friend FArchive& operator<<(FArchive& Ar, FBodyMesh& BodyMesh);
 
 public:
+	TArray<FVector> VertexArray; // set by CADKernel, filled by FillKioVertexPosition that merges coincident vertices (CoreTechHelper)
 	TArray<FTessellationData> Faces;
 	FBox BBox;
 
 	uint32 TriangleCount = 0;
-	CadId BodyID = 0;
-	CADUUID MeshActorName = 0;
+	FCadId BodyID = 0;
+	FCADUUID MeshActorName = 0;
+
+	TArray<int32> VertexIds;  // StaticMesh FVertexID NO Serialize, filled by FillKioVertexPosition or FillVertexPosition
+	TArray<int32> SymmetricVertexIds; // StaticMesh FVertexID for sym part NO Serialize, filled by FillKioVertexPosition or FillVertexPosition
 
 	TSet<uint32> MaterialSet;
 	TSet<uint32> ColorSet;
diff --git a/Engine/Plugins/Enterprise/DatasmithCADImporter/Source/DatasmithCADTranslator/Private/DatasmithCADTranslator.cpp b/Engine/Plugins/Enterprise/DatasmithCADImporter/Source/DatasmithCADTranslator/Private/DatasmithCADTranslator.cpp
--- a/Engine/Plugins/Enterprise/DatasmithCADImporter/Source/DatasmithCADTranslator/Private/DatasmithCADTranslator.cpp
+++ b/Engine/Plugins/Enterprise/DatasmithCADImporter/Source/DatasmithCADTranslator/Private/DatasmithCADTranslator.cpp
@@ -21,6 +21,12 @@ static TAutoConsoleVariable<int32> CVarStaticCADTranslatorEnableThreadedImport(
 	TEXT("Activate to parallelise CAD file processing.\n"),
 	ECVF_Default);
 
+static TAutoConsoleVariable<int32> CVarStaticCADTranslatorEnableKernelIOTessellation(
+	TEXT("r.CADTranslator.EnableKernelIOTessellation"),
+	1,
+	TEXT("Activate to use KernelIO Tessellation.\n"),
+	ECVF_Default);
+
 void FDatasmithCADTranslator::Initialize(FDatasmithTranslatorCapabilities& OutCapabilities)
 {
 	if (ICADInterfacesModule::GetAvailability() == ECADInterfaceAvailability::Unavailable)
@@ -77,6 +83,11 @@ void FDatasmithCADTranslator::Initialize(FDatasmithTranslatorCapabilities& OutCa
 
 bool FDatasmithCADTranslator::LoadScene(TSharedRef<IDatasmithScene> DatasmithScene)
 {
+	ImportParameters.bEnableKernelIOTessellation = (CVarStaticCADTranslatorEnableKernelIOTessellation.GetValueOnAnyThread() != 0);;
+#ifdef CADKERNEL_DEBUG
+	//ImportParameters.bEnableKernelIOTessellation = false;
+#endif
+
 	ImportParameters.MetricUnit = 0.001;
 	ImportParameters.ScaleFactor = 0.1;
 	const FDatasmithTessellationOptions& TesselationOptions = GetCommonTessellationOptions();
@@ -125,31 +136,31 @@ bool FDatasmithCADTranslator::LoadScene(TSharedRef<IDatasmithScene> DatasmithSce
 		ImportParameters.bEnableCacheUsage = true;
 	}
 
+	// Only use multi-processed translation if it is required and cache's usage is enabled
+	if (ImportParameters.bEnableCacheUsage)
+	{
 		bool bWithProcessor = (CVarStaticCADTranslatorEnableThreadedImport.GetValueOnAnyThread() != 0);
 
 #ifdef CAD_TRANSLATOR_DEBUG
 		bWithProcessor = false;
 #endif //CAD_TRANSLATOR_DEBUG
 
-	// Only use multi-processed translation if it is required and cache's usage is enabled
-	if (bWithProcessor && ImportParameters.bEnableCacheUsage)
-	{
 		TMap<uint32, FString> CADFileToUEFileMap;
 		int32 NumCores = FPlatformMisc::NumberOfCores();
 		{
 			DatasmithDispatcher::FDatasmithDispatcher Dispatcher(ImportParameters, CachePath, NumCores, CADFileToUEFileMap, CADFileToUEGeomMap);
 			Dispatcher.AddTask(FileDescription);
 
-		Dispatcher.Process(bWithProcessor);
-	}
+			Dispatcher.Process(bWithProcessor);
+		}
 
-	FDatasmithSceneGraphBuilder SceneGraphBuilder(CADFileToUEFileMap, CachePath, DatasmithScene, GetSource(), ImportParameters);
-	SceneGraphBuilder.Build();
+		FDatasmithSceneGraphBuilder SceneGraphBuilder(CADFileToUEFileMap, CachePath, DatasmithScene, GetSource(), ImportParameters);
+		SceneGraphBuilder.Build();
 
-	MeshBuilderPtr = MakeUnique<FDatasmithMeshBuilder>(CADFileToUEGeomMap, CachePath, ImportParameters);
+		MeshBuilderPtr = MakeUnique<FDatasmithMeshBuilder>(CADFileToUEGeomMap, CachePath, ImportParameters);
 
-	return true;
-}
+		return true;
+	}
 
 	ImportParameters.bEnableCacheUsage = false;
 
diff --git a/Engine/Plugins/Enterprise/DatasmithCADImporter/Source/DatasmithCADTranslator/Private/DatasmithMeshBuilder.cpp b/Engine/Plugins/Enterprise/DatasmithCADImporter/Source/DatasmithCADTranslator/Private/DatasmithMeshBuilder.cpp
--- a/Engine/Plugins/Enterprise/DatasmithCADImporter/Source/DatasmithCADTranslator/Private/DatasmithMeshBuilder.cpp
+++ b/Engine/Plugins/Enterprise/DatasmithCADImporter/Source/DatasmithCADTranslator/Private/DatasmithMeshBuilder.cpp
@@ -42,7 +42,7 @@ void FDatasmithMeshBuilder::LoadMeshFiles(TMap<uint32, FString>& CADFileToMeshFi
 TOptional<FMeshDescription> FDatasmithMeshBuilder::GetMeshDescription(TSharedRef<IDatasmithMeshElement> OutMeshElement, CADLibrary::FMeshParameters& OutMeshParameters)
 {
 	const TCHAR* NameLabel = OutMeshElement->GetName();
-	CADUUID BodyUuid = (CADUUID) FCString::Atoi64(OutMeshElement->GetName() + 2);  // +2 to remove 2 first char (Ox)
+	FCADUUID BodyUuid = (FCADUUID) FCString::Atoi64(OutMeshElement->GetName() + 2);  // +2 to remove 2 first char (Ox)
 	if (BodyUuid == 0)
 	{
 		return TOptional<FMeshDescription>();
diff --git a/Engine/Plugins/Enterprise/DatasmithCADImporter/Source/DatasmithCADTranslator/Private/DatasmithSceneGraphBuilder.cpp b/Engine/Plugins/Enterprise/DatasmithCADImporter/Source/DatasmithCADTranslator/Private/DatasmithSceneGraphBuilder.cpp
--- a/Engine/Plugins/Enterprise/DatasmithCADImporter/Source/DatasmithCADTranslator/Private/DatasmithSceneGraphBuilder.cpp
+++ b/Engine/Plugins/Enterprise/DatasmithCADImporter/Source/DatasmithCADTranslator/Private/DatasmithSceneGraphBuilder.cpp
@@ -136,7 +136,7 @@ void FDatasmithSceneGraphBuilder::FillAnchorActor(const TSharedRef<IDatasmithAct
 		return;
 	}
 
-	CadId RootId = 1;
+	FCadId RootId = 1;
 	int32* Index = SceneGraph->CADIdToComponentIndex.Find(RootId);
 	if (!Index)
 	{
@@ -202,7 +202,7 @@ FDatasmithSceneBaseGraphBuilder::FDatasmithSceneBaseGraphBuilder(CADLibrary::FAr
 
 bool FDatasmithSceneBaseGraphBuilder::Build()
 {
-	CadId RootId = 1;
+	FCadId RootId = 1;
 	int32* Index = SceneGraph->CADIdToComponentIndex.Find(RootId);
 	if (!Index)
 	{
@@ -320,7 +320,7 @@ TSharedPtr< IDatasmithActorElement >  FDatasmithSceneBaseGraphBuilder::BuildInst
 				{
 					AncestorSceneGraphHash.Add(InstanceSceneGraphHash);
 
-					CadId RootId = 1;
+					FCadId RootId = 1;
 					int32* Index = SceneGraph->CADIdToComponentIndex.Find(RootId);
 					if (Index)
 					{
@@ -433,7 +433,7 @@ void FDatasmithSceneBaseGraphBuilder::GetNodeUUIDAndName(
 	}
 	CleanName(OutName);
 
-	CADUUID UEUUID = HashCombine(GetTypeHash(InParentUEUUID), GetTypeHash(InComponentIndex));
+	FCADUUID UEUUID = HashCombine(GetTypeHash(InParentUEUUID), GetTypeHash(InComponentIndex));
 
 	if (InstanceUUID)
 	{
diff --git a/Engine/Plugins/Enterprise/DatasmithCADImporter/Source/DatasmithCADTranslator/Public/DatasmithMeshBuilder.h b/Engine/Plugins/Enterprise/DatasmithCADImporter/Source/DatasmithCADTranslator/Public/DatasmithMeshBuilder.h
--- a/Engine/Plugins/Enterprise/DatasmithCADImporter/Source/DatasmithCADTranslator/Public/DatasmithMeshBuilder.h
+++ b/Engine/Plugins/Enterprise/DatasmithCADImporter/Source/DatasmithCADTranslator/Public/DatasmithMeshBuilder.h
@@ -36,7 +36,7 @@ protected:
 	void LoadMeshFiles(TMap<uint32, FString>& CADFileToMeshFile);
 
 	TArray<TArray<CADLibrary::FBodyMesh>> BodyMeshes;
-	TMap<CADUUID, CADLibrary::FBodyMesh*> MeshActorNameToBodyMesh;
+	TMap<FCADUUID, CADLibrary::FBodyMesh*> MeshActorNameToBodyMesh;
 
 	CADLibrary::FImportParameters ImportParameters;
 };
diff --git a/Engine/Plugins/Enterprise/DatasmithCADImporter/Source/DatasmithCADTranslator/Public/DatasmithSceneGraphBuilder.h b/Engine/Plugins/Enterprise/DatasmithCADImporter/Source/DatasmithCADTranslator/Public/DatasmithSceneGraphBuilder.h
--- a/Engine/Plugins/Enterprise/DatasmithCADImporter/Source/DatasmithCADTranslator/Public/DatasmithSceneGraphBuilder.h
+++ b/Engine/Plugins/Enterprise/DatasmithCADImporter/Source/DatasmithCADTranslator/Public/DatasmithSceneGraphBuilder.h
@@ -91,13 +91,13 @@ protected:
 	TArray<CADLibrary::FArchiveSceneGraph> ArchiveMockUps;
 	TMap<uint32, CADLibrary::FArchiveSceneGraph*> CADFileToSceneGraphArchive;
 
-	TMap< CADUUID, TSharedPtr< IDatasmithMeshElement > > BodyUuidToMeshElement;
+	TMap< FCADUUID, TSharedPtr< IDatasmithMeshElement > > BodyUuidToMeshElement;
 
-	TMap< CADUUID, TSharedPtr< IDatasmithUEPbrMaterialElement > > MaterialUuidMap;
+	TMap< FCADUUID, TSharedPtr< IDatasmithUEPbrMaterialElement > > MaterialUuidMap;
 	TSharedPtr<IDatasmithUEPbrMaterialElement > DefaultMaterial;
 
-	TMap<CADUUID, CADLibrary::FArchiveColor> ColorNameToColorArchive; 
-	TMap<CADUUID, CADLibrary::FArchiveMaterial> MaterialNameToMaterialArchive; 
+	TMap<FCADUUID, CADLibrary::FArchiveColor> ColorNameToColorArchive; 
+	TMap<FCADUUID, CADLibrary::FArchiveMaterial> MaterialNameToMaterialArchive; 
 
 	TArray<uint32> AncestorSceneGraphHash;
 
