CADKernel Retesselation implementation
CoretechExtension become ParametricSurfaceExtension (Step1)
CoretechSurface module will derived from the new module ParametricSurface
the new CADkernelSurface module will derived from the new module ParametricSurface and will manage data for CADKernel

Step 1: Rename CoretechExtension module into ParametricSurfaceExtension

diff --git a/Engine/Plugins/Enterprise/DatasmithCADImporter/DatasmithCADImporter.uplugin b/Engine/Plugins/Enterprise/DatasmithCADImporter/DatasmithCADImporter.uplugin
--- a/Engine/Plugins/Enterprise/DatasmithCADImporter/DatasmithCADImporter.uplugin
+++ b/Engine/Plugins/Enterprise/DatasmithCADImporter/DatasmithCADImporter.uplugin
@@ -74,7 +74,7 @@
 			]
 		},
 		{
-			"Name": "CoreTechExtension",
+			"Name": "ParametricSurfaceExtension",
 			"Type": "Editor",
 			"LoadingPhase": "Default",
 			"WhitelistPlatforms": [
diff --git a/Engine/Plugins/Enterprise/DatasmithCADImporter/Source/CoreTechExtension/CoreTechExtension.Build.cs b/Engine/Plugins/Enterprise/DatasmithCADImporter/Source/CoreTechExtension/CoreTechExtension.Build.cs
deleted file mode 100644
diff --git a/Engine/Plugins/Enterprise/DatasmithCADImporter/Source/CoreTechExtension/Private/CoreTechBlueprintLibrary.cpp b/Engine/Plugins/Enterprise/DatasmithCADImporter/Source/CoreTechExtension/Private/CoreTechBlueprintLibrary.cpp
deleted file mode 100644
diff --git a/Engine/Plugins/Enterprise/DatasmithCADImporter/Source/CoreTechExtension/Private/CoreTechExtensionModule.cpp b/Engine/Plugins/Enterprise/DatasmithCADImporter/Source/CoreTechExtension/Private/CoreTechExtensionModule.cpp
deleted file mode 100644
diff --git a/Engine/Plugins/Enterprise/DatasmithCADImporter/Source/CoreTechExtension/Private/CoreTechRetessellateAction.cpp b/Engine/Plugins/Enterprise/DatasmithCADImporter/Source/CoreTechExtension/Private/CoreTechRetessellateAction.cpp
deleted file mode 100644
diff --git a/Engine/Plugins/Enterprise/DatasmithCADImporter/Source/CoreTechExtension/Private/DataprepTessellationOperation.cpp b/Engine/Plugins/Enterprise/DatasmithCADImporter/Source/CoreTechExtension/Private/DataprepTessellationOperation.cpp
deleted file mode 100644
diff --git a/Engine/Plugins/Enterprise/DatasmithCADImporter/Source/CoreTechExtension/Private/DataprepTessellationOperation.h b/Engine/Plugins/Enterprise/DatasmithCADImporter/Source/CoreTechExtension/Private/DataprepTessellationOperation.h
deleted file mode 100644
diff --git a/Engine/Plugins/Enterprise/DatasmithCADImporter/Source/CoreTechExtension/Public/CoreTechBlueprintLibrary.h b/Engine/Plugins/Enterprise/DatasmithCADImporter/Source/CoreTechExtension/Public/CoreTechBlueprintLibrary.h
deleted file mode 100644
diff --git a/Engine/Plugins/Enterprise/DatasmithCADImporter/Source/CoreTechExtension/Public/CoreTechExtensionModule.h b/Engine/Plugins/Enterprise/DatasmithCADImporter/Source/CoreTechExtension/Public/CoreTechExtensionModule.h
deleted file mode 100644
diff --git a/Engine/Plugins/Enterprise/DatasmithCADImporter/Source/CoreTechExtension/Public/CoreTechRetessellateAction.h b/Engine/Plugins/Enterprise/DatasmithCADImporter/Source/CoreTechExtension/Public/CoreTechRetessellateAction.h
deleted file mode 100644
diff --git a/Engine/Plugins/Enterprise/DatasmithCADImporter/Source/ParametricSurfaceExtension/ParametricSurfaceExtension.Build.cs b/Engine/Plugins/Enterprise/DatasmithCADImporter/Source/ParametricSurfaceExtension/ParametricSurfaceExtension.Build.cs
new file mode 100644
--- /dev/null
+++ b/Engine/Plugins/Enterprise/DatasmithCADImporter/Source/ParametricSurfaceExtension/ParametricSurfaceExtension.Build.cs
@@ -0,0 +1,40 @@
+// Copyright Epic Games, Inc. All Rights Reserved.
+
+using System.IO;
+
+namespace UnrealBuildTool.Rules
+{
+	public class ParametricSurfaceExtension : ModuleRules
+	{
+		public ParametricSurfaceExtension(ReadOnlyTargetRules Target) : base(Target)
+		{
+			PublicDependencyModuleNames.AddRange(
+				new string[]
+				{
+					"Core",
+				}
+			);
+
+			PrivateDependencyModuleNames.AddRange(
+				new string[]
+				{
+					"Core",
+					"CoreUObject",
+					"MeshDescription",
+					"DataprepCore",
+					"DatasmithContent",
+					"CoreTechSurface",
+					"DatasmithImporter",
+					"Engine",
+					"StaticMeshEditor",
+					"Slate",
+					"SlateCore",
+					"StaticMeshDescription",
+					"EditorFramework",
+					"UnrealEd",
+					"CADLibrary",
+				}
+			);
+		}
+	}
+}
diff --git a/Engine/Plugins/Enterprise/DatasmithCADImporter/Source/ParametricSurfaceExtension/Private/CoreTechBlueprintLibrary.cpp b/Engine/Plugins/Enterprise/DatasmithCADImporter/Source/ParametricSurfaceExtension/Private/CoreTechBlueprintLibrary.cpp
new file mode 100644
--- /dev/null
+++ b/Engine/Plugins/Enterprise/DatasmithCADImporter/Source/ParametricSurfaceExtension/Private/CoreTechBlueprintLibrary.cpp
@@ -0,0 +1,106 @@
+// Copyright Epic Games, Inc. All Rights Reserved.
+
+#include "CoreTechBlueprintLibrary.h"
+
+#include "CoreTechRetessellateAction.h"
+#include "DatasmithAdditionalData.h"
+#include "DatasmithStaticMeshImporter.h" // Call to BuildStaticMesh
+#include "DatasmithUtils.h"
+
+#include "AssetData.h"
+#include "Engine/StaticMesh.h"
+#include "IStaticMeshEditor.h"
+#include "StaticMeshAttributes.h"
+#include "Toolkits/ToolkitManager.h"
+
+
+#define LOCTEXT_NAMESPACE "CoreTechRetessellateAction"
+
+
+bool UCoreTechBlueprintLibrary::RetessellateStaticMesh(UStaticMesh* StaticMesh, const FDatasmithRetessellationOptions& TessellationSettings, FText& FailureReason)
+{
+	return RetessellateStaticMeshWithNotification(StaticMesh, TessellationSettings, true, FailureReason);
+}
+
+bool UCoreTechBlueprintLibrary::RetessellateStaticMeshWithNotification(UStaticMesh* StaticMesh, const FDatasmithRetessellationOptions& TessellationSettings, bool bApplyChanges, FText& FailureReason)
+{
+	bool bTessellationOutcome = false;
+
+	int32 LODIndex = 0;
+
+	FAssetData AssetData( StaticMesh );
+	if (UCoreTechParametricSurfaceData* CoreTechData = Datasmith::GetAdditionalData<UCoreTechParametricSurfaceData>(AssetData))
+	{
+		// Make sure MeshDescription exists
+		FMeshDescription* DestinationMeshDescription = StaticMesh->GetMeshDescription( LODIndex );
+		if (DestinationMeshDescription == nullptr)
+		{
+			DestinationMeshDescription = StaticMesh->CreateMeshDescription(0);
+		}
+
+		if (DestinationMeshDescription)
+		{
+			if (bApplyChanges)
+			{
+				StaticMesh->Modify();
+				StaticMesh->PreEditChange( nullptr );
+			}
+
+			const int32 OldNumberOfUVChannels = FStaticMeshAttributes(*DestinationMeshDescription).GetVertexInstanceUVs().GetNumChannels();
+			if (FCoreTechRetessellate_Impl::ApplyOnOneAsset( *StaticMesh, *CoreTechData, TessellationSettings ))
+			{
+				const int32 NumberOfUVChannels = FStaticMeshAttributes(*DestinationMeshDescription).GetVertexInstanceUVs().GetNumChannels();
+				if (NumberOfUVChannels < OldNumberOfUVChannels)
+				{
+					FailureReason = FText::Format(NSLOCTEXT("BlueprintRetessellation", "UVChannelsDestroyed", "Tessellation operation on Static Mesh {0} is destroying all UV channels above channel #{1}"), FText::FromString(StaticMesh->GetName()), NumberOfUVChannels - 1);
+				}
+
+				// Post static mesh has changed
+				if (bApplyChanges)
+				{
+					FDatasmithStaticMeshImporter::PreBuildStaticMesh(StaticMesh); // handle uvs stuff
+					FDatasmithStaticMeshImporter::BuildStaticMesh(StaticMesh);
+
+					StaticMesh->PostEditChange();
+					StaticMesh->MarkPackageDirty();
+
+					// Refresh associated editor
+					TSharedPtr<IToolkit> EditingToolkit = FToolkitManager::Get().FindEditorForAsset(StaticMesh);
+					if (IStaticMeshEditor* StaticMeshEditorInUse = StaticCastSharedPtr<IStaticMeshEditor>(EditingToolkit).Get())
+					{
+						StaticMeshEditorInUse->RefreshTool();
+					}
+				}
+				// No posting required, just make sure the new tessellation is committed
+				else
+				{
+					UStaticMesh::FCommitMeshDescriptionParams Params;
+					Params.bMarkPackageDirty = false;
+					Params.bUseHashAsGuid = true;
+					StaticMesh->CommitMeshDescription( LODIndex, Params );
+				}
+
+				// Save last tessellation settings
+				CoreTechData->LastTessellationOptions = TessellationSettings;
+
+				bTessellationOutcome = true;
+			}
+			else
+			{
+				FailureReason = NSLOCTEXT("BlueprintRetessellation", "LoadFailed", "Cannot generate mesh from parametric surface data");
+			}
+		}
+		else
+		{
+			FailureReason = NSLOCTEXT("BlueprintRetessellation", "MeshDescriptionMissing", "Cannot create mesh description");
+		}
+	}
+	else
+	{
+		FailureReason = NSLOCTEXT("BlueprintRetessellation", "MissingData", "No tessellation data attached to the static mesh");
+	}
+
+	return bTessellationOutcome;
+}
+
+#undef LOCTEXT_NAMESPACE
\ No newline at end of file
diff --git a/Engine/Plugins/Enterprise/DatasmithCADImporter/Source/ParametricSurfaceExtension/Private/CoreTechExtensionModule.cpp b/Engine/Plugins/Enterprise/DatasmithCADImporter/Source/ParametricSurfaceExtension/Private/CoreTechExtensionModule.cpp
new file mode 100644
--- /dev/null
+++ b/Engine/Plugins/Enterprise/DatasmithCADImporter/Source/ParametricSurfaceExtension/Private/CoreTechExtensionModule.cpp
@@ -0,0 +1,101 @@
+// Copyright Epic Games, Inc. All Rights Reserved.
+
+#include "CoreTechExtensionModule.h"
+
+#include "CoreTechRetessellateAction.h"
+
+#include "Framework/MultiBox/MultiBoxBuilder.h"
+#include "IStaticMeshEditor.h"
+#include "Modules/ModuleInterface.h"
+#include "Modules/ModuleManager.h"
+#include "StaticMeshEditorModule.h"
+#include "UObject/StrongObjectPtr.h"
+
+
+#define LOCTEXT_NAMESPACE "CoreTechExtensionModule"
+
+
+/** UI extension that displays a Retessellate action in the StaticMeshEditor */
+namespace StaticMeshEditorExtenser
+{
+	bool CanExecute(UStaticMesh* Target)
+	{
+		TArray<FAssetData> AssetData;
+		AssetData.Emplace(Target);
+
+		FCoreTechRetessellate_Impl RetessellateAction;
+		return RetessellateAction.CanApplyOnAssets(AssetData);
+	}
+
+	void Execute(UStaticMesh* Target)
+	{
+		TArray<FAssetData> AssetData;
+		AssetData.Emplace(Target);
+
+		FCoreTechRetessellate_Impl RetessellateAction;
+		RetessellateAction.ApplyOnAssets(AssetData);
+	}
+
+	void ExtendAssetMenu(FMenuBuilder& MenuBuilder, UStaticMesh* Target)
+	{
+		MenuBuilder.AddMenuEntry(
+			FCoreTechRetessellate_Impl::Label,
+			FCoreTechRetessellate_Impl::Tooltip,
+			FSlateIcon(),
+			FUIAction(
+				FExecuteAction::CreateStatic(&StaticMeshEditorExtenser::Execute, Target),
+				FCanExecuteAction::CreateStatic(&StaticMeshEditorExtenser::CanExecute, Target)
+			)
+		);
+	}
+
+	TSharedRef<FExtender> CreateExtenderForObjects(const TSharedRef<FUICommandList> CommandList, const TArray<UObject*> Objects)
+	{
+		TSharedRef<FExtender> Extender = MakeShareable(new FExtender);
+
+		if (UStaticMesh* Target = Objects.Num() ? Cast<UStaticMesh>(Objects[0]) : nullptr)
+		{
+			Extender->AddMenuExtension(
+				"AssetEditorActions",
+				EExtensionHook::Position::After,
+				CommandList,
+				FMenuExtensionDelegate::CreateStatic(&StaticMeshEditorExtenser::ExtendAssetMenu, Target)
+			);
+		}
+
+		return Extender;
+	}
+
+	void Register()
+	{
+		if (!IsRunningCommandlet())
+		{
+			IStaticMeshEditorModule& StaticMeshEditorModule = FModuleManager::Get().LoadModuleChecked<IStaticMeshEditorModule>("StaticMeshEditor");
+			TArray<FAssetEditorExtender>& ExtenderDelegates = StaticMeshEditorModule.GetMenuExtensibilityManager()->GetExtenderDelegates();
+			ExtenderDelegates.Add(FAssetEditorExtender::CreateStatic(&StaticMeshEditorExtenser::CreateExtenderForObjects));
+		}
+	}
+};
+
+FCoreTechExtensionModule& FCoreTechExtensionModule::Get()
+{
+	return FModuleManager::LoadModuleChecked< FCoreTechExtensionModule >(PARAMETRICSURFACEEXTENSION_MODULE_NAME);
+}
+
+bool FCoreTechExtensionModule::IsAvailable()
+{
+	return FModuleManager::Get().IsModuleLoaded(PARAMETRICSURFACEEXTENSION_MODULE_NAME);
+}
+
+void FCoreTechExtensionModule::StartupModule()
+{
+	if (!IsRunningCommandlet())
+	{
+		StaticMeshEditorExtenser::Register();
+	}
+}
+
+IMPLEMENT_MODULE(FCoreTechExtensionModule, CoreTechExtension);
+
+#undef LOCTEXT_NAMESPACE // "DatasmithCoreTechExtensionModule"
+
diff --git a/Engine/Plugins/Enterprise/DatasmithCADImporter/Source/ParametricSurfaceExtension/Private/CoreTechRetessellateAction.cpp b/Engine/Plugins/Enterprise/DatasmithCADImporter/Source/ParametricSurfaceExtension/Private/CoreTechRetessellateAction.cpp
new file mode 100644
--- /dev/null
+++ b/Engine/Plugins/Enterprise/DatasmithCADImporter/Source/ParametricSurfaceExtension/Private/CoreTechRetessellateAction.cpp
@@ -0,0 +1,269 @@
+// Copyright Epic Games, Inc. All Rights Reserved.
+
+#include "CoreTechRetessellateAction.h"
+
+#include "CoreTechHelper.h"
+#include "CoreTechTypes.h"
+
+#include "DatasmithAdditionalData.h"
+#include "DatasmithStaticMeshImporter.h" // Call to BuildStaticMesh
+#include "DatasmithUtils.h"
+#include "DatasmithTranslator.h"
+#include "UI/DatasmithDisplayHelper.h"
+
+#include "Algo/AnyOf.h"
+#include "AssetData.h"
+#include "Async/ParallelFor.h"
+#include "Engine/StaticMesh.h"
+#include "HAL/PlatformFileManager.h"
+#include "IStaticMeshEditor.h"
+#include "StaticMeshAttributes.h"
+#include "Misc/FileHelper.h"
+#include "Misc/ScopedSlowTask.h"
+#include "Toolkits/IToolkit.h"
+#include "Toolkits/ToolkitManager.h"
+#include "UObject/StrongObjectPtr.h"
+#include "Algo/Transform.h"
+
+
+#define LOCTEXT_NAMESPACE "CoreTechRetessellateAction"
+
+
+const FText FCoreTechRetessellate_Impl::Label = LOCTEXT("RetessellateActionLabel", "Retessellate");
+const FText FCoreTechRetessellate_Impl::Tooltip = LOCTEXT("RetessellateActionTooltip", "Tessellate the original NURBS surfaces to re-generate the mesh geometry");
+
+
+bool FCoreTechRetessellate_Impl::CanApplyOnAssets(const TArray<FAssetData>& SelectedAssets)
+{
+	return Algo::AnyOf(SelectedAssets, [](const FAssetData& Asset){ return Datasmith::GetAdditionalData<UCoreTechParametricSurfaceData>(Asset) != nullptr; });
+}
+
+void FCoreTechRetessellate_Impl::ApplyOnAssets(const TArray<FAssetData>& SelectedAssets)
+{
+	TFunction<void(UStaticMesh*)> FinalizeChanges = [](UStaticMesh* StaticMesh) -> void
+	{
+		StaticMesh->PostEditChange();
+		StaticMesh->MarkPackageDirty();
+
+		// Refresh associated editor
+		TSharedPtr<IToolkit> EditingToolkit = FToolkitManager::Get().FindEditorForAsset(StaticMesh);
+		if (IStaticMeshEditor* StaticMeshEditorInUse = StaticCastSharedPtr<IStaticMeshEditor>(EditingToolkit).Get())
+		{
+			StaticMeshEditorInUse->RefreshTool();
+		}
+	};
+
+	TStrongObjectPtr<UCoreTechRetessellateActionOptions> RetessellateOptions = Datasmith::MakeOptions<UCoreTechRetessellateActionOptions>();
+
+	bool bSameOptionsForAll = false;
+	int32 NumAssetsToProcess = SelectedAssets.Num();
+	bool bAskForSameOption = NumAssetsToProcess > 1;
+
+	TArray<UStaticMesh*> TessellatedMeshes;
+	TessellatedMeshes.Reserve( SelectedAssets.Num() );
+
+	TUniquePtr<FScopedSlowTask> Progress;
+	int32 AssetIndex = -1;
+	for (const FAssetData& Asset : SelectedAssets)
+	{
+		AssetIndex++;
+		if (UCoreTechParametricSurfaceData* CoreTechData = Datasmith::GetAdditionalData<UCoreTechParametricSurfaceData>(Asset))
+		{
+			if (CoreTechData->RawData.Num())
+			{
+				if (UStaticMesh* StaticMesh = Cast<UStaticMesh>(Asset.GetAsset()))
+				{
+					if (!bSameOptionsForAll)
+					{
+						Datasmith::FDisplayParameters Parameters;
+						Parameters.bAskForSameOption = bAskForSameOption;
+						Parameters.WindowTitle = LOCTEXT("OptionWindow_WindowTitle", "Datasmith Retessellation Options");
+						Parameters.FileLabel = FText::Format(LOCTEXT("OptionWindow_AssetLabel", "Tessellate StaticMesh: {0}"), FText::FromString(StaticMesh->GetName()));
+						Parameters.FileTooltip = FText::FromString(StaticMesh->GetPathName());
+// 						Parameters.PackageLabel = FText::FromName(StaticMesh->GetOutermost()->FileName);
+						Parameters.ProceedButtonLabel = LOCTEXT("OptionWindow_ProceedButtonLabel", "Tessellate");
+						Parameters.ProceedButtonTooltip = LOCTEXT("OptionWindow_ProceedButtonTooltip", "Retessellate this mesh based on included nurbs data");
+						Parameters.CancelButtonLabel = LOCTEXT("OptionWindow_CancelButtonLabel", "Cancel");
+						Parameters.CancelButtonTooltip = LOCTEXT("OptionWindow_CancelButtonTooltip", "Cancel the retessellation operation");
+
+						bAskForSameOption = false; // ask only the fist time
+						RetessellateOptions->Options = CoreTechData->LastTessellationOptions;
+						Datasmith::FDisplayResult Result = Datasmith::DisplayOptions(RetessellateOptions, Parameters);
+						if (!Result.bValidated)
+						{
+							return;
+						}
+						bSameOptionsForAll |= Result.bUseSameOption;
+					}
+					CoreTechData->LastTessellationOptions = RetessellateOptions->Options;
+
+					int32 RemainingAssetsToProcess = NumAssetsToProcess - AssetIndex;
+					if (bSameOptionsForAll && !Progress.IsValid() && RemainingAssetsToProcess > 1)
+					{
+						Progress = MakeUnique<FScopedSlowTask>(RemainingAssetsToProcess);
+						Progress->MakeDialog(true);
+					}
+					if (Progress)
+					{
+						if (Progress->ShouldCancel())
+						{
+							return;
+						}
+
+						FText Text = FText::Format(LOCTEXT("RetessellateAssetMessage", "Tessellate StaticMesh ({0}/{1}): {2}"),
+							AssetIndex,
+							NumAssetsToProcess,
+							FText::FromString(StaticMesh->GetName())
+						);
+						Progress->EnterProgressFrame(1, Text);
+					}
+
+					if( StaticMesh->GetMeshDescription(0) == nullptr)
+					{
+						StaticMesh->CreateMeshDescription( 0 );
+					}
+
+					if(StaticMesh->GetMeshDescription(0) != nullptr)
+					{
+						StaticMesh->Modify();
+						StaticMesh->PreEditChange( nullptr );
+
+						if( ApplyOnOneAsset(*StaticMesh, *CoreTechData, RetessellateOptions->Options) )
+						{
+							TessellatedMeshes.Add(StaticMesh);
+						}
+						else
+						{
+							FinalizeChanges(StaticMesh);
+						}
+					}
+				}
+			}
+		}
+	}
+
+	// Make sure lightmap settings are valid
+	if(TessellatedMeshes.Num() > 1)
+	{
+		ParallelFor(TessellatedMeshes.Num(), [&](int32 Index)
+		{
+			FDatasmithStaticMeshImporter::PreBuildStaticMesh(TessellatedMeshes[Index]); 
+		});
+	}
+	else if(TessellatedMeshes.Num() > 0)
+	{
+		FDatasmithStaticMeshImporter::PreBuildStaticMesh( TessellatedMeshes[0] );
+	}
+
+	FDatasmithStaticMeshImporter::BuildStaticMeshes( TessellatedMeshes );
+
+	for(UStaticMesh* StaticMesh : TessellatedMeshes)
+	{
+		FinalizeChanges(StaticMesh);
+	}
+}
+
+bool FCoreTechRetessellate_Impl::ApplyOnOneAsset(UStaticMesh& StaticMesh, UCoreTechParametricSurfaceData& CoreTechData, const FDatasmithRetessellationOptions& RetessellateOptions)
+{
+	bool bSuccessfulTessellation = false;
+
+	// make a temporary file as CoreTech can only deal with files.
+	int32 Hash = GetTypeHash(StaticMesh.GetPathName());
+	FString ResourceFile = FPaths::ConvertRelativePathToFull(FPaths::ProjectIntermediateDir() / FString::Printf(TEXT("0x%08x.ct"), Hash));
+
+	FFileHelper::SaveArrayToFile(CoreTechData.RawData, *ResourceFile);
+
+	CADLibrary::FCTMesh Mesh;
+
+	CADLibrary::FImportParameters ImportParameters;
+	ImportParameters.MetricUnit = CoreTechData.SceneParameters.MetricUnit;
+	ImportParameters.ScaleFactor = CoreTechData.SceneParameters.ScaleFactor;
+	ImportParameters.ChordTolerance = RetessellateOptions.ChordTolerance;
+	ImportParameters.MaxEdgeLength = RetessellateOptions.MaxEdgeLength;
+	ImportParameters.MaxNormalAngle = RetessellateOptions.NormalTolerance;
+	ImportParameters.ModelCoordSys = static_cast<FDatasmithUtils::EModelCoordSystem>(CoreTechData.SceneParameters.ModelCoordSys);
+	ImportParameters.StitchingTechnique = CADLibrary::EStitchingTechnique(RetessellateOptions.StitchingTechnique);
+
+	CADLibrary::FMeshParameters MeshParameters;
+	MeshParameters.bNeedSwapOrientation = CoreTechData.MeshParameters.bNeedSwapOrientation;
+	MeshParameters.bIsSymmetric = CoreTechData.MeshParameters.bIsSymmetric;
+	MeshParameters.SymmetricNormal = CoreTechData.MeshParameters.SymmetricNormal;
+	MeshParameters.SymmetricOrigin = CoreTechData.MeshParameters.SymmetricOrigin;
+
+	// Previous MeshDescription is get to be able to create a new one with the same order of PolygonGroup (the matching of color and partition is currently based on their order)
+	if(FMeshDescription* DestinationMeshDescription = StaticMesh.GetMeshDescription(0))
+	{
+		FMeshDescription MeshDescription;
+		FStaticMeshAttributes MeshDescriptionAttributes(MeshDescription);
+		MeshDescriptionAttributes.Register();
+
+		if (RetessellateOptions.RetessellationRule == EDatasmithCADRetessellationRule::SkipDeletedSurfaces)
+		{
+			CADLibrary::CopyPatchGroups(*DestinationMeshDescription, MeshDescription);
+		}
+
+		if ( CADLibrary::LoadFile(ResourceFile, MeshDescription, ImportParameters, MeshParameters))
+		{
+			// To update the SectionInfoMap 
+			{
+				TPolygonGroupAttributesConstRef<FName> MaterialSlotNames = MeshDescriptionAttributes.GetPolygonGroupMaterialSlotNames();
+				FMeshSectionInfoMap& SectionInfoMap = StaticMesh.GetSectionInfoMap();
+
+				for (FPolygonGroupID PolygonGroupID : MeshDescription.PolygonGroups().GetElementIDs())
+				{
+					FMeshSectionInfo Section = SectionInfoMap.Get(0, PolygonGroupID.GetValue());
+					int32 MaterialIndex = StaticMesh.GetMaterialIndex(MaterialSlotNames[PolygonGroupID]);
+					if (MaterialIndex < 0)
+					{
+						MaterialIndex = 0;
+					}
+					Section.MaterialIndex = MaterialIndex;
+					SectionInfoMap.Set(0, PolygonGroupID.GetValue(), Section);
+				}
+			}
+			*DestinationMeshDescription = MoveTemp(MeshDescription);
+			bSuccessfulTessellation = true;
+		}
+	}
+
+	// Remove temporary file
+	FPlatformFileManager::Get().GetPlatformFile().DeleteFile(*ResourceFile);
+
+	return bSuccessfulTessellation;
+}
+
+TSet<UStaticMesh*> GetReferencedStaticMeshes(const TArray<AActor*>& SelectedActors)
+{
+	TSet<UStaticMesh*> ReferencedStaticMeshes;
+
+	for (const AActor* Actor : SelectedActors)
+	{
+		if (Actor)
+		{
+			for (const auto& Component : Actor->GetComponents())
+			{
+				if (UStaticMeshComponent* SMC = Cast<UStaticMeshComponent>(Component))
+				{
+					ReferencedStaticMeshes.Add(SMC->GetStaticMesh());
+				}
+			}
+		}
+	}
+
+	return ReferencedStaticMeshes;
+}
+
+bool UCoreTechRetessellateAction::CanApplyOnActors(const TArray<AActor*>& SelectedActors)
+{
+	const TSet<UStaticMesh*> ReferencedStaticMeshes = GetReferencedStaticMeshes(SelectedActors);
+	return Algo::AnyOf(ReferencedStaticMeshes, [](const UStaticMesh* Mesh){ return Datasmith::GetAdditionalData<UCoreTechParametricSurfaceData>(FAssetData(Mesh)); });
+}
+
+void UCoreTechRetessellateAction::ApplyOnActors(const TArray<AActor*>& SelectedActors)
+{
+	TArray<FAssetData> AssetData;
+	Algo::Transform(GetReferencedStaticMeshes(SelectedActors), AssetData, [](UStaticMesh* Mesh){ return FAssetData(Mesh);});
+	return ApplyOnAssets(AssetData);
+}
+
+#undef LOCTEXT_NAMESPACE
diff --git a/Engine/Plugins/Enterprise/DatasmithCADImporter/Source/ParametricSurfaceExtension/Private/DataprepTessellationOperation.cpp b/Engine/Plugins/Enterprise/DatasmithCADImporter/Source/ParametricSurfaceExtension/Private/DataprepTessellationOperation.cpp
new file mode 100644
--- /dev/null
+++ b/Engine/Plugins/Enterprise/DatasmithCADImporter/Source/ParametricSurfaceExtension/Private/DataprepTessellationOperation.cpp
@@ -0,0 +1,123 @@
+// Copyright Epic Games, Inc. All Rights Reserved.
+
+#include "DataprepTessellationOperation.h"
+
+#include "CoreTechHelper.h"
+#include "CoreTechBlueprintLibrary.h"
+#include "DatasmithAdditionalData.h"
+#include "DatasmithAssetImportData.h"
+#include "DatasmithUtils.h"
+
+#include "IDataprepProgressReporter.h"
+
+#include "ActorEditorUtils.h"
+#include "Engine/StaticMesh.h"
+#include "StaticMeshAttributes.h"
+#include "GenericPlatform/GenericPlatformTime.h"
+#include "Logging/LogMacros.h"
+#include "Misc/FileHelper.h"
+
+#define LOCTEXT_NAMESPACE "DataprepTessellationOperation"
+
+DEFINE_LOG_CATEGORY_STATIC(LogCADLibrary, Log, All)
+
+void UDataprepTessellationOperation::PostLoad()
+{
+	if( HasAnyFlags(RF_WasLoaded) && TessellationSettings_DEPRECATED.ChordTolerance != -MAX_FLT)
+	{
+		ChordTolerance = TessellationSettings_DEPRECATED.ChordTolerance;
+		MaxEdgeLength = TessellationSettings_DEPRECATED.MaxEdgeLength;
+		NormalTolerance = TessellationSettings_DEPRECATED.NormalTolerance;
+		// Mark TessellationSettings_DEPRECATED as non usable
+		TessellationSettings_DEPRECATED.ChordTolerance = -MAX_FLT;
+
+		MarkPackageDirty();
+	}
+
+	Super::PostLoad();
+}
+
+void UDataprepTessellationOperation::OnExecution_Implementation(const FDataprepContext& InContext)
+{
+	// Collect start time to log amount of time spent to import incoming file
+	uint64 StartTime = FPlatformTime::Cycles64();
+
+	TSet<UStaticMesh*> SelectedMeshes;
+
+	for (UObject* Object : InContext.Objects)
+	{
+		if ( UStaticMesh* StaticMesh = Cast<UStaticMesh>(Object) )
+		{
+			SelectedMeshes.Add( StaticMesh );
+		}
+		else if (AActor* Actor = Cast<AActor>(Object) )
+		{
+			TInlineComponentArray<UStaticMeshComponent*> StaticMeshComponents( Actor );
+			for (UStaticMeshComponent* StaticMeshComponent : StaticMeshComponents)
+			{
+				if((StaticMesh = StaticMeshComponent->GetStaticMesh()) != nullptr)
+				{
+					SelectedMeshes.Add( StaticMesh );
+				}
+			}
+		}
+	}
+
+	if(	!IsCancelled() && SelectedMeshes.Num() > 0)
+	{
+		FDatasmithRetessellationOptions TessellationSettings( ChordTolerance, MaxEdgeLength, NormalTolerance);
+
+		TSharedPtr<FDataprepWorkReporter> Task = CreateTask( LOCTEXT( "LogCADLibrary_Tessellating", "Tessellating meshes ..." ), (float)SelectedMeshes.Num() );
+
+		TArray<UObject*> ModifiedStaticMeshes;
+		ModifiedStaticMeshes.Reserve( SelectedMeshes.Num() );
+
+		for( UStaticMesh* StaticMesh : SelectedMeshes )
+		{
+			if( IsCancelled() )
+			{
+				break;
+			}
+
+			Task->ReportNextStep( FText::Format( LOCTEXT( "LogCADLibrary_Tessellating_One_Mesh", "Tessellating {0} ..."), FText::FromString( StaticMesh->GetName() ) ) );
+
+			if( StaticMesh->IsMeshDescriptionValid( 0 ) )
+			{
+				FText OutReason;
+				if( UCoreTechBlueprintLibrary::RetessellateStaticMeshWithNotification( StaticMesh, TessellationSettings, false, OutReason ) )
+				{
+					ModifiedStaticMeshes.Add( StaticMesh );
+					if (!OutReason.IsEmpty())
+					{
+						FText WarningMsg = FText::Format(LOCTEXT("DataprepTessellationOperation_TessellationCompletedWithWarning", "{0}"), OutReason);
+						LogInfo(WarningMsg);
+					}
+				}
+				else
+				{
+					FText WarningMsg = FText::Format( LOCTEXT( "DataprepTessellationOperation_TessellationFailed", "{0}" ), OutReason );
+					LogWarning( WarningMsg );
+				}
+			}
+			else
+			{
+				FText WarningMsg = FText::Format( LOCTEXT( "DataprepTessellationOperation_EmptyMesh", "No triangles in static mesh {0}" ), FText::FromString( StaticMesh->GetName() ) );
+				LogWarning( WarningMsg );
+			}
+		}
+
+		// Log time spent to import incoming file in minutes and seconds
+		double ElapsedSeconds = FPlatformTime::ToSeconds64(FPlatformTime::Cycles64() - StartTime);
+
+		int ElapsedMin = int(ElapsedSeconds / 60.0);
+		ElapsedSeconds -= 60.0 * (double)ElapsedMin;
+		UE_LOG( LogCADLibrary, Log, TEXT("Tessellation of %d out of %d static mesh(es) in [%d min %.3f s]"), ModifiedStaticMeshes.Num(), SelectedMeshes.Num(), ElapsedMin, ElapsedSeconds );
+
+		if(ModifiedStaticMeshes.Num() > 0)
+		{
+			AssetsModified( MoveTemp( ModifiedStaticMeshes ) );
+		}
+	}
+}
+
+#undef LOCTEXT_NAMESPACE
diff --git a/Engine/Plugins/Enterprise/DatasmithCADImporter/Source/ParametricSurfaceExtension/Private/DataprepTessellationOperation.h b/Engine/Plugins/Enterprise/DatasmithCADImporter/Source/ParametricSurfaceExtension/Private/DataprepTessellationOperation.h
new file mode 100644
--- /dev/null
+++ b/Engine/Plugins/Enterprise/DatasmithCADImporter/Source/ParametricSurfaceExtension/Private/DataprepTessellationOperation.h
@@ -0,0 +1,80 @@
+// Copyright Epic Games, Inc. All Rights Reserved.
+
+#pragma once
+
+#include "CoreMinimal.h"
+
+#include "DataprepOperation.h"
+#include "DatasmithImportOptions.h"
+#include "DatasmithUtils.h"
+
+#include "DataprepTessellationOperation.generated.h"
+
+namespace CADLibrary
+{
+	struct FImportParameters;
+}
+struct FMeshDescription;
+class UCoreTechParametricSurfaceData;
+class UDatasmithStaticMeshCADImportData;
+
+UCLASS(Experimental, Category = MeshOperation, Meta = (DisplayName="Datasmith Tessellation", ToolTip = "For each static mesh to process, retessellate the mesh if the object contains the required data") )
+class UDataprepTessellationOperation : public UDataprepOperation
+{
+	GENERATED_BODY()
+
+	UDataprepTessellationOperation()
+		: ChordTolerance(0.2f)
+		, MaxEdgeLength(0.0f)
+		, NormalTolerance(20.0f)
+		// Mark TessellationSettings_DEPRECATED as non usable
+		, TessellationSettings_DEPRECATED(-MAX_FLT)
+	{
+	}
+
+	// UObject interface
+	/** Overridden for backward compatibility with 4.23. To be removed in 4.25 */
+	virtual void PostLoad() override;
+	// End of UObject interface
+
+public:
+	/**
+	 * Maximum distance between any point on a triangle generated by the tessellation process and the actual surface.
+	 * The lower the value the more triangles.
+	 * Default value is 0.2.
+	 */
+	UPROPERTY(EditAnywhere, BlueprintReadWrite, Category = "Tessellation Options", meta = (Units = cm, ToolTip = "Maximum distance between any generated triangle and the original surface. Smaller values make more triangles.", ClampMin = "0.0"))
+	float ChordTolerance;
+
+	/**
+	 * Maximum length of edges of triangles generated by the tessellation process.
+	 * The length is in scene/model unit. The smaller the more triangles are generated.
+	 * Value of 0 means no constraint on length of edges
+	 * Default value is 0.
+	 */
+	UPROPERTY(EditAnywhere, BlueprintReadWrite, AdvancedDisplay, Category = "Tessellation Options", meta = (Units = cm, DisplayName = "Max Edge Length", ToolTip = "Maximum length of any edge in the generated triangles. Smaller values make more triangles.", ClampMin = "0.0"))
+	float MaxEdgeLength;
+
+	/**
+	 * Maximum angle between the normal of two triangles generated by the tessellation process.
+	 * The angle is expressed in degree. The smaller the more triangles are generated.
+	 * Default value is 20 degrees.
+	 */
+	UPROPERTY(EditAnywhere, BlueprintReadWrite, AdvancedDisplay, Category = "Tessellation Options", meta = (Units = deg, ToolTip = "Maximum angle between adjacent triangles. Smaller values make more triangles.", ClampMin = "0.0", ClampMax = "90.0"))
+	float NormalTolerance;
+
+	/** Version 4.24 : Deprecated - not serialized anymore */
+	UPROPERTY()
+	FDatasmithTessellationOptions TessellationSettings_DEPRECATED;
+
+	//~ Begin UDataprepOperation Interface
+public:
+	virtual FText GetCategory_Implementation() const override
+	{
+		return FDataprepOperationCategories::MeshOperation;
+	}
+
+protected:
+	virtual void OnExecution_Implementation(const FDataprepContext& InContext) override;
+	//~ End UDataprepOperation Interface
+};
diff --git a/Engine/Plugins/Enterprise/DatasmithCADImporter/Source/ParametricSurfaceExtension/Public/CoreTechBlueprintLibrary.h b/Engine/Plugins/Enterprise/DatasmithCADImporter/Source/ParametricSurfaceExtension/Public/CoreTechBlueprintLibrary.h
new file mode 100644
--- /dev/null
+++ b/Engine/Plugins/Enterprise/DatasmithCADImporter/Source/ParametricSurfaceExtension/Public/CoreTechBlueprintLibrary.h
@@ -0,0 +1,44 @@
+// Copyright Epic Games, Inc. All Rights Reserved.
+
+#pragma once
+
+#include "CoreMinimal.h"
+
+#include "DatasmithImportOptions.h"
+
+#include "Kismet/BlueprintFunctionLibrary.h"
+#include "UObject/TextProperty.h"
+
+#include "CoreTechBlueprintLibrary.generated.h"
+
+class UStaticMesh;
+
+UCLASS(Blueprintable, BlueprintType, meta = (DisplayName = "CAD Surface Operations Library"))
+class PARAMETRICSURFACEEXTENSION_API UCoreTechBlueprintLibrary : public UBlueprintFunctionLibrary
+{
+	GENERATED_BODY()
+
+public:
+	/**
+	 * Re-tessellate LOD 0 of a static mesh if it contains parametric surface data.
+	 * @param	StaticMesh				Static mesh to re-tessellate.
+	 * @param	TessellationSettings	Tessellation settings to use.
+	 * @param	FailureReason			Text describing the reason of failure.
+	 * @return True if successful, false otherwise
+	 */
+	UFUNCTION(BlueprintCallable, Category = "Datasmith | Surface Operations" )
+	static bool RetessellateStaticMesh(UStaticMesh* StaticMesh, const FDatasmithRetessellationOptions& TessellationSettings, FText& FailureReason);
+
+	/**
+	 * Re-tessellate LOD 0 of a static mesh if it contains parametric surface data.
+	 * This implementations allows to skip post edit operations (bApplyChanges=false),
+	 * the caller is then responsible to handle those operations.
+	 * @param	StaticMesh				Static mesh to re-tessellate.
+	 * @param	TessellationSettings	Tessellation settings to use.
+	 * @param	bApplyChanges			Indicates if change must be notified.
+	 * @param	FailureReason			Text describing the reason of failure, or a warning if the operation was successful.
+	 * @return True if successful, false otherwise
+	 */
+	UFUNCTION(Category = "Datasmith | Surface Operations" )
+	static bool RetessellateStaticMeshWithNotification(UStaticMesh* StaticMesh, const FDatasmithRetessellationOptions& TessellationSettings, bool bApplyChanges, FText& FailureReason);
+};
diff --git a/Engine/Plugins/Enterprise/DatasmithCADImporter/Source/ParametricSurfaceExtension/Public/CoreTechExtensionModule.h b/Engine/Plugins/Enterprise/DatasmithCADImporter/Source/ParametricSurfaceExtension/Public/CoreTechExtensionModule.h
new file mode 100644
--- /dev/null
+++ b/Engine/Plugins/Enterprise/DatasmithCADImporter/Source/ParametricSurfaceExtension/Public/CoreTechExtensionModule.h
@@ -0,0 +1,20 @@
+// Copyright Epic Games, Inc. All Rights Reserved.
+
+#pragma once
+
+#include "Modules/ModuleInterface.h"
+
+#define PARAMETRICSURFACEEXTENSION_MODULE_NAME TEXT("ParametricSurfaceExtension")
+
+/**
+ * This module exposes additional features for assets containing CoreTech data.
+ */
+class FCoreTechExtensionModule : public IModuleInterface
+{
+public:
+    static FCoreTechExtensionModule& Get();
+    static bool IsAvailable();
+
+private:
+	virtual void StartupModule() override;
+};
diff --git a/Engine/Plugins/Enterprise/DatasmithCADImporter/Source/ParametricSurfaceExtension/Public/CoreTechRetessellateAction.h b/Engine/Plugins/Enterprise/DatasmithCADImporter/Source/ParametricSurfaceExtension/Public/CoreTechRetessellateAction.h
new file mode 100644
--- /dev/null
+++ b/Engine/Plugins/Enterprise/DatasmithCADImporter/Source/ParametricSurfaceExtension/Public/CoreTechRetessellateAction.h
@@ -0,0 +1,55 @@
+// Copyright Epic Games, Inc. All Rights Reserved.
+
+#pragma once
+
+#include "CoreMinimal.h"
+
+#include "CoreTechSurfaceExtension.h"
+#include "DatasmithCustomAction.h"
+#include "DatasmithAdditionalData.h"
+#include "DatasmithImportOptions.h"
+
+#include "CoreTechRetessellateAction.generated.h"
+
+class PARAMETRICSURFACEEXTENSION_API FCoreTechRetessellate_Impl
+{
+public:
+	static const FText Label;
+	static const FText Tooltip;
+
+	static bool CanApplyOnAssets(const TArray<FAssetData>& SelectedAssets);
+	static void ApplyOnAssets(const TArray<FAssetData>& SelectedAssets);
+	static bool ApplyOnOneAsset(UStaticMesh& StaticMesh, UCoreTechParametricSurfaceData& CoreTechData, const FDatasmithRetessellationOptions& RetesselateOptions);
+};
+
+
+UCLASS()
+class PARAMETRICSURFACEEXTENSION_API UCoreTechRetessellateAction : public UDatasmithCustomActionBase
+{
+	GENERATED_BODY()
+
+public:
+	virtual const FText& GetLabel() override { return FCoreTechRetessellate_Impl::Label; }
+	virtual const FText& GetTooltip() override { return FCoreTechRetessellate_Impl::Tooltip; }
+
+	virtual bool CanApplyOnAssets(const TArray<FAssetData>& SelectedAssets) override
+	{ return FCoreTechRetessellate_Impl::CanApplyOnAssets(SelectedAssets); }
+
+	virtual void ApplyOnAssets(const TArray<FAssetData>& SelectedAssets) override
+	{ return FCoreTechRetessellate_Impl::ApplyOnAssets(SelectedAssets); }
+
+	virtual bool CanApplyOnActors(const TArray<AActor*>& SelectedActors) override;
+
+	virtual void ApplyOnActors(const TArray<AActor*>& SelectedActors) override;
+};
+
+
+UCLASS(config = Editor, Transient)
+class PARAMETRICSURFACEEXTENSION_API UCoreTechRetessellateActionOptions : public UObject
+{
+	GENERATED_BODY()
+
+public:
+	UPROPERTY(config, EditAnywhere, BlueprintReadWrite, Category = "NotVisible", meta = (ShowOnlyInnerProperties))
+	FDatasmithRetessellationOptions Options;
+};
