CADKernel Retesselation implementation (Step2)
CoretechExtension become ParametricSurfaceExtension (Step1)
CoretechSurface module will derived from the new ParametricSurface module
the new CADKernelSurface module will derived from the new ParametricSurface module and will manage data for CADKernel

Step 2:
- Add ParametricSurface Module
- Add CADKernelSurface Module derived from ParametricSurface
- Rename ParametricSurfaceExtension files

Next Step: finalize implementation of CADKernelSurface

diff --git a/Engine/Plugins/Enterprise/DatasmithCADImporter/Config/BaseDatasmithCADImporter.ini b/Engine/Plugins/Enterprise/DatasmithCADImporter/Config/BaseDatasmithCADImporter.ini
--- a/Engine/Plugins/Enterprise/DatasmithCADImporter/Config/BaseDatasmithCADImporter.ini
+++ b/Engine/Plugins/Enterprise/DatasmithCADImporter/Config/BaseDatasmithCADImporter.ini
@@ -1,8 +1,15 @@
 [CoreRedirects]
 
-+ClassRedirects=(OldName="/Script/DatasmithCoreTechExtension.CoreTechBlueprintLibrary", NewName="/Script/CoreTechExtension.CoreTechBlueprintLibrary")
-+ClassRedirects=(OldName="/Script/DatasmithCoreTechExtension.CoreTechRetessellateAction", NewName="/Script/CoreTechExtension.CoreTechRetessellateAction")
-+ClassRedirects=(OldName="/Script/DatasmithCoreTechExtension.CoreTechRetessellateActionOptions", NewName="/Script/CoreTechExtension.CoreTechRetessellateActionOptions")
-+ClassRedirects=(OldName="/Script/DatasmithCoreTechExtension.DataprepTessellationOperation", NewName="/Script/CoreTechExtension.DataprepTessellationOperation")
++ClassRedirects=(OldName="/Script/DatasmithCoreTechExtension.CoreTechBlueprintLibrary", NewName="/Script/ParametricSurfaceExtension.ParametricSurfaceBlueprintLibrary")
++ClassRedirects=(OldName="/Script/CoreTechExtension.CoreTechBlueprintLibrary", NewName="/Script/ParametricSurfaceExtension.ParametricSurfaceBlueprintLibrary")
+
++ClassRedirects=(OldName="/Script/DatasmithCoreTechExtension.CoreTechRetessellateAction", NewName="/Script/ParametricSurfaceExtension.ParametricRetessellateAction")
++ClassRedirects=(OldName="/Script/CoreTechExtension.CoreTechRetessellateAction", NewName="/Script/ParametricSurfaceExtension.ParametricRetessellateAction")
+
++ClassRedirects=(OldName="/Script/DatasmithCoreTechExtension.CoreTechRetessellateActionOptions", NewName="/Script/ParametricSurfaceExtension.ParametricRetessellateActionOptions")
++ClassRedirects=(OldName="/Script/CoreTechExtension.CoreTechRetessellateActionOptions", NewName="/Script/ParametricSurfaceExtension.ParametricRetessellateActionOptions")
+
++ClassRedirects=(OldName="/Script/DatasmithCoreTechExtension.DataprepTessellationOperation", NewName="/Script/ParametricSurfaceExtension.DataprepTessellationOperation")
++ClassRedirects=(OldName="/Script/CoreTechExtension.DataprepTessellationOperation", NewName="/Script/ParametricSurfaceExtension.DataprepTessellationOperation")
 
 
diff --git a/Engine/Plugins/Enterprise/DatasmithCADImporter/DatasmithCADImporter.uplugin b/Engine/Plugins/Enterprise/DatasmithCADImporter/DatasmithCADImporter.uplugin
--- a/Engine/Plugins/Enterprise/DatasmithCADImporter/DatasmithCADImporter.uplugin
+++ b/Engine/Plugins/Enterprise/DatasmithCADImporter/DatasmithCADImporter.uplugin
@@ -81,6 +81,22 @@
 				"Win64"
 			]
 		},
+		{
+			"Name": "ParametricSurface",
+			"Type": "Runtime",
+			"LoadingPhase": "Default",
+			"WhitelistPlatforms": [
+				"Win64"
+			]
+		},
+		{
+			"Name": "CADKernelSurface",
+			"Type": "Runtime",
+			"LoadingPhase": "Default",
+			"WhitelistPlatforms": [
+				"Win64"
+			]
+		},
 		{
 			"Name": "CoreTechSurface",
 			"Type": "Runtime",
diff --git a/Engine/Plugins/Enterprise/DatasmithCADImporter/Source/CADKernelSurface/CADKernelSurface.Build.cs b/Engine/Plugins/Enterprise/DatasmithCADImporter/Source/CADKernelSurface/CADKernelSurface.Build.cs
new file mode 100644
--- /dev/null
+++ b/Engine/Plugins/Enterprise/DatasmithCADImporter/Source/CADKernelSurface/CADKernelSurface.Build.cs
@@ -0,0 +1,35 @@
+// Copyright Epic Games, Inc. All Rights Reserved.
+
+using System.IO;
+
+namespace UnrealBuildTool.Rules
+{
+	public class CADKernelSurface : ModuleRules
+	{
+		public CADKernelSurface(ReadOnlyTargetRules Target) : base(Target)
+		{
+			PublicDependencyModuleNames.AddRange(
+				new string[]
+				{
+					"Core",
+				}
+			);
+
+			PrivateDependencyModuleNames.AddRange(
+				new string[]
+				{
+					"CADLibrary",
+					"CADTools",
+					"CoreUObject",
+					"DatasmithContent",
+					"DatasmithCore",
+					"DatasmithTranslator",
+					"Engine",
+					"MeshDescription",
+					"ParametricSurface",
+					"StaticMeshDescription",
+				}
+			);
+		}
+	}
+}
diff --git a/Engine/Plugins/Enterprise/DatasmithCADImporter/Source/CADKernelSurface/Private/CADKernelSurfaceExtension.cpp b/Engine/Plugins/Enterprise/DatasmithCADImporter/Source/CADKernelSurface/Private/CADKernelSurfaceExtension.cpp
new file mode 100644
--- /dev/null
+++ b/Engine/Plugins/Enterprise/DatasmithCADImporter/Source/CADKernelSurface/Private/CADKernelSurfaceExtension.cpp
@@ -0,0 +1,105 @@
+// Copyright Epic Games, Inc. All Rights Reserved.
+
+#include "CADKernelSurfaceExtension.h"
+
+#include "CADOptions.h"
+#include "DatasmithPayload.h"
+#include "Engine/StaticMesh.h"
+#include "MeshDescription.h"
+#include "MeshDescriptionHelper.h"
+#include "StaticMeshAttributes.h"
+#include "UObject/EnterpriseObjectVersion.h"
+
+void UCADKernelParametricSurfaceData::Serialize(FArchive& Ar)
+{
+	Ar.UsingCustomVersion(FEnterpriseObjectVersion::GUID);
+	Super::Serialize(Ar);
+	Ar << RawData;
+}
+
+bool UCADKernelParametricSurfaceData::Tessellate(UStaticMesh& StaticMesh, const FDatasmithRetessellationOptions& RetessellateOptions)
+{
+	bool bSuccessfulTessellation = false;
+
+#if WITH_EDITOR
+	CADLibrary::FImportParameters ImportParameters;
+	ImportParameters.MetricUnit = SceneParameters.MetricUnit;
+	ImportParameters.ScaleFactor = SceneParameters.ScaleFactor;
+	ImportParameters.ChordTolerance = RetessellateOptions.ChordTolerance;
+	ImportParameters.MaxEdgeLength = RetessellateOptions.MaxEdgeLength;
+	ImportParameters.MaxNormalAngle = RetessellateOptions.NormalTolerance;
+	ImportParameters.ModelCoordSys = static_cast<FDatasmithUtils::EModelCoordSystem>(SceneParameters.ModelCoordSys);
+	ImportParameters.StitchingTechnique = CADLibrary::EStitchingTechnique(RetessellateOptions.StitchingTechnique);
+
+	CADLibrary::FMeshParameters CadMeshParameters;
+	CadMeshParameters.bNeedSwapOrientation = MeshParameters.bNeedSwapOrientation;
+	CadMeshParameters.bIsSymmetric = MeshParameters.bIsSymmetric;
+	CadMeshParameters.SymmetricNormal = MeshParameters.SymmetricNormal;
+	CadMeshParameters.SymmetricOrigin = MeshParameters.SymmetricOrigin;
+
+	// Previous MeshDescription is get to be able to create a new one with the same order of PolygonGroup (the matching of color and partition is currently based on their order)
+	if (FMeshDescription* DestinationMeshDescription = StaticMesh.GetMeshDescription(0))
+	{
+		FMeshDescription MeshDescription;
+		FStaticMeshAttributes MeshDescriptionAttributes(MeshDescription);
+		MeshDescriptionAttributes.Register();
+
+		if (RetessellateOptions.RetessellationRule == EDatasmithCADRetessellationRule::SkipDeletedSurfaces)
+		{
+			CADLibrary::CopyPatchGroups(*DestinationMeshDescription, MeshDescription);
+		}
+
+		if (false /*CADLibrary::LoadFile(ResourceFile, MeshDescription, ImportParameters, CadMeshParameters)*/)
+		{
+			// To update the SectionInfoMap 
+			{
+				TPolygonGroupAttributesConstRef<FName> MaterialSlotNames = MeshDescriptionAttributes.GetPolygonGroupMaterialSlotNames();
+				FMeshSectionInfoMap& SectionInfoMap = StaticMesh.GetSectionInfoMap();
+
+				for (FPolygonGroupID PolygonGroupID : MeshDescription.PolygonGroups().GetElementIDs())
+				{
+					FMeshSectionInfo Section = SectionInfoMap.Get(0, PolygonGroupID.GetValue());
+					int32 MaterialIndex = StaticMesh.GetMaterialIndex(MaterialSlotNames[PolygonGroupID]);
+					if (MaterialIndex < 0)
+					{
+						MaterialIndex = 0;
+					}
+					Section.MaterialIndex = MaterialIndex;
+					SectionInfoMap.Set(0, PolygonGroupID.GetValue(), Section);
+				}
+			}
+			*DestinationMeshDescription = MoveTemp(MeshDescription);
+			bSuccessfulTessellation = true;
+		}
+	}
+#endif 
+
+	return bSuccessfulTessellation;
+}
+
+namespace CADKernelSurface
+{
+	void AddSurfaceDataForMesh(const TSharedRef<IDatasmithMeshElement>& InMeshElement, const CADLibrary::FImportParameters& InSceneParameters, const CADLibrary::FMeshParameters& InMeshParameters, const FDatasmithTessellationOptions& InTessellationOptions, FDatasmithMeshElementPayload& OutMeshPayload)
+	{
+		// Store CADKernel additional data if provided
+		// TODO
+		TArray<uint8> ByteArray;
+		if (ByteArray.Num() > 0)
+		{
+			UCADKernelParametricSurfaceData* CADKernelData = Datasmith::MakeAdditionalData<UCADKernelParametricSurfaceData>();
+			CADKernelData->RawData = MoveTemp(ByteArray);
+			CADKernelData->SceneParameters.ModelCoordSys = uint8(InSceneParameters.ModelCoordSys);
+			CADKernelData->SceneParameters.MetricUnit = InSceneParameters.MetricUnit;
+			CADKernelData->SceneParameters.ScaleFactor = InSceneParameters.ScaleFactor;
+
+			CADKernelData->MeshParameters.bNeedSwapOrientation = InMeshParameters.bNeedSwapOrientation;
+			CADKernelData->MeshParameters.bIsSymmetric = InMeshParameters.bIsSymmetric;
+			CADKernelData->MeshParameters.SymmetricNormal = InMeshParameters.SymmetricNormal;
+			CADKernelData->MeshParameters.SymmetricOrigin = InMeshParameters.SymmetricOrigin;
+
+			CADKernelData->LastTessellationOptions = InTessellationOptions;
+			OutMeshPayload.AdditionalData.Add(CADKernelData);
+		}
+	}
+
+}
diff --git a/Engine/Plugins/Enterprise/DatasmithCADImporter/Source/CADKernelSurface/Private/CADKernelSurfaceModule.cpp b/Engine/Plugins/Enterprise/DatasmithCADImporter/Source/CADKernelSurface/Private/CADKernelSurfaceModule.cpp
new file mode 100644
--- /dev/null
+++ b/Engine/Plugins/Enterprise/DatasmithCADImporter/Source/CADKernelSurface/Private/CADKernelSurfaceModule.cpp
@@ -0,0 +1,28 @@
+// Copyright Epic Games, Inc. All Rights Reserved.
+
+#include "CADKernelSurfaceModule.h"
+
+#include "Modules/ModuleInterface.h"
+#include "Modules/ModuleManager.h"
+#include "UObject/CoreRedirects.h"
+
+#define LOCTEXT_NAMESPACE "CADKernelSurfaceModule"
+
+void FCADKernelSurfaceModule::StartupModule()
+{
+}
+
+FCADKernelSurfaceModule& FCADKernelSurfaceModule::Get()
+{
+	return FModuleManager::LoadModuleChecked< FCADKernelSurfaceModule >(CADKERNELSURFACE_MODULE_NAME);
+}
+
+bool FCADKernelSurfaceModule::IsAvailable()
+{
+	return FModuleManager::Get().IsModuleLoaded(CADKERNELSURFACE_MODULE_NAME);
+}
+
+IMPLEMENT_MODULE(FCADKernelSurfaceModule, CADKernelSurface);
+
+#undef LOCTEXT_NAMESPACE
+
diff --git a/Engine/Plugins/Enterprise/DatasmithCADImporter/Source/CADKernelSurface/Public/CADKernelSurfaceExtension.h b/Engine/Plugins/Enterprise/DatasmithCADImporter/Source/CADKernelSurface/Public/CADKernelSurfaceExtension.h
new file mode 100644
--- /dev/null
+++ b/Engine/Plugins/Enterprise/DatasmithCADImporter/Source/CADKernelSurface/Public/CADKernelSurfaceExtension.h
@@ -0,0 +1,51 @@
+// Copyright Epic Games, Inc. All Rights Reserved.
+
+#pragma once
+
+#include "CoreMinimal.h"
+
+#include "ParametricSurfaceExtension.h"
+#include "DatasmithAdditionalData.h"
+#include "DatasmithCustomAction.h"
+#include "DatasmithImportOptions.h"
+#include "DatasmithUtils.h"
+
+#include "CADKernelSurfaceExtension.generated.h"
+
+UCLASS(meta = (DisplayName = "CADKernel Parametric Surface Data"))
+class CADKERNELSURFACE_API UCADKernelParametricSurfaceData : public UParametricSurfaceData
+{
+	GENERATED_BODY()
+
+public:
+	UPROPERTY()
+	FString SourceFile;
+
+	// Too costly to serialize as a UPROPERTY, will use custom serialization.
+	TArray<uint8> RawData;
+
+	virtual bool IsValid() override
+	{
+		return RawData.Num() > 0;
+	}
+	
+	virtual bool Tessellate(UStaticMesh& StaticMesh, const FDatasmithRetessellationOptions& RetessellateOptions) override;
+
+private:
+
+	virtual void Serialize(FArchive& Ar) override;
+};
+
+class IDatasmithMeshElement;
+struct FDatasmithMeshElementPayload;
+
+namespace CADLibrary
+{
+	struct FImportParameters;
+	struct FMeshParameters;
+}
+
+namespace CADKernelSurface
+{
+	void CADKERNELSURFACE_API AddSurfaceDataForMesh(const TSharedRef<IDatasmithMeshElement>& InMeshElement, const CADLibrary::FImportParameters& InSceneParameters, const CADLibrary::FMeshParameters&, const FDatasmithTessellationOptions& InTessellationOptions, FDatasmithMeshElementPayload& OutMeshPayload);
+}
diff --git a/Engine/Plugins/Enterprise/DatasmithCADImporter/Source/CADKernelSurface/Public/CADKernelSurfaceModule.h b/Engine/Plugins/Enterprise/DatasmithCADImporter/Source/CADKernelSurface/Public/CADKernelSurfaceModule.h
new file mode 100644
--- /dev/null
+++ b/Engine/Plugins/Enterprise/DatasmithCADImporter/Source/CADKernelSurface/Public/CADKernelSurfaceModule.h
@@ -0,0 +1,19 @@
+// Copyright Epic Games, Inc. All Rights Reserved.
+
+#pragma once
+
+#include "Modules/ModuleInterface.h"
+
+#define CADKERNELSURFACE_MODULE_NAME TEXT("CADKernelSurface")
+
+/**
+ * This module exposes additional features for assets containing CADKernel data.
+ */
+class FCADKernelSurfaceModule : public IModuleInterface
+{
+public:
+	virtual void StartupModule() override;
+
+	static FCADKernelSurfaceModule& Get();
+	static bool IsAvailable();
+};
diff --git a/Engine/Plugins/Enterprise/DatasmithCADImporter/Source/CADLibrary/CADLibrary.Build.cs b/Engine/Plugins/Enterprise/DatasmithCADImporter/Source/CADLibrary/CADLibrary.Build.cs
--- a/Engine/Plugins/Enterprise/DatasmithCADImporter/Source/CADLibrary/CADLibrary.Build.cs
+++ b/Engine/Plugins/Enterprise/DatasmithCADImporter/Source/CADLibrary/CADLibrary.Build.cs
@@ -14,7 +14,6 @@ public class CADLibrary : ModuleRules
 			new string[]
 			{
 				"Core",
-				"CADInterfaces",
 				"CADTools",
 				"DatasmithCore",
 				"MeshDescription",
diff --git a/Engine/Plugins/Enterprise/DatasmithCADImporter/Source/CADLibrary/Private/CoreTechHelper.cpp b/Engine/Plugins/Enterprise/DatasmithCADImporter/Source/CADLibrary/Private/CoreTechHelper.cpp
deleted file mode 100644
diff --git a/Engine/Plugins/Enterprise/DatasmithCADImporter/Source/CADLibrary/Private/MeshDescriptionHelper.cpp b/Engine/Plugins/Enterprise/DatasmithCADImporter/Source/CADLibrary/Private/MeshDescriptionHelper.cpp
new file mode 100644
--- /dev/null
+++ b/Engine/Plugins/Enterprise/DatasmithCADImporter/Source/CADLibrary/Private/MeshDescriptionHelper.cpp
@@ -0,0 +1,621 @@
+// Copyright Epic Games, Inc. All Rights Reserved.
+#include "MeshDescriptionHelper.h"
+
+#include "CADData.h"
+#include "CADOptions.h"
+#include "DatasmithMaterialElements.h"
+#include "DatasmithSceneFactory.h"
+#include "DatasmithUtils.h"
+#include "Math.h"
+#include "MeshDescription.h"
+#include "MeshOperator.h"
+#include "Misc/FileHelper.h"
+#include "StaticMeshAttributes.h"
+#include "StaticMeshOperations.h"
+
+typedef uint32 TriangleIndex[3];
+
+namespace CADLibrary
+{
+	struct FVertexData
+	{
+		float Z;
+		int32 Index;
+		FVector Coordinates;
+		bool bIsMerged;
+		FVertexID VertexID;
+		FVertexID SymVertexID;
+
+		/** Default constructor. */
+		FVertexData() {}
+
+		/** Initialization constructor. */
+		FVertexData(int32 InIndex, const FVector& V)
+		{
+			//Z = V.X + V.Y + V.Z;
+			Z = 0.30f * V.X + 0.33f * V.Y + 0.37f * V.Z;
+			Index = InIndex;
+			Coordinates = V;
+			bIsMerged = false;
+			VertexID = INDEX_NONE;
+		}
+	};
+
+	struct FCompareVertexZ
+	{
+		FORCEINLINE bool operator()(FVertexData const& A, FVertexData const& B) const { return A.Z < B.Z; }
+	};
+
+	// Verify the 3 input indices are not defining a degenerated triangle and fill up the corresponding FVertexIDs
+	bool IsTriangleDegenerated(const int32_t* Indices, const TArray<FVertexID>& RemapVertexPosition, FVertexID VertexIDs[3])
+	{
+		if (Indices[0] == Indices[1] || Indices[0] == Indices[2] || Indices[1] == Indices[2])
+		{
+			return true;
+		}
+
+		for (int32 Corner = 0; Corner < 3; ++Corner)
+		{
+			VertexIDs[Corner] = RemapVertexPosition[Indices[Corner]];
+		}
+
+		return (VertexIDs[0] == VertexIDs[1] || VertexIDs[0] == VertexIDs[2] || VertexIDs[1] == VertexIDs[2]);
+	}
+
+	void FillKioVertexPosition(const FImportParameters& ImportParams, const FMeshParameters& MeshParameters, FBodyMesh& Body, FMeshDescription& MeshDescription)
+	{
+		int32 TriangleCount = Body.TriangleCount;
+		TArray<FTessellationData>& FaceTessellationSet = Body.Faces;
+
+		// Add offset to the bounding box to avoid to remove good vertex
+		FVector Size = Body.BBox.GetSize();
+		FBox BBox = Body.BBox.ExpandBy(Size.Size());
+		BBox.IsValid = Body.BBox.IsValid;
+		BBox.Min *= ImportParams.ScaleFactor;
+		BBox.Max *= ImportParams.ScaleFactor;
+
+		TVertexAttributesRef<FVector3f> VertexPositions = MeshDescription.GetVertexPositions();
+
+		// Create a list of vertex Z/index pairs
+		TArray<FVertexData> VertexDataSet;
+		VertexDataSet.Reserve(TriangleCount * 3);
+
+		FVector Position;
+		int32 GlobalVertexCount = 0;
+		for (FTessellationData& CTTessellation : FaceTessellationSet)
+		{
+			for (const FVector& Vertex : CTTessellation.PositionArray)
+			{
+				VertexDataSet.Emplace(GlobalVertexCount, Vertex * ImportParams.ScaleFactor);
+				++GlobalVertexCount;
+			}
+		}
+		VertexDataSet.SetNum(GlobalVertexCount);
+
+		// Sort the vertices by z value
+		VertexDataSet.Sort(FCompareVertexZ());
+
+		TArray<int32> NewIndexOf;
+		NewIndexOf.SetNumZeroed(GlobalVertexCount);
+
+		TArray<int32> IndexOfCoincidentNode;
+		IndexOfCoincidentNode.SetNumZeroed(GlobalVertexCount);
+
+		int32 VertexCount = 0;
+		// Search for duplicates, quickly!
+		for (int32 i = 0; i < VertexDataSet.Num(); i++)
+		{
+			NewIndexOf[VertexDataSet[i].Index] = i;
+			if (VertexDataSet[i].bIsMerged)
+			{
+				continue;
+			}
+
+			VertexDataSet[i].bIsMerged = true;
+			int32 Index_i = VertexDataSet[i].Index;
+			IndexOfCoincidentNode[Index_i] = Index_i;
+
+			// Check if mesh vertex is not outside body bbox
+			const FVector& PositionA = VertexDataSet[i].Coordinates;
+			if (BBox.IsValid && !BBox.IsInside(PositionA))
+			{
+				VertexDataSet[i].Index = INDEX_NONE;
+				continue;
+			}
+
+			// only need to search forward, since we add pairs both ways
+			for (int32 j = i + 1; j < VertexDataSet.Num(); j++)
+			{
+				if (FMath::Abs(VertexDataSet[j].Z - VertexDataSet[i].Z) > KINDA_SMALL_NUMBER)
+				{
+					break; // can't be any more duplicated
+				}
+
+				const FVector& PositionB = VertexDataSet[j].Coordinates;
+				if (PositionA.Equals(PositionB, KINDA_SMALL_NUMBER))
+				{
+					VertexDataSet[j].bIsMerged = true;
+					IndexOfCoincidentNode[VertexDataSet[j].Index] = Index_i;
+				}
+			}
+			VertexCount++;
+		}
+
+		// if Symmetric mesh, the symmetric side of the mesh have to be generated
+		FMatrix SymmetricMatrix;
+		bool bIsSymmetricMesh = MeshParameters.bIsSymmetric;
+		if (bIsSymmetricMesh)
+		{
+			SymmetricMatrix = FDatasmithUtils::GetSymmetricMatrix(MeshParameters.SymmetricOrigin, MeshParameters.SymmetricNormal);
+		}
+
+		// Make MeshDescription.VertexPositions and VertexID
+		MeshDescription.ReserveNewVertices(bIsSymmetricMesh ? VertexCount * 2 : VertexCount);
+		Body.VertexIds.Reserve(VertexCount);
+
+		int32 GlobalVertexIndex = 0;
+
+		for (int32 VertexIndex = 0; VertexIndex < GlobalVertexCount; ++VertexIndex)
+		{
+			int32 RealIndex = VertexDataSet[VertexIndex].Index;
+
+			// Vertex is outside bbox
+			if (RealIndex < 0)
+			{
+				continue;
+			}
+
+			if (IndexOfCoincidentNode[RealIndex] != RealIndex)
+			{
+				continue;
+			}
+
+			FVertexID VertexID = MeshDescription.CreateVertex();
+			Body.VertexIds.Add(VertexID);
+			VertexPositions[VertexID] = FDatasmithUtils::ConvertVector((FDatasmithUtils::EModelCoordSystem) ImportParams.ModelCoordSys, VertexDataSet[VertexIndex].Coordinates);
+			VertexDataSet[VertexIndex].VertexID = VertexID;
+		}
+
+		if (bIsSymmetricMesh)
+		{
+			Body.SymmetricVertexIds.Reserve(Body.VertexIds.Num());
+			for (int32 VertexIndex = 0; VertexIndex < GlobalVertexCount; ++VertexIndex)
+			{
+				int32 RealIndex = VertexDataSet[VertexIndex].Index;
+
+				// Vertex is outside bbox
+				if (RealIndex < 0)
+				{
+					continue;
+				}
+
+				if (IndexOfCoincidentNode[RealIndex] != RealIndex)
+				{
+					continue;
+				}
+
+				FVertexID VertexID = MeshDescription.CreateVertex();
+				Body.SymmetricVertexIds.Add(VertexID);
+				VertexPositions[VertexID] = FDatasmithUtils::ConvertVector((FDatasmithUtils::EModelCoordSystem) ImportParams.ModelCoordSys, VertexDataSet[VertexIndex].Coordinates);
+				VertexPositions[VertexID] = SymmetricMatrix.TransformPosition(VertexPositions[VertexID]);
+
+				VertexDataSet[VertexIndex].SymVertexID = VertexID;
+			}
+		}
+
+		// For each face, for each vertex set VertexId
+		GlobalVertexIndex = 0;
+		for (FTessellationData& CTTessellation : FaceTessellationSet)
+		{
+			CTTessellation.PositionIndices.SetNum(CTTessellation.PositionArray.Num());
+			for (int32 VertexIndex = 0; VertexIndex < CTTessellation.PositionArray.Num(); ++VertexIndex, ++GlobalVertexIndex)
+			{
+				int32 NewIndex = NewIndexOf[IndexOfCoincidentNode[GlobalVertexIndex]];
+				FVertexData& Vertex = VertexDataSet[NewIndex];
+				int32 Index = Vertex.VertexID.GetValue();
+				CTTessellation.PositionIndices[VertexIndex] = VertexDataSet[NewIndex].VertexID.GetValue();
+			}
+		}
+	}
+
+	// PolygonAttributes name used into modeling tools (ExtendedMeshAttribute::PolyTriGroups)
+	const FName PolyTriGroups("PolyTriGroups");
+
+	// Copy of FMeshDescriptionBuilder::EnablePolyGroups()
+	TPolygonAttributesRef<int32> EnableCADPatchGroups(FMeshDescription& MeshDescription)
+	{
+		TPolygonAttributesRef<int32> PatchGroups = MeshDescription.PolygonAttributes().GetAttributesRef<int32>(PolyTriGroups);
+		if (PatchGroups.IsValid() == false)
+		{
+			MeshDescription.PolygonAttributes().RegisterAttribute<int32>(PolyTriGroups, 1, 0, EMeshAttributeFlags::AutoGenerated);
+			PatchGroups = MeshDescription.PolygonAttributes().GetAttributesRef<int32>(PolyTriGroups);
+			check(PatchGroups.IsValid());
+		}
+		return PatchGroups;
+	}
+
+	/**
+	 * Polygon group is an attribute of polygons.As long as the mesh description is empty(no polygon), polygon group cannot be defined.
+	 * The work around used is to create PolygonGroups and to set them a PolygonGroupAttributes.
+	 * To get the existing polygon groups, for each created polygonGroups the PolygonGroupId is got(see GetExistingPatches)
+	 * Warning: CopyPatchGroups is call in FCoreTechRetessellate_Impl::ApplyOnOneAsset(CoreTechRetessellateAction.cpp) only if the option 
+	 * RetessellateOptions.RetessellationRule equal EDatasmithCADRetessellationRule::SkipDeletedSurfaces
+	 */
+	void CopyPatchGroups(FMeshDescription& MeshSource, FMeshDescription& MeshDestination)
+	{
+		TPolygonGroupAttributesRef<int32> PatchGroups = MeshDestination.PolygonGroupAttributes().GetAttributesRef<int32>(PolyTriGroups);
+		if (PatchGroups.IsValid() == false)
+		{
+			MeshDestination.PolygonGroupAttributes().RegisterAttribute<int32>(PolyTriGroups, 1, 0, EMeshAttributeFlags::AutoGenerated);
+			PatchGroups = MeshDestination.PolygonGroupAttributes().GetAttributesRef<int32>(PolyTriGroups);
+			check(PatchGroups.IsValid());
+		}
+
+		TSet<int32> PatchIdSet;
+		TPolygonAttributesRef<int32> ElementToGroups = MeshSource.PolygonAttributes().GetAttributesRef<int32>(PolyTriGroups);
+		for (const FPolygonID TriangleID : MeshSource.Polygons().GetElementIDs())
+		{
+			int32 PatchId = ElementToGroups[TriangleID];
+
+			bool bIsAlreadyInSet = false;
+			PatchIdSet.Add(PatchId, &bIsAlreadyInSet);
+			if (!bIsAlreadyInSet)
+			{
+				const FPolygonGroupID PolygonID = MeshDestination.CreatePolygonGroup();
+				PatchGroups[PolygonID] = PatchId;
+			}
+		}
+	}
+
+	/**
+	 * See CopyPatchGroups
+	 */
+	void GetExistingPatches(FMeshDescription& MeshDestination, TSet<int32>& OutPatchIdSet)
+	{
+		TPolygonGroupAttributesRef<int32> PatchGroups = MeshDestination.PolygonGroupAttributes().GetAttributesRef<int32>(PolyTriGroups);
+		if (PatchGroups.IsValid() == false)
+		{
+			return;
+		}
+
+		for (const FPolygonGroupID PolygonGroupID : MeshDestination.PolygonGroups().GetElementIDs())
+		{
+			int32 PatchId = PatchGroups[PolygonGroupID];
+			if (PatchId > 0) 
+			{
+				OutPatchIdSet.Add(PatchId);
+			}
+		}
+	}
+
+	bool FillMesh(const FMeshParameters& MeshParameters, const FImportParameters& ImportParams, FBodyMesh& BodyTessellation, FMeshDescription& MeshDescription)
+	{
+		const int32 UVChannel = 0;
+		const int32 TriangleCount = 3;
+		const TriangleIndex Clockwise = { 0, 1, 2 };
+		const TriangleIndex CounterClockwise = { 0, 2, 1 };
+
+		TArray<FVertexInstanceID> TriangleVertexInstanceIDs;
+		TriangleVertexInstanceIDs.SetNum(TriangleCount);
+
+		TArray<FVertexInstanceID> MeshVertexInstanceIDs;
+		TArray<uint32> NewFaceIndex;  // new CT face index to remove degenerated face
+
+		// Gather all array data
+		FStaticMeshAttributes Attributes(MeshDescription);
+		TVertexInstanceAttributesRef<FVector3f> VertexInstanceNormals = Attributes.GetVertexInstanceNormals();
+		TVertexInstanceAttributesRef<FVector3f> VertexInstanceTangents = Attributes.GetVertexInstanceTangents();
+		TVertexInstanceAttributesRef<float> VertexInstanceBinormalSigns = Attributes.GetVertexInstanceBinormalSigns();
+		TVertexInstanceAttributesRef<FVector4> VertexInstanceColors = Attributes.GetVertexInstanceColors();
+		TVertexInstanceAttributesRef<FVector2D> VertexInstanceUVs = Attributes.GetVertexInstanceUVs();
+		TPolygonGroupAttributesRef<FName> PolygonGroupImportedMaterialSlotNames = Attributes.GetPolygonGroupMaterialSlotNames();
+
+		if (!VertexInstanceNormals.IsValid() || !VertexInstanceTangents.IsValid() || !VertexInstanceBinormalSigns.IsValid() || !VertexInstanceColors.IsValid() || !VertexInstanceUVs.IsValid() || !PolygonGroupImportedMaterialSlotNames.IsValid())
+		{
+			return false;
+		}
+
+		// Find all the materials used
+		TMap<uint32, FPolygonGroupID> MaterialToPolygonGroupMapping;
+		for (const FTessellationData& FaceTessellation : BodyTessellation.Faces)
+		{
+			// we assume that face has only color
+			MaterialToPolygonGroupMapping.Add(FaceTessellation.ColorName, INDEX_NONE);
+		}
+
+		// Add to the mesh, a polygon groups per material
+		for (auto& Material : MaterialToPolygonGroupMapping)
+		{
+			uint32 MaterialHash = Material.Key;
+			FName ImportedSlotName = *LexToString<uint32>(MaterialHash);
+
+			FPolygonGroupID PolyGroupID = MeshDescription.CreatePolygonGroup();
+			PolygonGroupImportedMaterialSlotNames[PolyGroupID] = ImportedSlotName;
+			Material.Value = PolyGroupID;
+		}
+
+		VertexInstanceUVs.SetNumChannels(1);
+
+		int32 NbStep = 1;
+		if (MeshParameters.bIsSymmetric)
+		{
+			NbStep = 2;
+		}
+
+		TSet<int32> PatchIdSet;
+		GetExistingPatches(MeshDescription, PatchIdSet);
+		bool bImportOnlyAlreadyPresent = (bool) PatchIdSet.Num();
+
+		TPolygonAttributesRef<int32> PatchGroups = EnableCADPatchGroups(MeshDescription);
+		int32 PatchIndex = 0;
+		for (int32 Step = 0; Step < NbStep; ++Step)
+		{
+			// Swap mesh if needed
+			const TriangleIndex& Orientation = (!MeshParameters.bNeedSwapOrientation == (bool)Step) ? CounterClockwise : Clockwise;
+			TArray<int32>& VertexIdSet = (Step == 0) ? BodyTessellation.VertexIds : BodyTessellation.SymmetricVertexIds;
+
+			for (FTessellationData& Tessellation : BodyTessellation.Faces)
+			{
+				if (bImportOnlyAlreadyPresent && !PatchIdSet.Contains(Tessellation.PatchId))
+				{
+					continue;
+				}
+
+				// Get the polygonGroup
+				const FPolygonGroupID* PolygonGroupID = MaterialToPolygonGroupMapping.Find(Tessellation.ColorName);
+				if (PolygonGroupID == nullptr)
+				{
+					continue;
+				}
+
+				int32 VertexIDs[3];
+				int32 VertexIndex[3];
+				FVector Temp3D = { 0, 0, 0 };
+				FVector2D TexCoord2D = { 0, 0 };
+
+				MeshVertexInstanceIDs.Empty(Tessellation.VertexIndices.Num());
+				NewFaceIndex.Empty(Tessellation.VertexIndices.Num());
+
+				PatchIndex++;
+
+				// build each valid face i.e. 3 different indexes
+				for (int32 Index = 0; Index < Tessellation.VertexIndices.Num(); Index += 3)
+				{
+					VertexIndex[0] = Tessellation.PositionIndices[Tessellation.VertexIndices[Index + Orientation[0]]];
+					VertexIndex[1] = Tessellation.PositionIndices[Tessellation.VertexIndices[Index + Orientation[1]]];
+					VertexIndex[2] = Tessellation.PositionIndices[Tessellation.VertexIndices[Index + Orientation[2]]];
+
+					// Verify the 3 input indices are not defining a degenerated triangle
+					if (VertexIndex[0] == INDEX_NONE || VertexIndex[1] == INDEX_NONE || VertexIndex[2] == INDEX_NONE)
+					{
+						continue;
+					}
+
+					VertexIDs[0] = VertexIdSet[VertexIndex[0]];
+					VertexIDs[1] = VertexIdSet[VertexIndex[1]];
+					VertexIDs[2] = VertexIdSet[VertexIndex[2]];
+
+					NewFaceIndex.Add(Tessellation.VertexIndices[Index + Orientation[0]]);
+					NewFaceIndex.Add(Tessellation.VertexIndices[Index + Orientation[1]]);
+					NewFaceIndex.Add(Tessellation.VertexIndices[Index + Orientation[2]]);
+
+					MeshVertexInstanceIDs.Add(TriangleVertexInstanceIDs[0] = MeshDescription.CreateVertexInstance((FVertexID)VertexIDs[0]));
+					MeshVertexInstanceIDs.Add(TriangleVertexInstanceIDs[1] = MeshDescription.CreateVertexInstance((FVertexID)VertexIDs[1]));
+					MeshVertexInstanceIDs.Add(TriangleVertexInstanceIDs[2] = MeshDescription.CreateVertexInstance((FVertexID)VertexIDs[2]));
+
+					// Add the triangle as a polygon to the mesh description
+					const FPolygonID PolygonID = MeshDescription.CreatePolygon(*PolygonGroupID, TriangleVertexInstanceIDs);
+
+					// Set patch id attribute
+					PatchGroups[PolygonID] = Tessellation.PatchId;
+				}
+
+				for (int32 IndexFace = 0; IndexFace < MeshVertexInstanceIDs.Num(); IndexFace += 3)
+				{
+					for (int32 Index = 0; Index < TriangleCount; Index++)
+					{
+						const FVertexInstanceID VertexInstanceID = MeshVertexInstanceIDs[IndexFace + Orientation[Index]];
+						VertexInstanceUVs.Set(VertexInstanceID, UVChannel, Tessellation.TexCoordArray[NewFaceIndex[IndexFace + Index]]);
+
+						VertexInstanceColors[VertexInstanceID] = FLinearColor::White;
+						VertexInstanceTangents[VertexInstanceID] = FVector(ForceInitToZero);
+						VertexInstanceBinormalSigns[VertexInstanceID] = 0.0f;
+					}
+				}
+
+				if (Tessellation.TexCoordArray.Num())
+				{
+					for (int32 IndexFace = 0; IndexFace < NewFaceIndex.Num(); IndexFace += 3)
+					{
+						for (int32 Index = 0; Index < TriangleCount; Index++)
+						{
+							FVertexInstanceID VertexInstanceID = MeshVertexInstanceIDs[IndexFace + Orientation[Index]];
+							VertexInstanceUVs.Set(VertexInstanceID, UVChannel, Tessellation.TexCoordArray[NewFaceIndex[IndexFace + Index]]);
+						}
+					}
+				}
+
+				if (!Step)
+				{
+					FDatasmithUtils::ConvertVectorArray(ImportParams.ModelCoordSys, Tessellation.NormalArray);
+					for (FVector& Normal : Tessellation.NormalArray)
+					{
+						Normal = Normal.GetSafeNormal();
+					}
+				}
+
+				if (Tessellation.NormalArray.Num() == 1)
+				{
+					const FVector& Normal = Tessellation.NormalArray[0];
+					for (int32 Index = 0; Index < NewFaceIndex.Num(); Index++)
+					{
+						const FVertexInstanceID VertexInstanceID = MeshVertexInstanceIDs[Index];
+						VertexInstanceNormals[VertexInstanceID] = Normal;
+					}
+				}
+				else
+				{
+					for (FVector& Normal : Tessellation.NormalArray)
+					{
+						Normal = Normal.GetSafeNormal();
+					}
+
+					for (int32 IndexFace = 0; IndexFace < NewFaceIndex.Num(); IndexFace += 3)
+					{
+						for (int32 Index = 0; Index < 3; Index++)
+						{
+							const FVertexInstanceID VertexInstanceID = MeshVertexInstanceIDs[IndexFace + Orientation[Index]];
+							VertexInstanceNormals[VertexInstanceID] = Tessellation.NormalArray[NewFaceIndex[IndexFace + Index]];
+						}
+					}
+				}
+
+				// compute normals
+				if (Step)
+				{
+					// compute normals of Symmetric vertex
+					FMatrix SymmetricMatrix;
+					SymmetricMatrix = FDatasmithUtils::GetSymmetricMatrix(MeshParameters.SymmetricOrigin, MeshParameters.SymmetricNormal);
+					for (const FVertexInstanceID& VertexInstanceID : MeshVertexInstanceIDs)
+					{
+						VertexInstanceNormals[VertexInstanceID] = SymmetricMatrix.TransformVector(VertexInstanceNormals[VertexInstanceID]);;
+					}
+				}
+
+				if (MeshParameters.bNeedSwapOrientation)
+				{
+					for (int32 Index = 0; Index < MeshVertexInstanceIDs.Num(); Index++)
+					{
+						VertexInstanceNormals[MeshVertexInstanceIDs[Index]] = VertexInstanceNormals[MeshVertexInstanceIDs[Index]] * -1.f;
+					}
+				}
+			}
+		}
+		return true;
+	}
+
+	bool ConvertBodyMeshToMeshDescription(const FImportParameters& ImportParams, const FMeshParameters& MeshParameters, FBodyMesh& Body, FMeshDescription& MeshDescription)
+	{
+		// in a closed big mesh VertexCount ~ TriangleCount / 2, EdgeCount ~ 1.5* TriangleCount
+		MeshDescription.ReserveNewVertexInstances(Body.VertexArray.Num());
+		MeshDescription.ReserveNewPolygons(Body.TriangleCount);
+		MeshDescription.ReserveNewEdges(Body.TriangleCount * 3);
+
+		// CoreTech is generating position duplicates. make sure to remove them before filling the mesh description
+		if (Body.VertexArray.Num() == 0)
+		{
+			FillKioVertexPosition(ImportParams, MeshParameters, Body, MeshDescription);
+		}
+		else
+		{
+			//FillVertexPosition(ImportParams, MeshParameters, Body, MeshDescription);
+		}
+		
+		if (!FillMesh(MeshParameters, ImportParams, Body, MeshDescription))
+		{
+			return false;
+		}
+
+		// Orient mesh
+		MeshOperator::OrientMesh(MeshDescription);
+
+		// Build edge meta data
+		FStaticMeshOperations::DetermineEdgeHardnessesFromVertexInstanceNormals(MeshDescription);
+
+		return MeshDescription.Polygons().Num() > 0;
+	}
+
+	TSharedPtr<IDatasmithUEPbrMaterialElement> CreateDefaultUEPbrMaterial()
+	{
+		// Take the Material diffuse color and connect it to the BaseColor of a UEPbrMaterial
+		TSharedRef<IDatasmithUEPbrMaterialElement> MaterialElement = FDatasmithSceneFactory::CreateUEPbrMaterial(TEXT("0"));
+		MaterialElement->SetLabel(TEXT("DefaultCADImportMaterial"));
+
+		FLinearColor LinearColor = FLinearColor::FromPow22Color(FColor(200, 200, 200, 255));
+		IDatasmithMaterialExpressionColor* ColorExpression = MaterialElement->AddMaterialExpression<IDatasmithMaterialExpressionColor>();
+		ColorExpression->SetName(TEXT("Base Color"));
+		ColorExpression->GetColor() = LinearColor;
+		MaterialElement->GetBaseColor().SetExpression(ColorExpression);
+		MaterialElement->SetParentLabel(TEXT("M_DatasmithCAD"));
+
+		return MaterialElement;
+	}
+
+	TSharedPtr<IDatasmithUEPbrMaterialElement> CreateUEPbrMaterialFromColor(const FColor& InColor)
+	{
+		FString Name = FString::FromInt(BuildColorName(InColor));
+		FString Label = FString::Printf(TEXT("color_%02x%02x%02x%02x"), InColor.R, InColor.G, InColor.B, InColor.A);
+
+		// Take the Material diffuse color and connect it to the BaseColor of a UEPbrMaterial
+		TSharedRef<IDatasmithUEPbrMaterialElement> MaterialElement = FDatasmithSceneFactory::CreateUEPbrMaterial(*Name);
+		MaterialElement->SetLabel(*Label);
+
+		FLinearColor LinearColor = FLinearColor::FromSRGBColor(InColor);
+
+		IDatasmithMaterialExpressionColor* ColorExpression = MaterialElement->AddMaterialExpression<IDatasmithMaterialExpressionColor>();
+		ColorExpression->SetName(TEXT("Base Color"));
+		ColorExpression->GetColor() = LinearColor;
+
+		MaterialElement->GetBaseColor().SetExpression(ColorExpression);
+
+		if (LinearColor.A < 1.0f)
+		{
+			MaterialElement->SetBlendMode(/*EBlendMode::BLEND_Translucent*/2);
+
+			IDatasmithMaterialExpressionScalar* Scalar = MaterialElement->AddMaterialExpression<IDatasmithMaterialExpressionScalar>();
+			Scalar->GetScalar() = LinearColor.A;
+			Scalar->SetName(TEXT("Opacity Level"));
+
+			MaterialElement->GetOpacity().SetExpression(Scalar);
+			MaterialElement->SetParentLabel(TEXT("M_DatasmithCADTransparent"));
+		}
+		else
+		{
+			MaterialElement->SetParentLabel(TEXT("M_DatasmithCAD"));
+		}
+
+		return MaterialElement;
+	}
+
+	TSharedPtr<IDatasmithUEPbrMaterialElement> CreateUEPbrMaterialFromMaterial(FCADMaterial& InMaterial, TSharedRef<IDatasmithScene> Scene)
+	{
+		FString Name = FString::FromInt(BuildMaterialName(InMaterial));
+
+		// Take the Material diffuse color and connect it to the BaseColor of a UEPbrMaterial
+		TSharedRef<IDatasmithUEPbrMaterialElement> MaterialElement = FDatasmithSceneFactory::CreateUEPbrMaterial(*Name);
+		FString MaterialLabel(InMaterial.MaterialName);
+		if (MaterialLabel.IsEmpty())
+		{
+			MaterialLabel = TEXT("Material");
+		}
+		MaterialElement->SetLabel(*MaterialLabel);
+
+		// Set a diffuse color if there's nothing in the BaseColor
+		if (MaterialElement->GetBaseColor().GetExpression() == nullptr)
+		{
+			FLinearColor LinearColor = FLinearColor::FromPow22Color(InMaterial.Diffuse);
+
+			IDatasmithMaterialExpressionColor* ColorExpression = MaterialElement->AddMaterialExpression<IDatasmithMaterialExpressionColor>();
+			ColorExpression->SetName(TEXT("Base Color"));
+			ColorExpression->GetColor() = LinearColor;
+
+			MaterialElement->GetBaseColor().SetExpression(ColorExpression);
+		}
+
+		if (InMaterial.Transparency > 0.0f)
+		{
+			MaterialElement->SetBlendMode(/*EBlendMode::BLEND_Translucent*/2);
+			IDatasmithMaterialExpressionScalar* Scalar = MaterialElement->AddMaterialExpression<IDatasmithMaterialExpressionScalar>();
+			Scalar->GetScalar() = InMaterial.Transparency;
+			Scalar->SetName(TEXT("Opacity Level"));
+			MaterialElement->GetOpacity().SetExpression(Scalar);
+			MaterialElement->SetParentLabel(TEXT("M_DatasmithCADTransparent"));
+		}
+		else
+		{
+			MaterialElement->SetParentLabel(TEXT("M_DatasmithCAD"));
+		}
+
+		return MaterialElement;
+	}
+}
diff --git a/Engine/Plugins/Enterprise/DatasmithCADImporter/Source/CADLibrary/Public/CoreTechHelper.h b/Engine/Plugins/Enterprise/DatasmithCADImporter/Source/CADLibrary/Public/CoreTechHelper.h
deleted file mode 100644
diff --git a/Engine/Plugins/Enterprise/DatasmithCADImporter/Source/CADLibrary/Public/MeshDescriptionHelper.h b/Engine/Plugins/Enterprise/DatasmithCADImporter/Source/CADLibrary/Public/MeshDescriptionHelper.h
new file mode 100644
--- /dev/null
+++ b/Engine/Plugins/Enterprise/DatasmithCADImporter/Source/CADLibrary/Public/MeshDescriptionHelper.h
@@ -0,0 +1,30 @@
+// Copyright Epic Games, Inc. All Rights Reserved.
+#pragma once
+
+#include "CoreMinimal.h"
+
+#include "CADData.h"
+#include "CADOptions.h"
+
+struct FCTMesh;
+struct FMeshDescription;
+class FBodyMesh;
+class FCADMaterial;
+class IDatasmithMaterialIDElement;
+class IDatasmithUEPbrMaterialElement;
+class IDatasmithScene;
+
+#define LAST_CT_MATERIAL_ID 0x00ffffff
+
+namespace CADLibrary
+{
+	using TColorMap = TMap<uint32, uint32>;
+
+	CADLIBRARY_API TSharedPtr<IDatasmithUEPbrMaterialElement> CreateDefaultUEPbrMaterial();
+	CADLIBRARY_API TSharedPtr<IDatasmithUEPbrMaterialElement> CreateUEPbrMaterialFromColor(const FColor& InColor);
+	CADLIBRARY_API TSharedPtr<IDatasmithUEPbrMaterialElement> CreateUEPbrMaterialFromMaterial(FCADMaterial& InMaterial, TSharedRef<IDatasmithScene> Scene);
+
+	CADLIBRARY_API bool ConvertBodyMeshToMeshDescription(const FImportParameters& ImportParams, const FMeshParameters& MeshParameters, FBodyMesh& Body, FMeshDescription& MeshDescription);
+	CADLIBRARY_API void CopyPatchGroups(FMeshDescription& MeshSource, FMeshDescription& MeshDestination);
+
+} // namespace CADLibrary
diff --git a/Engine/Plugins/Enterprise/DatasmithCADImporter/Source/CoreTechSurface/CoreTechSurface.Build.cs b/Engine/Plugins/Enterprise/DatasmithCADImporter/Source/CoreTechSurface/CoreTechSurface.Build.cs
--- a/Engine/Plugins/Enterprise/DatasmithCADImporter/Source/CoreTechSurface/CoreTechSurface.Build.cs
+++ b/Engine/Plugins/Enterprise/DatasmithCADImporter/Source/CoreTechSurface/CoreTechSurface.Build.cs
@@ -18,15 +18,17 @@ namespace UnrealBuildTool.Rules
 			PrivateDependencyModuleNames.AddRange(
 				new string[]
 				{
+					"CADInterfaces",
+					"CADLibrary",
+					"CADTools",
 					"Core",
 					"CoreUObject",
-					"MeshDescription",
 					"DatasmithContent",
 					"DatasmithTranslator",
 					"Engine",
+					"MeshDescription",
+					"ParametricSurface",
 					"StaticMeshDescription",
-					"CADInterfaces",
-					"CADLibrary",
 				}
 			);
 		}
diff --git a/Engine/Plugins/Enterprise/DatasmithCADImporter/Source/CoreTechSurface/Private/CoreTechSurfaceExtension.cpp b/Engine/Plugins/Enterprise/DatasmithCADImporter/Source/CoreTechSurface/Private/CoreTechSurfaceExtension.cpp
--- a/Engine/Plugins/Enterprise/DatasmithCADImporter/Source/CoreTechSurface/Private/CoreTechSurfaceExtension.cpp
+++ b/Engine/Plugins/Enterprise/DatasmithCADImporter/Source/CoreTechSurface/Private/CoreTechSurfaceExtension.cpp
@@ -3,14 +3,19 @@
 #include "CoreTechSurfaceExtension.h"
 
 #include "CADInterfacesModule.h"
+#include "CoreTechSurfaceHelper.h"
+#include "CoreTechTypes.h"
+#include "DatasmithAdditionalData.h"
 #include "DatasmithPayload.h"
+#include "Engine/StaticMesh.h"
+#include "HAL/PlatformFileManager.h"
 #include "IDatasmithSceneElements.h"
+#include "MeshDescriptionHelper.h"
 #include "Misc/FileHelper.h"
 #include "Misc/Paths.h"
+#include "StaticMeshAttributes.h"
 #include "UObject/EnterpriseObjectVersion.h"
 
-#include "CADOptions.h"
-
 void UCoreTechParametricSurfaceData::Serialize(FArchive& Ar)
 {
 	Ar.UsingCustomVersion(FEnterpriseObjectVersion::GUID);
@@ -28,37 +33,74 @@ void UCoreTechParametricSurfaceData::Serialize(FArchive& Ar)
 	}
 }
 
-namespace CoreTechSurface
+bool UCoreTechParametricSurfaceData::Tessellate(UStaticMesh& StaticMesh, const FDatasmithRetessellationOptions& RetessellateOptions)
 {
-	// TODO: convert to FCoreTechSceneParameters/FCoreTechMeshParameters ?
-	void AddCoreTechSurfaceDataForMesh(const TSharedRef<IDatasmithMeshElement>& InMeshElement, const CADLibrary::FImportParameters& InSceneParameters, const CADLibrary::FMeshParameters& InMeshParameters, const FDatasmithTessellationOptions& InTessellationOptions, FDatasmithMeshElementPayload& OutMeshPayload)
+	bool bSuccessfulTessellation = false;
+
+#if WITH_EDITOR
+	// make a temporary file as CoreTech can only deal with files.
+	int32 Hash = GetTypeHash(StaticMesh.GetPathName());
+	FString ResourceFile = FPaths::ConvertRelativePathToFull(FPaths::ProjectIntermediateDir() / FString::Printf(TEXT("0x%08x.ct"), Hash));
+
+	FFileHelper::SaveArrayToFile(RawData, *ResourceFile);
+
+	CADLibrary::FCTMesh Mesh;
+
+	CADLibrary::FImportParameters ImportParameters;
+	ImportParameters.MetricUnit = SceneParameters.MetricUnit;
+	ImportParameters.ScaleFactor = SceneParameters.ScaleFactor;
+	ImportParameters.ChordTolerance = RetessellateOptions.ChordTolerance;
+	ImportParameters.MaxEdgeLength = RetessellateOptions.MaxEdgeLength;
+	ImportParameters.MaxNormalAngle = RetessellateOptions.NormalTolerance;
+	ImportParameters.ModelCoordSys = static_cast<FDatasmithUtils::EModelCoordSystem>(SceneParameters.ModelCoordSys);
+	ImportParameters.StitchingTechnique = CADLibrary::EStitchingTechnique(RetessellateOptions.StitchingTechnique);
+
+	CADLibrary::FMeshParameters CadMeshParameters;
+	CadMeshParameters.bNeedSwapOrientation = MeshParameters.bNeedSwapOrientation;
+	CadMeshParameters.bIsSymmetric = MeshParameters.bIsSymmetric;
+	CadMeshParameters.SymmetricNormal = MeshParameters.SymmetricNormal;
+	CadMeshParameters.SymmetricOrigin = MeshParameters.SymmetricOrigin;
+
+	// Previous MeshDescription is get to be able to create a new one with the same order of PolygonGroup (the matching of color and partition is currently based on their order)
+	if (FMeshDescription* DestinationMeshDescription = StaticMesh.GetMeshDescription(0))
 	{
-		if (ICADInterfacesModule::GetAvailability() == ECADInterfaceAvailability::Available)
+		FMeshDescription MeshDescription;
+		FStaticMeshAttributes MeshDescriptionAttributes(MeshDescription);
+		MeshDescriptionAttributes.Register();
+
+		if (RetessellateOptions.RetessellationRule == EDatasmithCADRetessellationRule::SkipDeletedSurfaces)
+		{
+			CADLibrary::CopyPatchGroups(*DestinationMeshDescription, MeshDescription);
+		}
+
+		if (CoreTechSurface::LoadFile(ResourceFile, MeshDescription, ImportParameters, CadMeshParameters))
 		{
-			// Store CoreTech additional data if provided
-			const TCHAR* CoretechFile = InMeshElement->GetFile();
-			if (FPaths::FileExists(CoretechFile))
+			// To update the SectionInfoMap 
 			{
-				TArray<uint8> ByteArray;
-				if (FFileHelper::LoadFileToArray(ByteArray, CoretechFile))
+				TPolygonGroupAttributesConstRef<FName> MaterialSlotNames = MeshDescriptionAttributes.GetPolygonGroupMaterialSlotNames();
+				FMeshSectionInfoMap& SectionInfoMap = StaticMesh.GetSectionInfoMap();
+
+				for (FPolygonGroupID PolygonGroupID : MeshDescription.PolygonGroups().GetElementIDs())
 				{
-					UCoreTechParametricSurfaceData* CoreTechData = Datasmith::MakeAdditionalData<UCoreTechParametricSurfaceData>();
-					CoreTechData->SourceFile = CoretechFile;
-					CoreTechData->RawData = MoveTemp(ByteArray);
-					CoreTechData->SceneParameters.ModelCoordSys = uint8(InSceneParameters.ModelCoordSys);
-					CoreTechData->SceneParameters.MetricUnit = InSceneParameters.MetricUnit;
-					CoreTechData->SceneParameters.ScaleFactor = InSceneParameters.ScaleFactor;
-
-					CoreTechData->MeshParameters.bNeedSwapOrientation = InMeshParameters.bNeedSwapOrientation;
-					CoreTechData->MeshParameters.bIsSymmetric = InMeshParameters.bIsSymmetric;
-					CoreTechData->MeshParameters.SymmetricNormal = InMeshParameters.SymmetricNormal;
-					CoreTechData->MeshParameters.SymmetricOrigin = InMeshParameters.SymmetricOrigin;
-
-					CoreTechData->LastTessellationOptions = InTessellationOptions;
-					OutMeshPayload.AdditionalData.Add(CoreTechData);
+					FMeshSectionInfo Section = SectionInfoMap.Get(0, PolygonGroupID.GetValue());
+					int32 MaterialIndex = StaticMesh.GetMaterialIndex(MaterialSlotNames[PolygonGroupID]);
+					if (MaterialIndex < 0)
+					{
+						MaterialIndex = 0;
+					}
+					Section.MaterialIndex = MaterialIndex;
+					SectionInfoMap.Set(0, PolygonGroupID.GetValue(), Section);
 				}
 			}
+			*DestinationMeshDescription = MoveTemp(MeshDescription);
+			bSuccessfulTessellation = true;
 		}
 	}
 
+	// Remove temporary file
+	FPlatformFileManager::Get().GetPlatformFile().DeleteFile(*ResourceFile);
+#endif
+	return bSuccessfulTessellation;
 }
+
+
diff --git a/Engine/Plugins/Enterprise/DatasmithCADImporter/Source/CoreTechSurface/Private/CoreTechSurfaceHelper.cpp b/Engine/Plugins/Enterprise/DatasmithCADImporter/Source/CoreTechSurface/Private/CoreTechSurfaceHelper.cpp
new file mode 100644
--- /dev/null
+++ b/Engine/Plugins/Enterprise/DatasmithCADImporter/Source/CoreTechSurface/Private/CoreTechSurfaceHelper.cpp
@@ -0,0 +1,96 @@
+// Copyright Epic Games, Inc. All Rights Reserved.
+#include "CoreTechSurfaceHelper.h"
+
+#include "CADInterfacesModule.h"
+#include "CADOptions.h"
+#include "CoreTechSurfaceExtension.h"
+#include "CoreTechTypes.h"
+#include "DatasmithPayload.h"
+#include "IDatasmithSceneElements.h"
+#include "MeshDescription.h"
+#include "MeshDescriptionHelper.h"
+#include "Misc/FileHelper.h"
+#include "UObject/NameTypes.h"
+
+typedef uint32 TriangleIndex[3];
+
+namespace CoreTechSurface
+{
+	bool Tessellate(uint64 MainObjectId, const CADLibrary::FImportParameters& ImportParams, FMeshDescription& MeshDesc, CADLibrary::FMeshParameters& MeshParameters)
+	{
+		CADLibrary::CTKIO_SetCoreTechTessellationState(ImportParams);
+
+		CADLibrary::FBodyMesh BodyMesh;
+		BodyMesh.BodyID = 1;
+
+		CADLibrary::CTKIO_GetTessellation(MainObjectId, BodyMesh, false);
+
+		if (BodyMesh.Faces.Num() == 0)
+		{
+			return false;
+		}
+
+		if (!CADLibrary::ConvertBodyMeshToMeshDescription(ImportParams, MeshParameters, BodyMesh, MeshDesc))
+		{
+			ensureMsgf(false, TEXT("Error during mesh conversion"));
+			return false;
+		}
+
+		return true;
+	}
+
+	bool LoadFile(const FString& FileName, FMeshDescription& MeshDescription, const CADLibrary::FImportParameters& ImportParameters, CADLibrary::FMeshParameters& MeshParameters)
+	{
+		CADLibrary::FCoreTechSessionBase Session(TEXT("CoreTechMeshLoader::LoadFile"), ImportParameters.MetricUnit);
+		if (!Session.IsSessionValid())
+		{
+			return false;
+		}
+
+		uint64 MainObjectID;
+		if(!CADLibrary::CTKIO_LoadModel(*FileName, MainObjectID, 0x00020000 /* CT_LOAD_FLAGS_READ_META_DATA */))
+		{
+			// Something wrong happened during the load, abort
+			return false;
+		}
+
+		if (ImportParameters.StitchingTechnique != CADLibrary::EStitchingTechnique::StitchingNone)
+		{
+			CADLibrary::CTKIO_Repair(MainObjectID, CADLibrary::EStitchingTechnique::StitchingHeal);
+		}
+
+		return Tessellate(MainObjectID, ImportParameters, MeshDescription, MeshParameters);
+	}
+
+	// TODO: convert to FCoreTechSceneParameters/FCoreTechMeshParameters ?
+	void AddSurfaceDataForMesh(const TSharedRef<IDatasmithMeshElement>& InMeshElement, const CADLibrary::FImportParameters& InSceneParameters, const CADLibrary::FMeshParameters& InMeshParameters, const FDatasmithTessellationOptions& InTessellationOptions, FDatasmithMeshElementPayload& OutMeshPayload)
+	{
+		if (ICADInterfacesModule::GetAvailability() == ECADInterfaceAvailability::Available)
+		{
+			// Store CoreTech additional data if provided
+			const TCHAR* CoretechFile = InMeshElement->GetFile();
+			if (FPaths::FileExists(CoretechFile))
+			{
+				TArray<uint8> ByteArray;
+				if (FFileHelper::LoadFileToArray(ByteArray, CoretechFile))
+				{
+					UCoreTechParametricSurfaceData* CoreTechData = Datasmith::MakeAdditionalData<UCoreTechParametricSurfaceData>();
+					CoreTechData->SourceFile = CoretechFile;
+					CoreTechData->RawData = MoveTemp(ByteArray);
+					CoreTechData->SceneParameters.ModelCoordSys = uint8(InSceneParameters.ModelCoordSys);
+					CoreTechData->SceneParameters.MetricUnit = InSceneParameters.MetricUnit;
+					CoreTechData->SceneParameters.ScaleFactor = InSceneParameters.ScaleFactor;
+
+					CoreTechData->MeshParameters.bNeedSwapOrientation = InMeshParameters.bNeedSwapOrientation;
+					CoreTechData->MeshParameters.bIsSymmetric = InMeshParameters.bIsSymmetric;
+					CoreTechData->MeshParameters.SymmetricNormal = InMeshParameters.SymmetricNormal;
+					CoreTechData->MeshParameters.SymmetricOrigin = InMeshParameters.SymmetricOrigin;
+
+					CoreTechData->LastTessellationOptions = InTessellationOptions;
+					OutMeshPayload.AdditionalData.Add(CoreTechData);
+				}
+			}
+		}
+	}
+
+}
diff --git a/Engine/Plugins/Enterprise/DatasmithCADImporter/Source/CoreTechSurface/Private/CoreTechTranslator.cpp b/Engine/Plugins/Enterprise/DatasmithCADImporter/Source/CoreTechSurface/Private/CoreTechTranslator.cpp
deleted file mode 100644
diff --git a/Engine/Plugins/Enterprise/DatasmithCADImporter/Source/CoreTechSurface/Public/CoreTechSurfaceExtension.h b/Engine/Plugins/Enterprise/DatasmithCADImporter/Source/CoreTechSurface/Public/CoreTechSurfaceExtension.h
--- a/Engine/Plugins/Enterprise/DatasmithCADImporter/Source/CoreTechSurface/Public/CoreTechSurfaceExtension.h
+++ b/Engine/Plugins/Enterprise/DatasmithCADImporter/Source/CoreTechSurface/Public/CoreTechSurfaceExtension.h
@@ -4,6 +4,7 @@
 
 #include "CoreMinimal.h"
 
+#include "ParametricSurfaceExtension.h"
 #include "DatasmithAdditionalData.h"
 #include "DatasmithCustomAction.h"
 #include "DatasmithImportOptions.h"
@@ -11,44 +12,8 @@
 
 #include "CoreTechSurfaceExtension.generated.h"
 
-
-USTRUCT(BlueprintType)
-struct CORETECHSURFACE_API FCoreTechSceneParameters
-{
-	GENERATED_BODY()
-
-	// value from FDatasmithUtils::EModelCoordSystem
-	UPROPERTY()
-	uint8 ModelCoordSys = (uint8)FDatasmithUtils::EModelCoordSystem::ZUp_LeftHanded;
-
-	UPROPERTY()
-	float MetricUnit = 0.01f;
-
-	UPROPERTY()
-	float ScaleFactor = 1.0f;
-};
-
-USTRUCT()
-struct CORETECHSURFACE_API FCoreTechMeshParameters
-{
-	GENERATED_BODY()
-
-	UPROPERTY()
-	bool bNeedSwapOrientation = false;
-
-	UPROPERTY()
-	bool bIsSymmetric = false;
-
-	UPROPERTY()
-	FVector SymmetricOrigin = FVector::ZeroVector;
-
-	UPROPERTY()
-	FVector SymmetricNormal = FVector::ZeroVector;
-};
-
-
-UCLASS(meta = (DisplayName = "Parametric Surface Data"))
-class CORETECHSURFACE_API UCoreTechParametricSurfaceData : public UDatasmithAdditionalData
+UCLASS(meta = (DisplayName = "Kernel IO Parametric Surface Data"))
+class CORETECHSURFACE_API UCoreTechParametricSurfaceData : public UParametricSurfaceData
 {
 	GENERATED_BODY()
 
@@ -59,14 +24,12 @@ public:
 	// Too costly to serialize as a UPROPERTY, will use custom serialization.
 	TArray<uint8> RawData;
 
-	UPROPERTY()
-	FCoreTechSceneParameters SceneParameters;
-
-	UPROPERTY()
-	FCoreTechMeshParameters MeshParameters;
-
-	UPROPERTY(EditAnywhere, Category=NURBS)
-	FDatasmithTessellationOptions LastTessellationOptions;
+	virtual bool IsValid() override
+	{
+		return RawData.Num() > 0;
+	}
+	
+	virtual bool Tessellate(UStaticMesh& StaticMesh, const FDatasmithRetessellationOptions& RetessellateOptions) override;
 
 private:
 	UPROPERTY()
@@ -74,17 +37,3 @@ private:
 
 	virtual void Serialize(FArchive& Ar) override;
 };
-
-class IDatasmithMeshElement;
-struct FDatasmithMeshElementPayload;
-
-namespace CADLibrary
-{
-	struct FImportParameters;
-	struct FMeshParameters;
-}
-
-namespace CoreTechSurface
-{
-	void CORETECHSURFACE_API AddCoreTechSurfaceDataForMesh(const TSharedRef<IDatasmithMeshElement>& InMeshElement, const CADLibrary::FImportParameters& InSceneParameters, const CADLibrary::FMeshParameters&, const FDatasmithTessellationOptions& InTessellationOptions, FDatasmithMeshElementPayload& OutMeshPayload);
-}
diff --git a/Engine/Plugins/Enterprise/DatasmithCADImporter/Source/CoreTechSurface/Public/CoreTechSurfaceHelper.h b/Engine/Plugins/Enterprise/DatasmithCADImporter/Source/CoreTechSurface/Public/CoreTechSurfaceHelper.h
new file mode 100644
--- /dev/null
+++ b/Engine/Plugins/Enterprise/DatasmithCADImporter/Source/CoreTechSurface/Public/CoreTechSurfaceHelper.h
@@ -0,0 +1,22 @@
+// Copyright Epic Games, Inc. All Rights Reserved.
+#pragma once
+
+#include "CoreMinimal.h"
+
+namespace CADLibrary
+{
+	struct FImportParameters;
+	struct FMeshParameters;
+}
+
+class IDatasmithMeshElement;
+struct FDatasmithMeshElementPayload;
+struct FDatasmithTessellationOptions;
+struct FMeshDescription;
+
+namespace CoreTechSurface
+{
+	CORETECHSURFACE_API void AddSurfaceDataForMesh(const TSharedRef<IDatasmithMeshElement>& InMeshElement, const CADLibrary::FImportParameters& InSceneParameters, const CADLibrary::FMeshParameters&, const FDatasmithTessellationOptions& InTessellationOptions, FDatasmithMeshElementPayload& OutMeshPayload);
+	CORETECHSURFACE_API bool LoadFile(const FString& FileName, FMeshDescription& MeshDescription, const CADLibrary::FImportParameters& ImportParameters, CADLibrary::FMeshParameters& MeshParameters);
+	CORETECHSURFACE_API bool Tessellate(uint64 MainObjectId, const CADLibrary::FImportParameters& ImportParams, FMeshDescription& Mesh, CADLibrary::FMeshParameters& MeshParameters);
+}
\ No newline at end of file
diff --git a/Engine/Plugins/Enterprise/DatasmithCADImporter/Source/CoreTechSurface/Public/CoreTechTranslator.h b/Engine/Plugins/Enterprise/DatasmithCADImporter/Source/CoreTechSurface/Public/CoreTechTranslator.h
deleted file mode 100644
diff --git a/Engine/Plugins/Enterprise/DatasmithCADImporter/Source/DatasmithCADTranslator/DatasmithCADTranslator.Build.cs b/Engine/Plugins/Enterprise/DatasmithCADImporter/Source/DatasmithCADTranslator/DatasmithCADTranslator.Build.cs
--- a/Engine/Plugins/Enterprise/DatasmithCADImporter/Source/DatasmithCADTranslator/DatasmithCADTranslator.Build.cs
+++ b/Engine/Plugins/Enterprise/DatasmithCADImporter/Source/DatasmithCADTranslator/DatasmithCADTranslator.Build.cs
@@ -29,6 +29,7 @@ namespace UnrealBuildTool.Rules
 					"DatasmithTranslator",
 					"Engine",
 					"MeshDescription",
+					"ParametricSurface",
 					"Sockets",
 				}
 			);
diff --git a/Engine/Plugins/Enterprise/DatasmithCADImporter/Source/DatasmithCADTranslator/Private/DatasmithCADTranslator.cpp b/Engine/Plugins/Enterprise/DatasmithCADImporter/Source/DatasmithCADTranslator/Private/DatasmithCADTranslator.cpp
--- a/Engine/Plugins/Enterprise/DatasmithCADImporter/Source/DatasmithCADTranslator/Private/DatasmithCADTranslator.cpp
+++ b/Engine/Plugins/Enterprise/DatasmithCADImporter/Source/DatasmithCADTranslator/Private/DatasmithCADTranslator.cpp
@@ -3,17 +3,12 @@
 #include "DatasmithCADTranslator.h"
 
 #include "CADInterfacesModule.h"
-#include "CoreTechFileParser.h"
-#include "CoreTechSurfaceExtension.h"
+#include "CoreTechSurfaceHelper.h"
 #include "DatasmithCADTranslatorModule.h"
 #include "DatasmithDispatcher.h"
 #include "DatasmithMeshBuilder.h"
 #include "DatasmithSceneGraphBuilder.h"
 #include "DatasmithUtils.h"
-#include "IDatasmithSceneElements.h"
-
-#include "HAL/IConsoleManager.h"
-#include "Misc/FileHelper.h"
 
 static TAutoConsoleVariable<int32> CVarStaticCADTranslatorEnableThreadedImport(
 	TEXT("r.CADTranslator.EnableThreadedImport"),
@@ -181,7 +176,7 @@ bool FDatasmithCADTranslator::LoadScene(TSharedRef<IDatasmithScene> DatasmithSce
 
 void FDatasmithCADTranslator::UnloadScene()
 {
-	MeshBuilderPtr = nullptr;
+	MeshBuilderPtr.Reset();
 
 	CADFileToUEGeomMap.Empty();
 }
@@ -198,17 +193,15 @@ bool FDatasmithCADTranslator::LoadStaticMesh(const TSharedRef<IDatasmithMeshElem
 	if (TOptional< FMeshDescription > Mesh = MeshBuilderPtr->GetMeshDescription(MeshElement, MeshParameters))
 	{
 		OutMeshPayload.LodMeshes.Add(MoveTemp(Mesh.GetValue()));
-
-		CoreTechSurface::AddCoreTechSurfaceDataForMesh(MeshElement, ImportParameters, MeshParameters, GetCommonTessellationOptions(), OutMeshPayload);
+		if(ImportParameters.bEnableKernelIOTessellation)
+		{
+			CoreTechSurface::AddSurfaceDataForMesh(MeshElement, ImportParameters, MeshParameters, GetCommonTessellationOptions(), OutMeshPayload);
+		}
+		else
+		{
+			//CADKernelSurface::AddSurfaceDataForMesh(MeshElement, ImportParameters, MeshParameters, GetCommonTessellationOptions(), OutMeshPayload);
+		}
 	}
 	return OutMeshPayload.LodMeshes.Num() > 0;
 }
 
-void FDatasmithCADTranslator::SetSceneImportOptions(TArray<TStrongObjectPtr<UDatasmithOptionsBase>>& Options)
-{
-	FCoreTechTranslator::SetSceneImportOptions(Options);
-}
-
-
-
-
diff --git a/Engine/Plugins/Enterprise/DatasmithCADImporter/Source/DatasmithCADTranslator/Private/DatasmithCADTranslator.h b/Engine/Plugins/Enterprise/DatasmithCADImporter/Source/DatasmithCADTranslator/Private/DatasmithCADTranslator.h
--- a/Engine/Plugins/Enterprise/DatasmithCADImporter/Source/DatasmithCADTranslator/Private/DatasmithCADTranslator.h
+++ b/Engine/Plugins/Enterprise/DatasmithCADImporter/Source/DatasmithCADTranslator/Private/DatasmithCADTranslator.h
@@ -4,18 +4,17 @@
 
 #include "CoreMinimal.h"
 
-#include "CoreTechTranslator.h"
-
 #include "CADOptions.h"
-#include "DatasmithDispatcher.h"
 #include "DatasmithMeshBuilder.h"
 #include "DatasmithSceneGraphBuilder.h"
-
+#include "ParametricSurfaceTranslator.h"
 #include "UObject/ObjectMacros.h"
 
 class IDatasmithMeshElement;
+class IDatasmithScene;
+struct FDatasmithMeshElementPayload;
 
-class FDatasmithCADTranslator : public FCoreTechTranslator
+class FDatasmithCADTranslator : public FParametricSurfaceTranslator
 {
 public:
 	virtual FName GetFName() const override { return "DatasmithCADTranslator"; };
@@ -28,8 +27,6 @@ public:
 
 	virtual bool LoadStaticMesh(const TSharedRef<IDatasmithMeshElement> MeshElement, FDatasmithMeshElementPayload& OutMeshPayload) override;
 
-	virtual void SetSceneImportOptions(TArray<TStrongObjectPtr<UDatasmithOptionsBase>>& Options) override;
-
 private:
 	TMap<uint32, FString> CADFileToUEGeomMap;
 
diff --git a/Engine/Plugins/Enterprise/DatasmithCADImporter/Source/DatasmithCADTranslator/Private/DatasmithMeshBuilder.cpp b/Engine/Plugins/Enterprise/DatasmithCADImporter/Source/DatasmithCADTranslator/Private/DatasmithMeshBuilder.cpp
--- a/Engine/Plugins/Enterprise/DatasmithCADImporter/Source/DatasmithCADTranslator/Private/DatasmithMeshBuilder.cpp
+++ b/Engine/Plugins/Enterprise/DatasmithCADImporter/Source/DatasmithCADTranslator/Private/DatasmithMeshBuilder.cpp
@@ -1,17 +1,15 @@
 // Copyright Epic Games, Inc. All Rights Reserved.
 #include "DatasmithMeshBuilder.h"
 
-#include "CoreTechHelper.h"
-
+#include "CADData.h"
 #include "IDatasmithSceneElements.h"
+#include "MeshDescriptionHelper.h"
 #include "Utility/DatasmithMeshHelper.h"
 
 #include "HAL/FileManager.h"
 #include "MeshDescription.h"
 #include "Misc/Paths.h"
 
-using namespace CADLibrary;
-
 FDatasmithMeshBuilder::FDatasmithMeshBuilder(TMap<uint32, FString>& CADFileToMeshFile, const FString& InCachePath, const CADLibrary::FImportParameters& InImportParameters)
 	: CachePath(InCachePath)
 	, ImportParameters(InImportParameters)
@@ -32,7 +30,7 @@ void FDatasmithMeshBuilder::LoadMeshFiles(TMap<uint32, FString>& CADFileToMeshFi
 		}
 		TArray<CADLibrary::FBodyMesh>& BodyMeshSet = BodyMeshes.Emplace_GetRef();
 		DeserializeBodyMeshFile(*MeshFile, BodyMeshSet);
-		for (FBodyMesh& Body : BodyMeshSet)
+		for (CADLibrary::FBodyMesh& Body : BodyMeshSet)
 		{
 			MeshActorNameToBodyMesh.Emplace(Body.MeshActorName, &Body);
 		}
@@ -48,13 +46,13 @@ TOptional<FMeshDescription> FDatasmithMeshBuilder::GetMeshDescription(TSharedRef
 		return TOptional<FMeshDescription>();
 	}
 
-	FBodyMesh** PPBody = MeshActorNameToBodyMesh.Find(BodyUuid);
+	CADLibrary::FBodyMesh** PPBody = MeshActorNameToBodyMesh.Find(BodyUuid);
 	if(PPBody == nullptr || *PPBody == nullptr)
 	{
 		return TOptional<FMeshDescription>();
 	}
 
-	FBodyMesh& Body = **PPBody;
+	CADLibrary::FBodyMesh& Body = **PPBody;
 
 	// FDatasmithSceneBaseGraphBuilder::BuildBody is performing a special treatment for
 	// FBodyMesh without color. Replicate the treatment here too.
@@ -66,7 +64,7 @@ TOptional<FMeshDescription> FDatasmithMeshBuilder::GetMeshDescription(TSharedRef
 
 		Body.ColorSet.Add(MaterialSlotId);
 
-		for (FTessellationData& Face : Body.Faces)
+		for (CADLibrary::FTessellationData& Face : Body.Faces)
 		{
 			Face.ColorName = MaterialSlotId;
 		}
diff --git a/Engine/Plugins/Enterprise/DatasmithCADImporter/Source/DatasmithCADTranslator/Private/DatasmithSceneGraphBuilder.cpp b/Engine/Plugins/Enterprise/DatasmithCADImporter/Source/DatasmithCADTranslator/Private/DatasmithSceneGraphBuilder.cpp
--- a/Engine/Plugins/Enterprise/DatasmithCADImporter/Source/DatasmithCADTranslator/Private/DatasmithSceneGraphBuilder.cpp
+++ b/Engine/Plugins/Enterprise/DatasmithCADImporter/Source/DatasmithCADTranslator/Private/DatasmithSceneGraphBuilder.cpp
@@ -4,11 +4,11 @@
 #include "CADData.h"
 #include "CADSceneGraph.h"
 #include "CoreTechFileParser.h"
-#include "CoreTechHelper.h"
 #include "DatasmithSceneFactory.h"
 #include "DatasmithSceneSource.h"
+#include "DatasmithUtils.h"
 #include "IDatasmithSceneElements.h"
-
+#include "MeshDescriptionHelper.h"
 #include "Misc/FileHelper.h"
 
 namespace
diff --git a/Engine/Plugins/Enterprise/DatasmithCADImporter/Source/DatasmithCADTranslator/Public/DatasmithSceneGraphBuilder.h b/Engine/Plugins/Enterprise/DatasmithCADImporter/Source/DatasmithCADTranslator/Public/DatasmithSceneGraphBuilder.h
--- a/Engine/Plugins/Enterprise/DatasmithCADImporter/Source/DatasmithCADTranslator/Public/DatasmithSceneGraphBuilder.h
+++ b/Engine/Plugins/Enterprise/DatasmithCADImporter/Source/DatasmithCADTranslator/Public/DatasmithSceneGraphBuilder.h
@@ -5,10 +5,10 @@
 
 #include "CADData.h"
 #include "CADSceneGraph.h"
-#include "CoreTechHelper.h"
 #include "CoreTechTypes.h"
 #include "CTSession.h"
 #include "DatasmithImportOptions.h"
+#include "MeshDescriptionHelper.h"
 
 #include "Containers/Map.h"
 #include "Containers/Queue.h"
diff --git a/Engine/Plugins/Enterprise/DatasmithCADImporter/Source/DatasmithOpenNurbsTranslator/DatasmithOpenNurbsTranslator.Build.cs b/Engine/Plugins/Enterprise/DatasmithCADImporter/Source/DatasmithOpenNurbsTranslator/DatasmithOpenNurbsTranslator.Build.cs
--- a/Engine/Plugins/Enterprise/DatasmithCADImporter/Source/DatasmithOpenNurbsTranslator/DatasmithOpenNurbsTranslator.Build.cs
+++ b/Engine/Plugins/Enterprise/DatasmithCADImporter/Source/DatasmithOpenNurbsTranslator/DatasmithOpenNurbsTranslator.Build.cs
@@ -28,6 +28,7 @@ namespace UnrealBuildTool.Rules
 					"DatasmithTranslator",
 					"Engine",
 					"MeshDescription",
+					"ParametricSurface",
 					"StaticMeshDescription",
 				}
 			);
diff --git a/Engine/Plugins/Enterprise/DatasmithCADImporter/Source/DatasmithOpenNurbsTranslator/Private/DatasmithOpenNurbsTranslator.cpp b/Engine/Plugins/Enterprise/DatasmithCADImporter/Source/DatasmithOpenNurbsTranslator/Private/DatasmithOpenNurbsTranslator.cpp
--- a/Engine/Plugins/Enterprise/DatasmithCADImporter/Source/DatasmithOpenNurbsTranslator/Private/DatasmithOpenNurbsTranslator.cpp
+++ b/Engine/Plugins/Enterprise/DatasmithCADImporter/Source/DatasmithOpenNurbsTranslator/Private/DatasmithOpenNurbsTranslator.cpp
@@ -5,10 +5,11 @@
 
 #ifdef USE_OPENNURBS // The whole translation unit is skipped without OpenNurbs TPS library
 
-#include "CoreTechSurfaceExtension.h"
 #include "RhinoCoretechWrapper.h"
 
 #include "CADInterfacesModule.h"
+#include "CoreTechSurfaceExtension.h"
+#include "CoreTechSurfaceHelper.h"
 #include "DatasmithImportOptions.h"
 #include "DatasmithMaterialElements.h"
 #include "DatasmithMaterialsUtils.h"
@@ -115,7 +116,7 @@ TSharedPtr<IDatasmithMetaDataElement> DuplicateMetaDataElement(const TSharedPtr<
 		if (SourceProperty)
 		{
 			TSharedRef<IDatasmithKeyValueProperty> DuplicatedProperty = FDatasmithSceneFactory::CreateKeyValueProperty(SourceProperty->GetName());
-			
+
 			DuplicatedProperty->SetValue(SourceProperty->GetValue());
 			DuplicatedMetaData->AddProperty(DuplicatedProperty);
 		}
@@ -128,7 +129,7 @@ TSharedPtr<IDatasmithMetaDataElement> DuplicateMetaDataElement(const TSharedPtr<
 class FOpenNurbsObjectWrapper
 {
 public:
-	ON_Object *ObjectPtr;
+	ON_Object* ObjectPtr;
 	ON_3dmObjectAttributes Attributes;
 
 	FOpenNurbsObjectWrapper()
@@ -189,7 +190,7 @@ namespace DatasmithOpenNurbsTranslatorUtils
 
 	bool HasPackedTextureRegion(const ON_Mesh& mesh)
 	{
-		return ( ON_IsValid(mesh.m_srf_scale[0]) && mesh.m_srf_scale[0] > 0.0
+		return (ON_IsValid(mesh.m_srf_scale[0]) && mesh.m_srf_scale[0] > 0.0
 			&& ON_IsValid(mesh.m_srf_scale[1]) && mesh.m_srf_scale[1] > 0.0
 			&& mesh.m_packed_tex_domain[0].IsInterval()
 			&& ON_Interval::ZeroToOne.Includes(mesh.m_packed_tex_domain[0])
@@ -198,7 +199,7 @@ namespace DatasmithOpenNurbsTranslatorUtils
 			);
 	}
 
-	FVector2D GetMeshTexCoords( const ON_Mesh *mesh, const int vertexCount, const int texCoordIndex, bool hasPackedTexCoords )
+	FVector2D GetMeshTexCoords(const ON_Mesh* mesh, const int vertexCount, const int texCoordIndex, bool hasPackedTexCoords)
 	{
 		// Ref. getMeshTexCoords
 		// Use values in m_T if number of values in m_T matches number of vertices
@@ -554,7 +555,7 @@ namespace DatasmithOpenNurbsTranslatorUtils
 		FMD5 MD5;
 
 		// Hash the material properties that are used to create the Unreal material: diffuse color, transparency, shininess and texture maps
-		unsigned int ColorRef = (unsigned int) Material.Diffuse();
+		unsigned int ColorRef = (unsigned int)Material.Diffuse();
 		MD5.Update(reinterpret_cast<const uint8*>(&ColorRef), sizeof(unsigned int));
 
 		double Transparency = Material.Transparency();
@@ -568,7 +569,7 @@ namespace DatasmithOpenNurbsTranslatorUtils
 
 		for (int Index = 0; Index < Material.m_textures.Count(); ++Index)
 		{
-			const ON_Texture *Texture = Material.m_textures.At(Index);
+			const ON_Texture* Texture = Material.m_textures.At(Index);
 
 			if (!Texture->m_bOn || (Texture->m_type != ON_Texture::TYPE::bitmap_texture && Texture->m_type != ON_Texture::TYPE::bump_texture && Texture->m_type != ON_Texture::TYPE::transparency_texture))
 			{
@@ -668,7 +669,7 @@ public:
 	void SetBaseOptions(const FDatasmithImportBaseOptions& InBaseOptions);
 	void SetOpenNurbsOptions(const FDatasmithOpenNurbsOptions& Options);
 	void SetOutputPath(const FString& Path) { OutputPath = Path; }
-	double GetScalingFactor () const { return ScalingFactor; }
+	double GetScalingFactor() const { return ScalingFactor; }
 	double GetMetricUnit() const { return MetricUnit; }
 
 	void ShowMessageLog(const FString& Filename);
@@ -704,7 +705,7 @@ private:
 	TSharedPtr<IDatasmithBaseMaterialElement> GetMaterial(int MaterialIndex);
 	TSharedPtr<IDatasmithBaseMaterialElement> GetDefaultMaterial();
 
-	TSharedPtr<IDatasmithActorElement> GetParentElement(const FOpenNurbsObjectWrapper & Object);
+	TSharedPtr<IDatasmithActorElement> GetParentElement(const FOpenNurbsObjectWrapper& Object);
 	FString GetLayerName(const TSharedPtr<IDatasmithActorElement>& LayerElement);
 	void SetLayers(const TSharedPtr<IDatasmithActorElement>& ActorElement, const FOpenNurbsObjectWrapper& Object);
 	void SetTags(const TSharedPtr<IDatasmithActorElement>& ActorElement, const FOpenNurbsObjectWrapper& Object);
@@ -739,7 +740,7 @@ private:
 	FDatasmithOpenNurbsOptions OpenNurbsOptions;
 	uint32 OpenNurbsOptionsHash;
 	FDatasmithImportBaseOptions BaseOptions;
-	
+
 	TSharedPtr<FRhinoCoretechWrapper> LocalSession;
 
 private:
@@ -815,7 +816,7 @@ private:
 	TMap< IDatasmithMeshElement*, const FOpenNurbsObjectWrapper* > MeshElementToObjectMap;
 
 	/** OpenNurbs objects to Datasmith mesh elements */
-	TMap<const FOpenNurbsObjectWrapper* , TSharedPtr< IDatasmithMeshElement > > ObjectToMeshElementMap;
+	TMap<const FOpenNurbsObjectWrapper*, TSharedPtr< IDatasmithMeshElement > > ObjectToMeshElementMap;
 
 	TMap< TSharedPtr< IDatasmithMeshElement >, ON_3dVector > MeshElementToGeometryCenter;
 
@@ -837,7 +838,7 @@ void FOpenNurbsTranslatorImpl::ShowMessageLog(const FString& Filename)
 			FText::Format(LOCTEXT("DatasmithOpenNurbsTranslator_NoMeshDataForAllMeshes", "Rhino model \"{0}\" doesn't contain mesh data for all objects. \nEither resave the 3dm file with a \"rendered view\" or change the import settings to \"Import as NURBS, Tessellate in Unreal\""), FText::FromString(Filename))
 		));
 
-		for (const FString& Name: MissingRenderMeshes)
+		for (const FString& Name : MissingRenderMeshes)
 		{
 			FText ErrorMessage = FText::Format(LOCTEXT("DatasmithOpenNurbsTranslator_NoMesh", "  {0} doesn't have mesh information."), FText::FromString(Name));
 			LogListing->AddMessage(FTokenizedMessage::Create(EMessageSeverity::Error, ErrorMessage));
@@ -861,7 +862,7 @@ void FOpenNurbsTranslatorImpl::TranslateMaterialTable(const ON_ObjectArray<ON_Ma
 	// These materials are from the Materials tab and do not include materials from layers
 	for (int Index = 0; Index < InMaterialTable.Count(); ++Index)
 	{
-		const ON_Material &OpenNurbsMaterial = *InMaterialTable.At(Index);
+		const ON_Material& OpenNurbsMaterial = *InMaterialTable.At(Index);
 
 		FMD5Hash Hash = DatasmithOpenNurbsTranslatorUtils::ComputeMaterialHash(OpenNurbsMaterial);
 		TSharedPtr<IDatasmithBaseMaterialElement>* MaterialPtr = HashToMaterial.Find(Hash);
@@ -881,7 +882,7 @@ void FOpenNurbsTranslatorImpl::TranslateMaterialTable(const ON_ObjectArray<ON_Ma
 
 		// Note that in OpenNurbs, Alpha means Transparency, whereas it is usually an Opacity.
 		// Hence the (255 - Transparency) where an opacity is expected
-		FColor Color((uint8) Diffuse.Red(), (uint8) Diffuse.Green(), (uint8) Diffuse.Blue(), (uint8) (255 - Transparency));
+		FColor Color((uint8)Diffuse.Red(), (uint8)Diffuse.Green(), (uint8)Diffuse.Blue(), (uint8)(255 - Transparency));
 		FLinearColor LinearColor = FLinearColor::FromSRGBColor(Color);
 
 		FString MaterialLabel(OpenNurbsMaterial.Name().Array());
@@ -904,7 +905,7 @@ void FOpenNurbsTranslatorImpl::TranslateMaterialTable(const ON_ObjectArray<ON_Ma
 		{
 			for (int TextureIndex = 0; TextureIndex < OpenNurbsMaterial.m_textures.Count(); ++TextureIndex)
 			{
-				const ON_Texture *Texture = OpenNurbsMaterial.m_textures.At(TextureIndex);
+				const ON_Texture* Texture = OpenNurbsMaterial.m_textures.At(TextureIndex);
 
 				if (!Texture->m_bOn || (Texture->m_type != ON_Texture::TYPE::bitmap_texture && Texture->m_type != ON_Texture::TYPE::bump_texture && Texture->m_type != ON_Texture::TYPE::transparency_texture))
 				{
@@ -1034,8 +1035,8 @@ void FOpenNurbsTranslatorImpl::TranslateMaterialTable(const ON_ObjectArray<ON_Ma
 
 				UVParameters.UVTiling.Y = Tiling.Y;
 
-				if ( !FMath::IsNearlyZero( Tiling.X, KINDA_SMALL_NUMBER )
-					&& !FMath::IsNearlyZero( Tiling.Y, KINDA_SMALL_NUMBER ) )
+				if (!FMath::IsNearlyZero(Tiling.X, KINDA_SMALL_NUMBER)
+					&& !FMath::IsNearlyZero(Tiling.Y, KINDA_SMALL_NUMBER))
 				{
 					UVParameters.UVOffset.X = Translation.X / Tiling.X;
 
@@ -1106,7 +1107,7 @@ void FOpenNurbsTranslatorImpl::TranslateMaterialTable(const ON_ObjectArray<ON_Ma
 			{
 				// Transparent color
 				IDatasmithMaterialExpressionScalar* Scalar = static_cast<IDatasmithMaterialExpressionScalar*>(Material->AddMaterialExpression(EDatasmithMaterialExpressionType::ConstantScalar));
-				Scalar->SetName( TEXT( "Opacity" ) );
+				Scalar->SetName(TEXT("Opacity"));
 				Scalar->GetScalar() = LinearColor.A;
 
 				Material->GetOpacity().SetExpression(Scalar);
@@ -1115,10 +1116,10 @@ void FOpenNurbsTranslatorImpl::TranslateMaterialTable(const ON_ObjectArray<ON_Ma
 			{
 				// Modulate the opacity map with the color transparency setting
 				IDatasmithMaterialExpressionGeneric* Multiply = static_cast<IDatasmithMaterialExpressionGeneric*>(Material->AddMaterialExpression(EDatasmithMaterialExpressionType::Generic));
-				Multiply->SetExpressionName( TEXT( "Multiply" ) );
+				Multiply->SetExpressionName(TEXT("Multiply"));
 
 				IDatasmithMaterialExpressionScalar* Scalar = static_cast<IDatasmithMaterialExpressionScalar*>(Material->AddMaterialExpression(EDatasmithMaterialExpressionType::ConstantScalar));
-				Scalar->SetName( TEXT( "Opacity Output Level" ) );
+				Scalar->SetName(TEXT("Opacity Output Level"));
 				Scalar->GetScalar() = LinearColor.A;
 				Scalar->ConnectExpression(*Multiply->GetInput(0));
 
@@ -1134,7 +1135,7 @@ void FOpenNurbsTranslatorImpl::TranslateMaterialTable(const ON_ObjectArray<ON_Ma
 		if (!FMath::IsNearlyZero(Shininess))
 		{
 			IDatasmithMaterialExpressionScalar* Scalar = static_cast<IDatasmithMaterialExpressionScalar*>(Material->AddMaterialExpression(EDatasmithMaterialExpressionType::ConstantScalar));
-			Scalar->SetName( TEXT( "Roughness" ) );
+			Scalar->SetName(TEXT("Roughness"));
 			Scalar->GetScalar() = 1.f - Shininess;
 			Material->GetRoughness().SetExpression(Scalar);
 		}
@@ -1143,7 +1144,7 @@ void FOpenNurbsTranslatorImpl::TranslateMaterialTable(const ON_ObjectArray<ON_Ma
 		if (!FMath::IsNearlyZero(Reflectivity))
 		{
 			IDatasmithMaterialExpressionScalar* Scalar = static_cast<IDatasmithMaterialExpressionScalar*>(Material->AddMaterialExpression(EDatasmithMaterialExpressionType::ConstantScalar));
-			Scalar->SetName( TEXT( "Metallic" ) );
+			Scalar->SetName(TEXT("Metallic"));
 			Scalar->GetScalar() = Reflectivity;
 			Material->GetMetallic().SetExpression(Scalar);
 		}
@@ -1156,7 +1157,7 @@ void FOpenNurbsTranslatorImpl::TranslateLayerTable(const ON_ObjectArray<ON_Layer
 	TSet<ON_UUID> HiddenLayersUUIDs;
 	for (int Index = 0; Index < InLayerTable.Count(); ++Index)
 	{
-		const ON_Layer &CurrentLayer = *InLayerTable.At(Index);
+		const ON_Layer& CurrentLayer = *InLayerTable.At(Index);
 
 		TSharedPtr<IDatasmithActorElement> Parent;
 		if (ON_UuidCompare(ON_nil_uuid, CurrentLayer.ParentLayerId()) != 0)
@@ -1186,7 +1187,7 @@ void FOpenNurbsTranslatorImpl::TranslateLayerTable(const ON_ObjectArray<ON_Layer
 		{
 			Parent->AddChild(LayerElement);
 			FString ParentLayerName = GetLayerName(Parent);
-			FullLayerName =  ParentLayerName + TEXT(".") + FullLayerName;
+			FullLayerName = ParentLayerName + TEXT(".") + FullLayerName;
 		}
 		else
 		{
@@ -1216,7 +1217,7 @@ void FOpenNurbsTranslatorImpl::TranslateGroupTable(const ON_ObjectArray<ON_Group
 {
 	for (int Index = 0; Index < InGroupTable.Count(); ++Index)
 	{
-		const ON_Group &Group = *InGroupTable.At(Index);
+		const ON_Group& Group = *InGroupTable.At(Index);
 		GroupNames.Add(Group.Name().IsEmpty() ? FString::Printf(TEXT("Group%d"), Index) : Group.Name().Array());
 	}
 }
@@ -1286,7 +1287,7 @@ void FOpenNurbsTranslatorImpl::TranslateLightTable(const ON_ClassArray<FOpenNurb
 		LightElement->SetEnabled(LightObj.m_bOn);
 
 		// Diffuse color (Ambient and Specular color not supported and alpha from diffuse is ignored)
-		FColor Color((uint8) LightObj.Diffuse().Red(), (uint8) LightObj.Diffuse().Green(), (uint8) LightObj.Diffuse().Blue(), 255);
+		FColor Color((uint8)LightObj.Diffuse().Red(), (uint8)LightObj.Diffuse().Green(), (uint8)LightObj.Diffuse().Blue(), 255);
 		LightElement->SetColor(Color.ReinterpretAsLinear());
 
 		// Intensity (PowerWatts and ShadowIntensity not used)
@@ -1411,7 +1412,7 @@ void FOpenNurbsTranslatorImpl::TranslateInstanceDefinitionTable(const TArray<ON_
 	{
 		++instanceDefCount;
 
-		const ON_InstanceDefinition &instanceDef = *pInstanceDef;
+		const ON_InstanceDefinition& instanceDef = *pInstanceDef;
 		const ON_UUID& instanceDefUuid = instanceDef.Id();
 
 		FString InstanceDefName;
@@ -1645,48 +1646,48 @@ void FOpenNurbsTranslatorImpl::SetTags(const TSharedPtr<IDatasmithActorElement>&
 	const TCHAR* StrObjectType;
 	switch (objType)
 	{
-		case ON::instance_definition:
-			StrObjectType = TEXT("block definition");
-			break;
-		case ON::instance_reference:
-			StrObjectType = TEXT("block instance");
-			break;
-		case ON::point_object:
-			StrObjectType = TEXT("point");
-			break;
-		case ON::curve_object:
-			StrObjectType = TEXT("curve");
-			break;
-		case ON::surface_object:
-			StrObjectType = TEXT("surface");
-			break;
-		case ON::brep_object:
-			StrObjectType = TEXT("brep");
-			break;
-		case ON::mesh_object:
-			StrObjectType = TEXT("mesh");
-			break;
-		case ON::text_dot:
-			StrObjectType = TEXT("textdot");
-			break;
-		case ON::subd_object:
-			StrObjectType = TEXT("subd");
-			break;
-		case ON::loop_object:
-			StrObjectType = TEXT("loop");
-			break;
-		case ON::cage_object:
-			StrObjectType = TEXT("cage");
-			break;
-		case ON::clipplane_object:
-			StrObjectType = TEXT("clip plane");
-			break;
-		case ON::extrusion_object:
-			StrObjectType = TEXT("extrusion");
-			break;
-		default:
-			StrObjectType = TEXT("unknown");
-			break;
+	case ON::instance_definition:
+		StrObjectType = TEXT("block definition");
+		break;
+	case ON::instance_reference:
+		StrObjectType = TEXT("block instance");
+		break;
+	case ON::point_object:
+		StrObjectType = TEXT("point");
+		break;
+	case ON::curve_object:
+		StrObjectType = TEXT("curve");
+		break;
+	case ON::surface_object:
+		StrObjectType = TEXT("surface");
+		break;
+	case ON::brep_object:
+		StrObjectType = TEXT("brep");
+		break;
+	case ON::mesh_object:
+		StrObjectType = TEXT("mesh");
+		break;
+	case ON::text_dot:
+		StrObjectType = TEXT("textdot");
+		break;
+	case ON::subd_object:
+		StrObjectType = TEXT("subd");
+		break;
+	case ON::loop_object:
+		StrObjectType = TEXT("loop");
+		break;
+	case ON::cage_object:
+		StrObjectType = TEXT("cage");
+		break;
+	case ON::clipplane_object:
+		StrObjectType = TEXT("clip plane");
+		break;
+	case ON::extrusion_object:
+		StrObjectType = TEXT("extrusion");
+		break;
+	default:
+		StrObjectType = TEXT("unknown");
+		break;
 	}
 
 	ActorElement->AddTag(*FString::Printf(TEXT("Rhino.ID: %s"), *UUID));
@@ -1814,7 +1815,7 @@ bool FOpenNurbsTranslatorImpl::HasUnprocessedChildren(const ON_UUID& instanceDef
 bool FOpenNurbsTranslatorImpl::TranslateInstance(const FOpenNurbsObjectWrapper& Object)
 {
 	// Ref. visitInstance
-	ON_InstanceRef *instanceRef = ON_InstanceRef::Cast(Object.ObjectPtr);
+	ON_InstanceRef* instanceRef = ON_InstanceRef::Cast(Object.ObjectPtr);
 	if (instanceRef == nullptr)
 	{
 		return true;
@@ -2185,22 +2186,22 @@ FOpenNurbsTranslatorImpl::FMaterial FOpenNurbsTranslatorImpl::GetObjectMaterial(
 
 	switch (MaterialSource)
 	{
-		case ON::material_from_object:
+	case ON::material_from_object:
+	{
+		if (Object.Attributes.m_material_index != -1)
 		{
-			if (Object.Attributes.m_material_index != -1)
-			{
-				return { GetMaterial(Object.Attributes.m_material_index) , GetOpenNurbsMaterial(Object.Attributes.m_material_index) };
-			}
-			break;
+			return { GetMaterial(Object.Attributes.m_material_index) , GetOpenNurbsMaterial(Object.Attributes.m_material_index) };
 		}
-		case ON::material_from_layer:
+		break;
+	}
+	case ON::material_from_layer:
+	{
+		if (int* MaterialIndexPtr = LayerIndexToMaterialIndex.Find(Object.Attributes.m_layer_index))
 		{
-			if (int* MaterialIndexPtr = LayerIndexToMaterialIndex.Find(Object.Attributes.m_layer_index))
-			{
-				return { GetMaterial(*MaterialIndexPtr), GetOpenNurbsMaterial(*MaterialIndexPtr) };
-			}
-			break;
+			return { GetMaterial(*MaterialIndexPtr), GetOpenNurbsMaterial(*MaterialIndexPtr) };
 		}
+		break;
+	}
 	}
 
 	return { GetDefaultMaterial() };
@@ -2254,7 +2255,7 @@ TSharedPtr<IDatasmithBaseMaterialElement> FOpenNurbsTranslatorImpl::GetDefaultMa
 	return DefaultMaterial;
 }
 
-TSharedPtr<IDatasmithActorElement> FOpenNurbsTranslatorImpl::GetParentElement(const FOpenNurbsObjectWrapper & Object)
+TSharedPtr<IDatasmithActorElement> FOpenNurbsTranslatorImpl::GetParentElement(const FOpenNurbsObjectWrapper& Object)
 {
 	TSharedPtr<IDatasmithActorElement>* Parent = LayerIndexToContainer.Find(Object.Attributes.m_layer_index);
 	if (Parent)
@@ -3343,80 +3344,80 @@ TOptional<FMeshDescription> FOpenNurbsTranslatorImpl::GetMeshDescription(TShared
 
 	const FOpenNurbsObjectWrapper& Object = **ObjectPtr;
 
-ON_wString UUIDString;
-ON_UuidToString(Object.Attributes.m_uuid, UUIDString);
-FString UUID(UUIDString.Array());
+	ON_wString UUIDString;
+	ON_UuidToString(Object.Attributes.m_uuid, UUIDString);
+	FString UUID(UUIDString.Array());
 
-FMeshDescription MeshDescription;
-DatasmithMeshHelper::PrepareAttributeForStaticMesh(MeshDescription);
+	FMeshDescription MeshDescription;
+	DatasmithMeshHelper::PrepareAttributeForStaticMesh(MeshDescription);
 
-bool bHasNormal = false;
-bool bIsValid = false;
-if (Object.ObjectPtr->IsKindOf(&ON_Mesh::m_ON_Mesh_class_rtti))
-{
-	ON_3dVector Offset = GetGeometryOffset(MeshElement);
-	bIsValid = DatasmithOpenNurbsTranslatorUtils::TranslateMesh(ON_Mesh::Cast(Object.ObjectPtr), MeshDescription, bHasNormal, SelectedTranslator->ScalingFactor, Offset);
-}
-else if (Object.ObjectPtr->IsKindOf(&ON_Brep::m_ON_Brep_class_rtti))
-{
-	bIsValid = SelectedTranslator->TranslateBRep(ON_Brep::Cast(Object.ObjectPtr), Object.Attributes, MeshDescription, MeshElement, UUID, bHasNormal);
-}
-else if (const ON_Extrusion* extrusion = ON_Extrusion::Cast(Object.ObjectPtr))
-{
-	if (OpenNurbsOptions.Geometry == EDatasmithOpenNurbsBrepTessellatedSource::UseRenderMeshes)
+	bool bHasNormal = false;
+	bool bIsValid = false;
+	if (Object.ObjectPtr->IsKindOf(&ON_Mesh::m_ON_Mesh_class_rtti))
 	{
 		ON_3dVector Offset = GetGeometryOffset(MeshElement);
-		if (const ON_Mesh* Mesh = extrusion->m_mesh_cache.Mesh(ON::mesh_type::render_mesh))
+		bIsValid = DatasmithOpenNurbsTranslatorUtils::TranslateMesh(ON_Mesh::Cast(Object.ObjectPtr), MeshDescription, bHasNormal, SelectedTranslator->ScalingFactor, Offset);
+	}
+	else if (Object.ObjectPtr->IsKindOf(&ON_Brep::m_ON_Brep_class_rtti))
+	{
+		bIsValid = SelectedTranslator->TranslateBRep(ON_Brep::Cast(Object.ObjectPtr), Object.Attributes, MeshDescription, MeshElement, UUID, bHasNormal);
+	}
+	else if (const ON_Extrusion* extrusion = ON_Extrusion::Cast(Object.ObjectPtr))
+	{
+		if (OpenNurbsOptions.Geometry == EDatasmithOpenNurbsBrepTessellatedSource::UseRenderMeshes)
 		{
-			if (DatasmithOpenNurbsTranslatorUtils::TranslateMesh(Mesh, MeshDescription, bHasNormal, ScalingFactor, Offset))
+			ON_3dVector Offset = GetGeometryOffset(MeshElement);
+			if (const ON_Mesh* Mesh = extrusion->m_mesh_cache.Mesh(ON::mesh_type::render_mesh))
+			{
+				if (DatasmithOpenNurbsTranslatorUtils::TranslateMesh(Mesh, MeshDescription, bHasNormal, ScalingFactor, Offset))
+				{
+					return MeshDescription;
+				}
+			}
+			else
 			{
-				return MeshDescription;
+				MissingRenderMeshes.Add(MeshElement->GetLabel());
 			}
 		}
-		else
+
+		ON_Brep brep;
+		if (extrusion->BrepForm(&brep) != nullptr)
 		{
-			MissingRenderMeshes.Add(MeshElement->GetLabel());
+			bIsValid = SelectedTranslator->TranslateBRep(&brep, Object.Attributes, MeshDescription, MeshElement, UUID, bHasNormal);
 		}
 	}
-
-	ON_Brep brep;
-	if (extrusion->BrepForm(&brep) != nullptr)
+	else if (Object.ObjectPtr->IsKindOf(&ON_Hatch::m_ON_Hatch_class_rtti))
 	{
-		bIsValid = SelectedTranslator->TranslateBRep(&brep, Object.Attributes, MeshDescription, MeshElement, UUID, bHasNormal);
+		const ON_Hatch* hatch = ON_Hatch::Cast(Object.ObjectPtr);
+		ON_Brep brep;
+		if (hatch != nullptr && hatch->BrepForm(&brep) != nullptr)
+		{
+			bIsValid = SelectedTranslator->TranslateBRep(&brep, Object.Attributes, MeshDescription, MeshElement, UUID, bHasNormal);
+		}
 	}
-}
-else if (Object.ObjectPtr->IsKindOf(&ON_Hatch::m_ON_Hatch_class_rtti))
-{
-	const ON_Hatch* hatch = ON_Hatch::Cast(Object.ObjectPtr);
-	ON_Brep brep;
-	if (hatch != nullptr && hatch->BrepForm(&brep) != nullptr)
+	else if (Object.ObjectPtr->IsKindOf(&ON_PlaneSurface::m_ON_PlaneSurface_class_rtti))
 	{
-		bIsValid = SelectedTranslator->TranslateBRep(&brep, Object.Attributes, MeshDescription, MeshElement, UUID, bHasNormal);
+		const ON_PlaneSurface* planeSurface = ON_PlaneSurface::Cast(Object.ObjectPtr);
+		ON_Interval xInterval = planeSurface->Extents(0);
+		ON_Interval yInterval = planeSurface->Extents(1);
+
+		// TODO: Create simple plane
+		// 3 --- 2
+		// |     |
+		// |     |
+		// 0 --- 1
+
+		ON_3dVector vec0 = planeSurface->m_plane.origin + planeSurface->m_plane.xaxis * xInterval.m_t[0] + planeSurface->m_plane.yaxis * yInterval.m_t[0];
+		ON_3dVector vec1 = planeSurface->m_plane.origin + planeSurface->m_plane.xaxis * xInterval.m_t[1] + planeSurface->m_plane.yaxis * yInterval.m_t[0];
+		ON_3dVector vec2 = planeSurface->m_plane.origin + planeSurface->m_plane.xaxis * xInterval.m_t[1] + planeSurface->m_plane.yaxis * yInterval.m_t[1];
+		ON_3dVector vec3 = planeSurface->m_plane.origin + planeSurface->m_plane.xaxis * xInterval.m_t[0] + planeSurface->m_plane.yaxis * yInterval.m_t[1];
+	}
+	else if (Object.ObjectPtr->IsKindOf(&ON_LineCurve::m_ON_LineCurve_class_rtti))
+	{
+		// Not supported
 	}
-}
-else if (Object.ObjectPtr->IsKindOf(&ON_PlaneSurface::m_ON_PlaneSurface_class_rtti))
-{
-	const ON_PlaneSurface* planeSurface = ON_PlaneSurface::Cast(Object.ObjectPtr);
-	ON_Interval xInterval = planeSurface->Extents(0);
-	ON_Interval yInterval = planeSurface->Extents(1);
-
-	// TODO: Create simple plane
-	// 3 --- 2
-	// |     |
-	// |     |
-	// 0 --- 1
-
-	ON_3dVector vec0 = planeSurface->m_plane.origin + planeSurface->m_plane.xaxis * xInterval.m_t[0] + planeSurface->m_plane.yaxis * yInterval.m_t[0];
-	ON_3dVector vec1 = planeSurface->m_plane.origin + planeSurface->m_plane.xaxis * xInterval.m_t[1] + planeSurface->m_plane.yaxis * yInterval.m_t[0];
-	ON_3dVector vec2 = planeSurface->m_plane.origin + planeSurface->m_plane.xaxis * xInterval.m_t[1] + planeSurface->m_plane.yaxis * yInterval.m_t[1];
-	ON_3dVector vec3 = planeSurface->m_plane.origin + planeSurface->m_plane.xaxis * xInterval.m_t[0] + planeSurface->m_plane.yaxis * yInterval.m_t[1];
-}
-else if (Object.ObjectPtr->IsKindOf(&ON_LineCurve::m_ON_LineCurve_class_rtti))
-{
-	// Not supported
-}
 
-return bIsValid ? MoveTemp(MeshDescription) : TOptional< FMeshDescription >();
+	return bIsValid ? MoveTemp(MeshDescription) : TOptional< FMeshDescription >();
 }
 
 void FOpenNurbsTranslatorImpl::SetBaseOptions(const FDatasmithImportBaseOptions& InBaseOptions)
@@ -3444,7 +3445,7 @@ void FDatasmithOpenNurbsTranslator::Initialize(FDatasmithTranslatorCapabilities&
 #if WITH_EDITOR
 	if (GIsEditor && !GEditor->PlayWorld && !GIsPlayInEditorWorld)
 	{
-		OutCapabilities.SupportedFileFormats.Add(FFileFormatInfo{TEXT("3dm"), TEXT("Rhino file format")});
+		OutCapabilities.SupportedFileFormats.Add(FFileFormatInfo{ TEXT("3dm"), TEXT("Rhino file format") });
 		return;
 	}
 #endif
@@ -3504,7 +3505,14 @@ bool FDatasmithOpenNurbsTranslator::LoadStaticMesh(const TSharedRef<IDatasmithMe
 
 		CADLibrary::FMeshParameters MeshParameters;
 
-		CoreTechSurface::AddCoreTechSurfaceDataForMesh(MeshElement, ImportParameters, MeshParameters, OpenNurbsOptions, OutMeshPayload);
+		if (ImportParameters.bEnableKernelIOTessellation)
+		{
+			CoreTechSurface::AddSurfaceDataForMesh(MeshElement, ImportParameters, MeshParameters, OpenNurbsOptions, OutMeshPayload);
+		}
+		else
+		{
+			//CADKernelSurface::AddSurfaceDataForMesh(MeshElement, ImportParameters, MeshParameters, OpenNurbsOptions, OutMeshPayload);
+		}
 	}
 
 	return OutMeshPayload.LodMeshes.Num() > 0;
diff --git a/Engine/Plugins/Enterprise/DatasmithCADImporter/Source/DatasmithOpenNurbsTranslator/Private/RhinoCoretechWrapper.cpp b/Engine/Plugins/Enterprise/DatasmithCADImporter/Source/DatasmithOpenNurbsTranslator/Private/RhinoCoretechWrapper.cpp
--- a/Engine/Plugins/Enterprise/DatasmithCADImporter/Source/DatasmithOpenNurbsTranslator/Private/RhinoCoretechWrapper.cpp
+++ b/Engine/Plugins/Enterprise/DatasmithCADImporter/Source/DatasmithOpenNurbsTranslator/Private/RhinoCoretechWrapper.cpp
@@ -4,7 +4,7 @@
 
 
 #ifdef USE_OPENNURBS
-#include "CoreTechHelper.h"
+#include "CoreTechSurfaceHelper.h"
 
 #pragma warning(push)
 #pragma warning(disable:4265)
@@ -399,7 +399,7 @@ void BRepToKernelIOBodyTranslator::CreateCTFace(const ON_BrepFace& Face, TArray<
 
 bool FRhinoCoretechWrapper::Tessellate(FMeshDescription& Mesh, CADLibrary::FMeshParameters& MeshParameters)
 {
-	return CADLibrary::Tessellate(MainObjectId, ImportParams, Mesh, MeshParameters);
+	return CoreTechSurface::Tessellate(MainObjectId, ImportParams, Mesh, MeshParameters);
 }
 
 uint64 BRepToKernelIOBodyTranslator::CreateBody(const ON_3dVector& Offset)
diff --git a/Engine/Plugins/Enterprise/DatasmithCADImporter/Source/DatasmithPLMXMLTranslator/DatasmithPLMXMLTranslator.Build.cs b/Engine/Plugins/Enterprise/DatasmithCADImporter/Source/DatasmithPLMXMLTranslator/DatasmithPLMXMLTranslator.Build.cs
--- a/Engine/Plugins/Enterprise/DatasmithCADImporter/Source/DatasmithPLMXMLTranslator/DatasmithPLMXMLTranslator.Build.cs
+++ b/Engine/Plugins/Enterprise/DatasmithCADImporter/Source/DatasmithPLMXMLTranslator/DatasmithPLMXMLTranslator.Build.cs
@@ -34,6 +34,7 @@ namespace UnrealBuildTool.Rules
 					"DatasmithCADTranslator", // for DatasmithMeshBuilder
                     "DatasmithTranslator",
                     "DatasmithDispatcher",
+					"ParametricSurface",
 
                     "DatasmithContent",
                 }
diff --git a/Engine/Plugins/Enterprise/DatasmithCADImporter/Source/DatasmithPLMXMLTranslator/Private/DatasmithPlmXmlImporter.cpp b/Engine/Plugins/Enterprise/DatasmithCADImporter/Source/DatasmithPLMXMLTranslator/Private/DatasmithPlmXmlImporter.cpp
--- a/Engine/Plugins/Enterprise/DatasmithCADImporter/Source/DatasmithPLMXMLTranslator/Private/DatasmithPlmXmlImporter.cpp
+++ b/Engine/Plugins/Enterprise/DatasmithCADImporter/Source/DatasmithPLMXMLTranslator/Private/DatasmithPlmXmlImporter.cpp
@@ -20,6 +20,7 @@
 #include "CADData.h"
 #include "CADToolsModule.h"
 #include "CoreTechSurfaceExtension.h"
+#include "CoreTechSurfaceHelper.h"
 #include "DatasmithAdditionalData.h"
 #include "DatasmithDispatcher.h"
 #include "DatasmithMeshBuilder.h"
@@ -303,8 +304,15 @@ namespace PlmXml
 			{
 				OutMeshPayload.LodMeshes.Add(MoveTemp(Mesh.GetValue()));
 
-				CoreTechSurface::AddCoreTechSurfaceDataForMesh(MeshElement, ImportParameters, MeshParameters, TessellationOptions, OutMeshPayload);
-				
+				if (ImportParameters.bEnableKernelIOTessellation)
+				{
+					CoreTechSurface::AddSurfaceDataForMesh(MeshElement, ImportParameters, MeshParameters, TessellationOptions, OutMeshPayload);
+				}
+				else
+				{
+					//CADKernelSurface::AddSurfaceDataForMesh(MeshElement, ImportParameters, MeshParameters, TessellationOptions, OutMeshPayload);
+				}
+
 				return true;
 			}
 			return false;
diff --git a/Engine/Plugins/Enterprise/DatasmithCADImporter/Source/DatasmithWireTranslator/DatasmithWireTranslator.Build.cs b/Engine/Plugins/Enterprise/DatasmithCADImporter/Source/DatasmithWireTranslator/DatasmithWireTranslator.Build.cs
--- a/Engine/Plugins/Enterprise/DatasmithCADImporter/Source/DatasmithWireTranslator/DatasmithWireTranslator.Build.cs
+++ b/Engine/Plugins/Enterprise/DatasmithCADImporter/Source/DatasmithWireTranslator/DatasmithWireTranslator.Build.cs
@@ -27,7 +27,8 @@ public class DatasmithWireTranslator : ModuleRules
 				"DatasmithTranslator",
 				"Engine",
 				"MeshDescription",
-                "StaticMeshDescription",
+				"ParametricSurface",
+				"StaticMeshDescription",
             }
         );
 
diff --git a/Engine/Plugins/Enterprise/DatasmithCADImporter/Source/DatasmithWireTranslator/Private/AliasCoretechWrapper.cpp b/Engine/Plugins/Enterprise/DatasmithCADImporter/Source/DatasmithWireTranslator/Private/AliasCoretechWrapper.cpp
--- a/Engine/Plugins/Enterprise/DatasmithCADImporter/Source/DatasmithWireTranslator/Private/AliasCoretechWrapper.cpp
+++ b/Engine/Plugins/Enterprise/DatasmithCADImporter/Source/DatasmithWireTranslator/Private/AliasCoretechWrapper.cpp
@@ -4,7 +4,7 @@
 
 #include "Hal/PlatformMemory.h"
 
-#include "CoreTechHelper.h"
+#include "CoreTechSurfaceHelper.h"
 #include "CoreTechTypes.h"
 
 #ifdef USE_OPENMODEL
@@ -99,7 +99,7 @@ bool FAliasCoretechWrapper::Tessellate(FMeshDescription& Mesh, FMeshParameters&
 	TopoFixes(1.);
 
 	// Perform tessellation
-	return CADLibrary::Tessellate(MainObjectId, ImportParams, Mesh, MeshParameters);
+	return CoreTechSurface::Tessellate(MainObjectId, ImportParams, Mesh, MeshParameters);
 }
 
 uint64 FAliasCoretechWrapper::Add3DCurve(AlCurve& Curve)
diff --git a/Engine/Plugins/Enterprise/DatasmithCADImporter/Source/DatasmithWireTranslator/Private/DatasmithWireTranslator.cpp b/Engine/Plugins/Enterprise/DatasmithCADImporter/Source/DatasmithWireTranslator/Private/DatasmithWireTranslator.cpp
--- a/Engine/Plugins/Enterprise/DatasmithCADImporter/Source/DatasmithWireTranslator/Private/DatasmithWireTranslator.cpp
+++ b/Engine/Plugins/Enterprise/DatasmithCADImporter/Source/DatasmithWireTranslator/Private/DatasmithWireTranslator.cpp
@@ -28,6 +28,7 @@
 #include "AliasCoretechWrapper.h" // requires CoreTech as public dependency
 #include "CADInterfacesModule.h"
 #include "CoreTechSurfaceExtension.h"
+#include "CoreTechSurfaceHelper.h"
 
 #ifdef USE_OPENMODEL
 #include <AlChannel.h>
@@ -2119,7 +2120,14 @@ bool FDatasmithWireTranslator::LoadStaticMesh(const TSharedRef<IDatasmithMeshEle
 	{
 		OutMeshPayload.LodMeshes.Add(MoveTemp(Mesh.GetValue()));
 
-		CoreTechSurface::AddCoreTechSurfaceDataForMesh(MeshElement, ImportParameters, MeshParameters, GetCommonTessellationOptions(), OutMeshPayload);
+		if (ImportParameters.bEnableKernelIOTessellation)
+		{
+			CoreTechSurface::AddSurfaceDataForMesh(MeshElement, ImportParameters, MeshParameters, GetCommonTessellationOptions(), OutMeshPayload);
+		}
+		else
+		{
+			//CADKernelSurface::AddSurfaceDataForMesh(MeshElement, ImportParameters, MeshParameters, GetCommonTessellationOptions(), OutMeshPayload);
+		}
 	}
 	return OutMeshPayload.LodMeshes.Num() > 0;
 #else
@@ -2130,7 +2138,7 @@ bool FDatasmithWireTranslator::LoadStaticMesh(const TSharedRef<IDatasmithMeshEle
 void FDatasmithWireTranslator::SetSceneImportOptions(TArray<TStrongObjectPtr<UDatasmithOptionsBase>>& Options)
 {
 #ifdef USE_OPENMODEL
-	FCoreTechTranslator::SetSceneImportOptions(Options);
+	FParametricSurfaceTranslator::SetSceneImportOptions(Options);
 
 	if (Translator)
 	{
diff --git a/Engine/Plugins/Enterprise/DatasmithCADImporter/Source/DatasmithWireTranslator/Private/DatasmithWireTranslator.h b/Engine/Plugins/Enterprise/DatasmithCADImporter/Source/DatasmithWireTranslator/Private/DatasmithWireTranslator.h
--- a/Engine/Plugins/Enterprise/DatasmithCADImporter/Source/DatasmithWireTranslator/Private/DatasmithWireTranslator.h
+++ b/Engine/Plugins/Enterprise/DatasmithCADImporter/Source/DatasmithWireTranslator/Private/DatasmithWireTranslator.h
@@ -4,13 +4,13 @@
 
 #include "CoreMinimal.h"
 
-#include "CoreTechTranslator.h"
+#include "ParametricSurfaceTranslator.h"
 #include "UObject/ObjectMacros.h"
 
 class FWireTranslatorImpl;
 
 
-class FDatasmithWireTranslator : public FCoreTechTranslator
+class FDatasmithWireTranslator : public FParametricSurfaceTranslator
 {
 public:
 	FDatasmithWireTranslator();
diff --git a/Engine/Plugins/Enterprise/DatasmithCADImporter/Source/ParametricSurface/ParametricSurface.Build.cs b/Engine/Plugins/Enterprise/DatasmithCADImporter/Source/ParametricSurface/ParametricSurface.Build.cs
new file mode 100644
--- /dev/null
+++ b/Engine/Plugins/Enterprise/DatasmithCADImporter/Source/ParametricSurface/ParametricSurface.Build.cs
@@ -0,0 +1,29 @@
+// Copyright Epic Games, Inc. All Rights Reserved.
+
+using System.IO;
+
+namespace UnrealBuildTool.Rules
+{
+	public class ParametricSurface : ModuleRules
+	{
+		public ParametricSurface(ReadOnlyTargetRules Target) : base(Target)
+		{
+			PublicDependencyModuleNames.AddRange(
+				new string[]
+				{
+					"Core",
+				}
+			);
+
+			PrivateDependencyModuleNames.AddRange(
+				new string[]
+				{
+					"Core",
+					"CoreUObject",
+                    "DatasmithContent",
+					"DatasmithTranslator",
+				}
+			);
+		}
+	}
+}
diff --git a/Engine/Plugins/Enterprise/DatasmithCADImporter/Source/ParametricSurface/Private/ParametricSurfaceModule.cpp b/Engine/Plugins/Enterprise/DatasmithCADImporter/Source/ParametricSurface/Private/ParametricSurfaceModule.cpp
new file mode 100644
--- /dev/null
+++ b/Engine/Plugins/Enterprise/DatasmithCADImporter/Source/ParametricSurface/Private/ParametricSurfaceModule.cpp
@@ -0,0 +1,31 @@
+// Copyright Epic Games, Inc. All Rights Reserved.
+
+#include "ParametricSurfaceModule.h"
+
+#include "Modules/ModuleInterface.h"
+#include "Modules/ModuleManager.h"
+#include "UObject/CoreRedirects.h"
+
+#define LOCTEXT_NAMESPACE "ParametricSurfaceModule"
+
+void FParametricSurfaceModule::StartupModule()
+{
+	TArray<FCoreRedirect> Redirects;
+	Redirects.Emplace(ECoreRedirectFlags::Type_Property, TEXT("UParametricSurfaceData.RawData"), TEXT("RawData_DEPRECATED"));
+	FCoreRedirects::AddRedirectList(Redirects, PARAMETRICSURFACE_MODULE_NAME);
+}
+
+FParametricSurfaceModule& FParametricSurfaceModule::Get()
+{
+	return FModuleManager::LoadModuleChecked< FParametricSurfaceModule >(PARAMETRICSURFACE_MODULE_NAME);
+}
+
+bool FParametricSurfaceModule::IsAvailable()
+{
+	return FModuleManager::Get().IsModuleLoaded(PARAMETRICSURFACE_MODULE_NAME);
+}
+
+IMPLEMENT_MODULE(FParametricSurfaceModule, ParametricSurface);
+
+#undef LOCTEXT_NAMESPACE // "ParametricSurfaceModule"
+
diff --git a/Engine/Plugins/Enterprise/DatasmithCADImporter/Source/ParametricSurface/Private/ParametricSurfaceTranslator.cpp b/Engine/Plugins/Enterprise/DatasmithCADImporter/Source/ParametricSurface/Private/ParametricSurfaceTranslator.cpp
new file mode 100644
--- /dev/null
+++ b/Engine/Plugins/Enterprise/DatasmithCADImporter/Source/ParametricSurface/Private/ParametricSurfaceTranslator.cpp
@@ -0,0 +1,35 @@
+// Copyright Epic Games, Inc. All Rights Reserved.
+
+#include "ParametricSurfaceTranslator.h"
+
+#include "DatasmithImportOptions.h"
+
+void FParametricSurfaceTranslator::GetSceneImportOptions(TArray<TStrongObjectPtr<UDatasmithOptionsBase>>& Options)
+{
+	FString Extension = GetSource().GetSourceFileExtension();
+	if (Extension == "cgr" || Extension == "3dxml")
+	{
+		return;
+	}
+
+	TStrongObjectPtr<UDatasmithCommonTessellationOptions> CommonTessellationOptionsPtr = Datasmith::MakeOptions<UDatasmithCommonTessellationOptions>();
+	check(CommonTessellationOptionsPtr.IsValid());
+	InitCommonTessellationOptions(CommonTessellationOptionsPtr->Options);
+
+	Options.Add(CommonTessellationOptionsPtr);
+}
+
+void FParametricSurfaceTranslator::SetSceneImportOptions(TArray<TStrongObjectPtr<UDatasmithOptionsBase>>& Options)
+{
+	for (const TStrongObjectPtr<UDatasmithOptionsBase>& OptionPtr : Options)
+	{
+		if (UDatasmithCommonTessellationOptions* TessellationOptionsObject = Cast<UDatasmithCommonTessellationOptions>(OptionPtr.Get()))
+		{
+			CommonTessellationOptions = TessellationOptionsObject->Options;
+		}
+	}
+}
+
+
+
+
diff --git a/Engine/Plugins/Enterprise/DatasmithCADImporter/Source/ParametricSurface/Public/ParametricSurfaceExtension.h b/Engine/Plugins/Enterprise/DatasmithCADImporter/Source/ParametricSurface/Public/ParametricSurfaceExtension.h
new file mode 100644
--- /dev/null
+++ b/Engine/Plugins/Enterprise/DatasmithCADImporter/Source/ParametricSurface/Public/ParametricSurfaceExtension.h
@@ -0,0 +1,76 @@
+// Copyright Epic Games, Inc. All Rights Reserved.
+
+#pragma once
+
+#include "CoreMinimal.h"
+
+#include "DatasmithAdditionalData.h"
+#include "DatasmithCustomAction.h"
+#include "DatasmithImportOptions.h"
+#include "DatasmithUtils.h"
+
+#include "ParametricSurfaceExtension.generated.h"
+
+
+USTRUCT(BlueprintType)
+struct PARAMETRICSURFACE_API FParametricSceneParameters
+{
+	GENERATED_BODY()
+
+	// value from FDatasmithUtils::EModelCoordSystem
+	UPROPERTY()
+	uint8 ModelCoordSys = (uint8)FDatasmithUtils::EModelCoordSystem::ZUp_LeftHanded;
+
+	UPROPERTY()
+	float MetricUnit = 0.01f;
+
+	UPROPERTY()
+	float ScaleFactor = 1.0f;
+};
+
+USTRUCT()
+struct PARAMETRICSURFACE_API FParametricMeshParameters
+{
+	GENERATED_BODY()
+
+	UPROPERTY()
+	bool bNeedSwapOrientation = false;
+
+	UPROPERTY()
+	bool bIsSymmetric = false;
+
+	UPROPERTY()
+	FVector SymmetricOrigin = FVector::ZeroVector;
+
+	UPROPERTY()
+	FVector SymmetricNormal = FVector::ZeroVector;
+};
+
+UCLASS(meta = (DisplayName = "Parametric Surface Data"))
+class PARAMETRICSURFACE_API UParametricSurfaceData : public UDatasmithAdditionalData
+{
+	GENERATED_BODY()
+
+public:
+
+	UPROPERTY()
+	FParametricSceneParameters SceneParameters;
+
+	UPROPERTY()
+	FParametricMeshParameters MeshParameters;
+
+	UPROPERTY(EditAnywhere, Category=NURBS)
+	FDatasmithTessellationOptions LastTessellationOptions;
+
+	virtual bool IsValid()
+	{
+		return false;
+	}
+
+	virtual bool Tessellate(UStaticMesh& StaticMesh, const FDatasmithRetessellationOptions& RetessellateOptions)
+	{
+		return false;
+	}
+
+};
+
diff --git a/Engine/Plugins/Enterprise/DatasmithCADImporter/Source/ParametricSurface/Public/ParametricSurfaceModule.h b/Engine/Plugins/Enterprise/DatasmithCADImporter/Source/ParametricSurface/Public/ParametricSurfaceModule.h
new file mode 100644
--- /dev/null
+++ b/Engine/Plugins/Enterprise/DatasmithCADImporter/Source/ParametricSurface/Public/ParametricSurfaceModule.h
@@ -0,0 +1,19 @@
+// Copyright Epic Games, Inc. All Rights Reserved.
+
+#pragma once
+
+#include "Modules/ModuleInterface.h"
+
+#define PARAMETRICSURFACE_MODULE_NAME TEXT("ParametricSurface")
+
+/**
+ * This module exposes additional features for assets containing Parametric data.
+ */
+class FParametricSurfaceModule : public IModuleInterface
+{
+public:
+	virtual void StartupModule() override;
+
+	static FParametricSurfaceModule& Get();
+	static bool IsAvailable();
+};
diff --git a/Engine/Plugins/Enterprise/DatasmithCADImporter/Source/ParametricSurface/Public/ParametricSurfaceTranslator.h b/Engine/Plugins/Enterprise/DatasmithCADImporter/Source/ParametricSurface/Public/ParametricSurfaceTranslator.h
new file mode 100644
--- /dev/null
+++ b/Engine/Plugins/Enterprise/DatasmithCADImporter/Source/ParametricSurface/Public/ParametricSurfaceTranslator.h
@@ -0,0 +1,40 @@
+// Copyright Epic Games, Inc. All Rights Reserved.
+
+#pragma once
+
+#include "CoreMinimal.h"
+
+#include "DatasmithImportOptions.h"
+#include "DatasmithTranslator.h"
+#include "UObject/ObjectMacros.h"
+
+class UDatasmithCommonTessellationOptions;
+
+
+class PARAMETRICSURFACE_API FParametricSurfaceTranslator : public IDatasmithTranslator
+{
+public:
+	virtual ~FParametricSurfaceTranslator() {}
+
+	// Begin IDatasmithTranslator overrides
+	virtual void GetSceneImportOptions(TArray<TStrongObjectPtr<UDatasmithOptionsBase>>& Options) override;
+	virtual void SetSceneImportOptions(TArray<TStrongObjectPtr<UDatasmithOptionsBase>>& Options) override;
+	// End IDatasmithTranslator overrides
+
+protected:
+	const FDatasmithTessellationOptions& GetCommonTessellationOptions()
+	{
+		return CommonTessellationOptions;
+	}
+
+	/** 
+	 * Call when the UDatasmithCommonTessellationOptions object is created. This is the unique opportunity for
+	 * child class to overwrite some values
+	 * @param TessellationOptions Reference on member of UDatasmithCommonTessellationOptions
+	 */
+	virtual void InitCommonTessellationOptions(FDatasmithTessellationOptions& TessellationOptions) {}
+
+private:
+	FDatasmithTessellationOptions CommonTessellationOptions;
+};
+
diff --git a/Engine/Plugins/Enterprise/DatasmithCADImporter/Source/ParametricSurfaceExtension/ParametricSurfaceExtension.Build.cs b/Engine/Plugins/Enterprise/DatasmithCADImporter/Source/ParametricSurfaceExtension/ParametricSurfaceExtension.Build.cs
--- a/Engine/Plugins/Enterprise/DatasmithCADImporter/Source/ParametricSurfaceExtension/ParametricSurfaceExtension.Build.cs
+++ b/Engine/Plugins/Enterprise/DatasmithCADImporter/Source/ParametricSurfaceExtension/ParametricSurfaceExtension.Build.cs
@@ -23,7 +23,7 @@ namespace UnrealBuildTool.Rules
 					"MeshDescription",
 					"DataprepCore",
 					"DatasmithContent",
-					"CoreTechSurface",
+					"ParametricSurface",
 					"DatasmithImporter",
 					"Engine",
 					"StaticMeshEditor",
diff --git a/Engine/Plugins/Enterprise/DatasmithCADImporter/Source/ParametricSurfaceExtension/Private/CoreTechBlueprintLibrary.cpp b/Engine/Plugins/Enterprise/DatasmithCADImporter/Source/ParametricSurfaceExtension/Private/CoreTechBlueprintLibrary.cpp
deleted file mode 100644
diff --git a/Engine/Plugins/Enterprise/DatasmithCADImporter/Source/ParametricSurfaceExtension/Private/CoreTechExtensionModule.cpp b/Engine/Plugins/Enterprise/DatasmithCADImporter/Source/ParametricSurfaceExtension/Private/CoreTechExtensionModule.cpp
deleted file mode 100644
diff --git a/Engine/Plugins/Enterprise/DatasmithCADImporter/Source/ParametricSurfaceExtension/Private/CoreTechRetessellateAction.cpp b/Engine/Plugins/Enterprise/DatasmithCADImporter/Source/ParametricSurfaceExtension/Private/CoreTechRetessellateAction.cpp
deleted file mode 100644
diff --git a/Engine/Plugins/Enterprise/DatasmithCADImporter/Source/ParametricSurfaceExtension/Private/DataprepTessellationOperation.cpp b/Engine/Plugins/Enterprise/DatasmithCADImporter/Source/ParametricSurfaceExtension/Private/DataprepTessellationOperation.cpp
--- a/Engine/Plugins/Enterprise/DatasmithCADImporter/Source/ParametricSurfaceExtension/Private/DataprepTessellationOperation.cpp
+++ b/Engine/Plugins/Enterprise/DatasmithCADImporter/Source/ParametricSurfaceExtension/Private/DataprepTessellationOperation.cpp
@@ -2,11 +2,11 @@
 
 #include "DataprepTessellationOperation.h"
 
-#include "CoreTechHelper.h"
-#include "CoreTechBlueprintLibrary.h"
+#include "ParametricSurfaceBlueprintLibrary.h"
 #include "DatasmithAdditionalData.h"
 #include "DatasmithAssetImportData.h"
 #include "DatasmithUtils.h"
+#include "MeshDescriptionHelper.h"
 
 #include "IDataprepProgressReporter.h"
 
@@ -84,7 +84,7 @@ void UDataprepTessellationOperation::OnExecution_Implementation(const FDataprepC
 			if( StaticMesh->IsMeshDescriptionValid( 0 ) )
 			{
 				FText OutReason;
-				if( UCoreTechBlueprintLibrary::RetessellateStaticMeshWithNotification( StaticMesh, TessellationSettings, false, OutReason ) )
+				if( UParametricSurfaceBlueprintLibrary::RetessellateStaticMeshWithNotification( StaticMesh, TessellationSettings, false, OutReason ) )
 				{
 					ModifiedStaticMeshes.Add( StaticMesh );
 					if (!OutReason.IsEmpty())
diff --git a/Engine/Plugins/Enterprise/DatasmithCADImporter/Source/ParametricSurfaceExtension/Private/DataprepTessellationOperation.h b/Engine/Plugins/Enterprise/DatasmithCADImporter/Source/ParametricSurfaceExtension/Private/DataprepTessellationOperation.h
--- a/Engine/Plugins/Enterprise/DatasmithCADImporter/Source/ParametricSurfaceExtension/Private/DataprepTessellationOperation.h
+++ b/Engine/Plugins/Enterprise/DatasmithCADImporter/Source/ParametricSurfaceExtension/Private/DataprepTessellationOperation.h
@@ -15,7 +15,7 @@ namespace CADLibrary
 	struct FImportParameters;
 }
 struct FMeshDescription;
-class UCoreTechParametricSurfaceData;
+class UParametricSurfaceData;
 class UDatasmithStaticMeshCADImportData;
 
 UCLASS(Experimental, Category = MeshOperation, Meta = (DisplayName="Datasmith Tessellation", ToolTip = "For each static mesh to process, retessellate the mesh if the object contains the required data") )
diff --git a/Engine/Plugins/Enterprise/DatasmithCADImporter/Source/ParametricSurfaceExtension/Private/ParametricRetessellateAction.cpp b/Engine/Plugins/Enterprise/DatasmithCADImporter/Source/ParametricSurfaceExtension/Private/ParametricRetessellateAction.cpp
new file mode 100644
--- /dev/null
+++ b/Engine/Plugins/Enterprise/DatasmithCADImporter/Source/ParametricSurfaceExtension/Private/ParametricRetessellateAction.cpp
@@ -0,0 +1,216 @@
+// Copyright Epic Games, Inc. All Rights Reserved.
+
+#include "ParametricRetessellateAction.h"
+
+#include "ParametricRetessellateAction_Impl.h"
+
+#include "DatasmithAdditionalData.h"
+#include "DatasmithStaticMeshImporter.h" // Call to BuildStaticMesh
+#include "DatasmithUtils.h"
+#include "DatasmithTranslator.h"
+#include "UI/DatasmithDisplayHelper.h"
+#include "MeshDescriptionHelper.h"
+
+#include "Algo/AnyOf.h"
+#include "AssetData.h"
+#include "Async/ParallelFor.h"
+#include "Engine/StaticMesh.h"
+#include "HAL/PlatformFileManager.h"
+#include "IStaticMeshEditor.h"
+#include "StaticMeshAttributes.h"
+#include "Misc/FileHelper.h"
+#include "Misc/ScopedSlowTask.h"
+#include "Toolkits/IToolkit.h"
+#include "Toolkits/ToolkitManager.h"
+#include "UObject/StrongObjectPtr.h"
+#include "Algo/Transform.h"
+
+#define LOCTEXT_NAMESPACE "ParametricRetessellateAction"
+
+const FText FParametricRetessellateAction_Impl::Label = LOCTEXT("RetessellateActionLabel", "Retessellate");
+const FText FParametricRetessellateAction_Impl::Tooltip = LOCTEXT("RetessellateActionTooltip", "Tessellate the original NURBS surfaces to re-generate the mesh geometry");
+
+const FText& UParametricRetessellateAction::GetLabel()
+{
+	return FParametricRetessellateAction_Impl::Label;
+}
+
+const FText& UParametricRetessellateAction::GetTooltip()
+{
+	return FParametricRetessellateAction_Impl::Tooltip;
+}
+
+bool UParametricRetessellateAction::CanApplyOnAssets(const TArray<FAssetData>& SelectedAssets)
+{
+	return FParametricRetessellateAction_Impl::CanApplyOnAssets(SelectedAssets);
+}
+
+void UParametricRetessellateAction::ApplyOnAssets(const TArray<FAssetData>& SelectedAssets)
+{
+	return FParametricRetessellateAction_Impl::ApplyOnAssets(SelectedAssets);
+}
+
+bool FParametricRetessellateAction_Impl::CanApplyOnAssets(const TArray<FAssetData>& SelectedAssets)
+{
+	return Algo::AnyOf(SelectedAssets, [](const FAssetData& Asset){ return Datasmith::GetAdditionalData<UParametricSurfaceData>(Asset) != nullptr; });
+}
+
+void FParametricRetessellateAction_Impl::ApplyOnAssets(const TArray<FAssetData>& SelectedAssets)
+{
+	TFunction<void(UStaticMesh*)> FinalizeChanges = [](UStaticMesh* StaticMesh) -> void
+	{
+		StaticMesh->PostEditChange();
+		StaticMesh->MarkPackageDirty();
+
+		// Refresh associated editor
+		TSharedPtr<IToolkit> EditingToolkit = FToolkitManager::Get().FindEditorForAsset(StaticMesh);
+		if (IStaticMeshEditor* StaticMeshEditorInUse = StaticCastSharedPtr<IStaticMeshEditor>(EditingToolkit).Get())
+		{
+			StaticMeshEditorInUse->RefreshTool();
+		}
+	};
+
+	TStrongObjectPtr<UParametricRetessellateActionOptions> RetessellateOptions = Datasmith::MakeOptions<UParametricRetessellateActionOptions>();
+
+	bool bSameOptionsForAll = false;
+	int32 NumAssetsToProcess = SelectedAssets.Num();
+	bool bAskForSameOption = NumAssetsToProcess > 1;
+
+	TArray<UStaticMesh*> TessellatedMeshes;
+	TessellatedMeshes.Reserve( SelectedAssets.Num() );
+
+	TUniquePtr<FScopedSlowTask> Progress;
+	int32 AssetIndex = -1;
+	for (const FAssetData& Asset : SelectedAssets)
+	{
+		AssetIndex++;
+		if (UParametricSurfaceData* ParametricSurfaceData = Datasmith::GetAdditionalData<UParametricSurfaceData>(Asset))
+		{
+			if (ParametricSurfaceData->IsValid())
+			{
+				if (UStaticMesh* StaticMesh = Cast<UStaticMesh>(Asset.GetAsset()))
+				{
+					if (!bSameOptionsForAll)
+					{
+						Datasmith::FDisplayParameters Parameters;
+						Parameters.bAskForSameOption = bAskForSameOption;
+						Parameters.WindowTitle = LOCTEXT("OptionWindow_WindowTitle", "Datasmith Retessellation Options");
+						Parameters.FileLabel = FText::Format(LOCTEXT("OptionWindow_AssetLabel", "Tessellate StaticMesh: {0}"), FText::FromString(StaticMesh->GetName()));
+						Parameters.FileTooltip = FText::FromString(StaticMesh->GetPathName());
+						Parameters.ProceedButtonLabel = LOCTEXT("OptionWindow_ProceedButtonLabel", "Tessellate");
+						Parameters.ProceedButtonTooltip = LOCTEXT("OptionWindow_ProceedButtonTooltip", "Retessellate this mesh based on included nurbs data");
+						Parameters.CancelButtonLabel = LOCTEXT("OptionWindow_CancelButtonLabel", "Cancel");
+						Parameters.CancelButtonTooltip = LOCTEXT("OptionWindow_CancelButtonTooltip", "Cancel the retessellation operation");
+
+						bAskForSameOption = false; // ask only the fist time
+						RetessellateOptions->Options = ParametricSurfaceData->LastTessellationOptions;
+						Datasmith::FDisplayResult Result = Datasmith::DisplayOptions(RetessellateOptions, Parameters);
+						if (!Result.bValidated)
+						{
+							return;
+						}
+						bSameOptionsForAll |= Result.bUseSameOption;
+					}
+					ParametricSurfaceData->LastTessellationOptions = RetessellateOptions->Options;
+
+					int32 RemainingAssetsToProcess = NumAssetsToProcess - AssetIndex;
+					if (bSameOptionsForAll && !Progress.IsValid() && RemainingAssetsToProcess > 1)
+					{
+						Progress = MakeUnique<FScopedSlowTask>(RemainingAssetsToProcess);
+						Progress->MakeDialog(true);
+					}
+					if (Progress)
+					{
+						if (Progress->ShouldCancel())
+						{
+							return;
+						}
+
+						FText Text = FText::Format(LOCTEXT("RetessellateAssetMessage", "Tessellate StaticMesh ({0}/{1}): {2}"),
+							AssetIndex,
+							NumAssetsToProcess,
+							FText::FromString(StaticMesh->GetName())
+						);
+						Progress->EnterProgressFrame(1, Text);
+					}
+
+					if( StaticMesh->GetMeshDescription(0) == nullptr)
+					{
+						StaticMesh->CreateMeshDescription( 0 );
+					}
+
+					if(StaticMesh->GetMeshDescription(0) != nullptr)
+					{
+						StaticMesh->Modify();
+						StaticMesh->PreEditChange( nullptr );
+
+						if(ParametricSurfaceData != nullptr && ParametricSurfaceData->Tessellate(*StaticMesh, RetessellateOptions->Options))
+						{
+							TessellatedMeshes.Add(StaticMesh);
+						}
+						else
+						{
+							FinalizeChanges(StaticMesh);
+						}
+					}
+				}
+			}
+		}
+	}
+
+	// Make sure lightmap settings are valid
+	if(TessellatedMeshes.Num() > 1)
+	{
+		ParallelFor(TessellatedMeshes.Num(), [&](int32 Index)
+		{
+			FDatasmithStaticMeshImporter::PreBuildStaticMesh(TessellatedMeshes[Index]); 
+		});
+	}
+	else if(TessellatedMeshes.Num() > 0)
+	{
+		FDatasmithStaticMeshImporter::PreBuildStaticMesh( TessellatedMeshes[0] );
+	}
+
+	FDatasmithStaticMeshImporter::BuildStaticMeshes( TessellatedMeshes );
+
+	for(UStaticMesh* StaticMesh : TessellatedMeshes)
+	{
+		FinalizeChanges(StaticMesh);
+	}
+}
+
+TSet<UStaticMesh*> GetReferencedStaticMeshes(const TArray<AActor*>& SelectedActors)
+{
+	TSet<UStaticMesh*> ReferencedStaticMeshes;
+
+	for (const AActor* Actor : SelectedActors)
+	{
+		if (Actor)
+		{
+			for (const auto& Component : Actor->GetComponents())
+			{
+				if (UStaticMeshComponent* SMC = Cast<UStaticMeshComponent>(Component))
+				{
+					ReferencedStaticMeshes.Add(SMC->GetStaticMesh());
+				}
+			}
+		}
+	}
+
+	return ReferencedStaticMeshes;
+}
+
+bool UParametricRetessellateAction::CanApplyOnActors(const TArray<AActor*>& SelectedActors)
+{
+	const TSet<UStaticMesh*> ReferencedStaticMeshes = GetReferencedStaticMeshes(SelectedActors);
+	return Algo::AnyOf(ReferencedStaticMeshes, [](const UStaticMesh* Mesh){ return Datasmith::GetAdditionalData<UParametricSurfaceData>(FAssetData(Mesh)); });
+}
+
+void UParametricRetessellateAction::ApplyOnActors(const TArray<AActor*>& SelectedActors)
+{
+	TArray<FAssetData> AssetData;
+	Algo::Transform(GetReferencedStaticMeshes(SelectedActors), AssetData, [](UStaticMesh* Mesh){ return FAssetData(Mesh);});
+	return ApplyOnAssets(AssetData);
+}
+
+#undef LOCTEXT_NAMESPACE
diff --git a/Engine/Plugins/Enterprise/DatasmithCADImporter/Source/ParametricSurfaceExtension/Private/ParametricRetessellateAction_Impl.h b/Engine/Plugins/Enterprise/DatasmithCADImporter/Source/ParametricSurfaceExtension/Private/ParametricRetessellateAction_Impl.h
new file mode 100644
--- /dev/null
+++ b/Engine/Plugins/Enterprise/DatasmithCADImporter/Source/ParametricSurfaceExtension/Private/ParametricRetessellateAction_Impl.h
@@ -0,0 +1,17 @@
+// Copyright Epic Games, Inc. All Rights Reserved.
+
+#pragma once
+
+#include "CoreMinimal.h"
+
+#include "UObject/Object.h"
+
+class FParametricRetessellateAction_Impl
+{
+public:
+	static const FText Label;
+	static const FText Tooltip;
+
+	static bool CanApplyOnAssets(const TArray<FAssetData>& SelectedAssets);
+	static void ApplyOnAssets(const TArray<FAssetData>& SelectedAssets);
+};
\ No newline at end of file
diff --git a/Engine/Plugins/Enterprise/DatasmithCADImporter/Source/ParametricSurfaceExtension/Private/ParametricSurfaceBlueprintLibrary.cpp b/Engine/Plugins/Enterprise/DatasmithCADImporter/Source/ParametricSurfaceExtension/Private/ParametricSurfaceBlueprintLibrary.cpp
new file mode 100644
--- /dev/null
+++ b/Engine/Plugins/Enterprise/DatasmithCADImporter/Source/ParametricSurfaceExtension/Private/ParametricSurfaceBlueprintLibrary.cpp
@@ -0,0 +1,106 @@
+// Copyright Epic Games, Inc. All Rights Reserved.
+
+#include "ParametricSurfaceBlueprintLibrary.h"
+
+#include "ParametricRetessellateAction.h"
+#include "DatasmithAdditionalData.h"
+#include "DatasmithStaticMeshImporter.h" // Call to BuildStaticMesh
+#include "DatasmithUtils.h"
+
+#include "AssetData.h"
+#include "Engine/StaticMesh.h"
+#include "IStaticMeshEditor.h"
+#include "StaticMeshAttributes.h"
+#include "Toolkits/ToolkitManager.h"
+
+
+#define LOCTEXT_NAMESPACE "ParametricRetessellateAction"
+
+
+bool UParametricSurfaceBlueprintLibrary::RetessellateStaticMesh(UStaticMesh* StaticMesh, const FDatasmithRetessellationOptions& TessellationSettings, FText& FailureReason)
+{
+	return RetessellateStaticMeshWithNotification(StaticMesh, TessellationSettings, true, FailureReason);
+}
+
+bool UParametricSurfaceBlueprintLibrary::RetessellateStaticMeshWithNotification(UStaticMesh* StaticMesh, const FDatasmithRetessellationOptions& TessellationSettings, bool bApplyChanges, FText& FailureReason)
+{
+	bool bTessellationOutcome = false;
+
+	int32 LODIndex = 0;
+
+	FAssetData AssetData( StaticMesh );
+	if (UParametricSurfaceData* ParametricSurfaceData = Datasmith::GetAdditionalData<UParametricSurfaceData>(AssetData))
+	{
+		// Make sure MeshDescription exists
+		FMeshDescription* DestinationMeshDescription = StaticMesh->GetMeshDescription( LODIndex );
+		if (DestinationMeshDescription == nullptr)
+		{
+			DestinationMeshDescription = StaticMesh->CreateMeshDescription(0);
+		}
+
+		if (DestinationMeshDescription)
+		{
+			if (bApplyChanges)
+			{
+				StaticMesh->Modify();
+				StaticMesh->PreEditChange( nullptr );
+			}
+
+			const int32 OldNumberOfUVChannels = FStaticMeshAttributes(*DestinationMeshDescription).GetVertexInstanceUVs().GetNumChannels();
+			if (ParametricSurfaceData->Tessellate(*StaticMesh, TessellationSettings))
+			{
+				const int32 NumberOfUVChannels = FStaticMeshAttributes(*DestinationMeshDescription).GetVertexInstanceUVs().GetNumChannels();
+				if (NumberOfUVChannels < OldNumberOfUVChannels)
+				{
+					FailureReason = FText::Format(NSLOCTEXT("BlueprintRetessellation", "UVChannelsDestroyed", "Tessellation operation on Static Mesh {0} is destroying all UV channels above channel #{1}"), FText::FromString(StaticMesh->GetName()), NumberOfUVChannels - 1);
+				}
+
+				// Post static mesh has changed
+				if (bApplyChanges)
+				{
+					FDatasmithStaticMeshImporter::PreBuildStaticMesh(StaticMesh); // handle uvs stuff
+					FDatasmithStaticMeshImporter::BuildStaticMesh(StaticMesh);
+
+					StaticMesh->PostEditChange();
+					StaticMesh->MarkPackageDirty();
+
+					// Refresh associated editor
+					TSharedPtr<IToolkit> EditingToolkit = FToolkitManager::Get().FindEditorForAsset(StaticMesh);
+					if (IStaticMeshEditor* StaticMeshEditorInUse = StaticCastSharedPtr<IStaticMeshEditor>(EditingToolkit).Get())
+					{
+						StaticMeshEditorInUse->RefreshTool();
+					}
+				}
+				// No posting required, just make sure the new tessellation is committed
+				else
+				{
+					UStaticMesh::FCommitMeshDescriptionParams Params;
+					Params.bMarkPackageDirty = false;
+					Params.bUseHashAsGuid = true;
+					StaticMesh->CommitMeshDescription( LODIndex, Params );
+				}
+
+				// Save last tessellation settings
+				ParametricSurfaceData->LastTessellationOptions = TessellationSettings;
+
+				bTessellationOutcome = true;
+			}
+			else
+			{
+				FailureReason = NSLOCTEXT("BlueprintRetessellation", "LoadFailed", "Cannot generate mesh from parametric surface data");
+			}
+		}
+		else
+		{
+			FailureReason = NSLOCTEXT("BlueprintRetessellation", "MeshDescriptionMissing", "Cannot create mesh description");
+		}
+	}
+	else
+	{
+		FailureReason = NSLOCTEXT("BlueprintRetessellation", "MissingData", "No tessellation data attached to the static mesh");
+	}
+
+	return bTessellationOutcome;
+}
+
+#undef LOCTEXT_NAMESPACE
\ No newline at end of file
diff --git a/Engine/Plugins/Enterprise/DatasmithCADImporter/Source/ParametricSurfaceExtension/Private/ParametricSurfaceExtensionModule.cpp b/Engine/Plugins/Enterprise/DatasmithCADImporter/Source/ParametricSurfaceExtension/Private/ParametricSurfaceExtensionModule.cpp
new file mode 100644
--- /dev/null
+++ b/Engine/Plugins/Enterprise/DatasmithCADImporter/Source/ParametricSurfaceExtension/Private/ParametricSurfaceExtensionModule.cpp
@@ -0,0 +1,102 @@
+// Copyright Epic Games, Inc. All Rights Reserved.
+
+#include "ParametricSurfaceExtensionModule.h"
+
+#include "ParametricRetessellateAction.h"
+#include "ParametricRetessellateAction_Impl.h"
+
+#include "Framework/MultiBox/MultiBoxBuilder.h"
+#include "IStaticMeshEditor.h"
+#include "Modules/ModuleInterface.h"
+#include "Modules/ModuleManager.h"
+#include "StaticMeshEditorModule.h"
+#include "UObject/StrongObjectPtr.h"
+
+
+#define LOCTEXT_NAMESPACE "CoreTechExtensionModule"
+
+
+/** UI extension that displays a Retessellate action in the StaticMeshEditor */
+namespace StaticMeshEditorExtenser
+{
+	bool CanExecute(UStaticMesh* Target)
+	{
+		TArray<FAssetData> AssetData;
+		AssetData.Emplace(Target);
+
+		FParametricRetessellateAction_Impl RetessellateAction;
+		return RetessellateAction.CanApplyOnAssets(AssetData);
+	}
+
+	void Execute(UStaticMesh* Target)
+	{
+		TArray<FAssetData> AssetData;
+		AssetData.Emplace(Target);
+
+		FParametricRetessellateAction_Impl RetessellateAction;
+		RetessellateAction.ApplyOnAssets(AssetData);
+	}
+
+	void ExtendAssetMenu(FMenuBuilder& MenuBuilder, UStaticMesh* Target)
+	{
+		MenuBuilder.AddMenuEntry(
+			FParametricRetessellateAction_Impl::Label,
+			FParametricRetessellateAction_Impl::Tooltip,
+			FSlateIcon(),
+			FUIAction(
+				FExecuteAction::CreateStatic(&StaticMeshEditorExtenser::Execute, Target),
+				FCanExecuteAction::CreateStatic(&StaticMeshEditorExtenser::CanExecute, Target)
+			)
+		);
+	}
+
+	TSharedRef<FExtender> CreateExtenderForObjects(const TSharedRef<FUICommandList> CommandList, const TArray<UObject*> Objects)
+	{
+		TSharedRef<FExtender> Extender = MakeShareable(new FExtender);
+
+		if (UStaticMesh* Target = Objects.Num() ? Cast<UStaticMesh>(Objects[0]) : nullptr)
+		{
+			Extender->AddMenuExtension(
+				"AssetEditorActions",
+				EExtensionHook::Position::After,
+				CommandList,
+				FMenuExtensionDelegate::CreateStatic(&StaticMeshEditorExtenser::ExtendAssetMenu, Target)
+			);
+		}
+
+		return Extender;
+	}
+
+	void Register()
+	{
+		if (!IsRunningCommandlet())
+		{
+			IStaticMeshEditorModule& StaticMeshEditorModule = FModuleManager::Get().LoadModuleChecked<IStaticMeshEditorModule>("StaticMeshEditor");
+			TArray<FAssetEditorExtender>& ExtenderDelegates = StaticMeshEditorModule.GetMenuExtensibilityManager()->GetExtenderDelegates();
+			ExtenderDelegates.Add(FAssetEditorExtender::CreateStatic(&StaticMeshEditorExtenser::CreateExtenderForObjects));
+		}
+	}
+};
+
+FParametricSurfaceExtensionModule& FParametricSurfaceExtensionModule::Get()
+{
+	return FModuleManager::LoadModuleChecked< FParametricSurfaceExtensionModule >(PARAMETRICSURFACEEXTENSION_MODULE_NAME);
+}
+
+bool FParametricSurfaceExtensionModule::IsAvailable()
+{
+	return FModuleManager::Get().IsModuleLoaded(PARAMETRICSURFACEEXTENSION_MODULE_NAME);
+}
+
+void FParametricSurfaceExtensionModule::StartupModule()
+{
+	if (!IsRunningCommandlet())
+	{
+		StaticMeshEditorExtenser::Register();
+	}
+}
+
+IMPLEMENT_MODULE(FParametricSurfaceExtensionModule, CoreTechExtension);
+
+#undef LOCTEXT_NAMESPACE // "DatasmithCoreTechExtensionModule"
+
diff --git a/Engine/Plugins/Enterprise/DatasmithCADImporter/Source/ParametricSurfaceExtension/Public/CoreTechBlueprintLibrary.h b/Engine/Plugins/Enterprise/DatasmithCADImporter/Source/ParametricSurfaceExtension/Public/CoreTechBlueprintLibrary.h
deleted file mode 100644
diff --git a/Engine/Plugins/Enterprise/DatasmithCADImporter/Source/ParametricSurfaceExtension/Public/CoreTechExtensionModule.h b/Engine/Plugins/Enterprise/DatasmithCADImporter/Source/ParametricSurfaceExtension/Public/CoreTechExtensionModule.h
deleted file mode 100644
diff --git a/Engine/Plugins/Enterprise/DatasmithCADImporter/Source/ParametricSurfaceExtension/Public/CoreTechRetessellateAction.h b/Engine/Plugins/Enterprise/DatasmithCADImporter/Source/ParametricSurfaceExtension/Public/CoreTechRetessellateAction.h
deleted file mode 100644
diff --git a/Engine/Plugins/Enterprise/DatasmithCADImporter/Source/ParametricSurfaceExtension/Public/ParametricRetessellateAction.h b/Engine/Plugins/Enterprise/DatasmithCADImporter/Source/ParametricSurfaceExtension/Public/ParametricRetessellateAction.h
new file mode 100644
--- /dev/null
+++ b/Engine/Plugins/Enterprise/DatasmithCADImporter/Source/ParametricSurfaceExtension/Public/ParametricRetessellateAction.h
@@ -0,0 +1,41 @@
+// Copyright Epic Games, Inc. All Rights Reserved.
+
+#pragma once
+
+#include "CoreMinimal.h"
+
+#include "ParametricSurfaceExtension.h"
+#include "DatasmithCustomAction.h"
+#include "DatasmithAdditionalData.h"
+#include "DatasmithImportOptions.h"
+
+#include "ParametricRetessellateAction.generated.h"
+
+UCLASS()
+class PARAMETRICSURFACEEXTENSION_API UParametricRetessellateAction : public UDatasmithCustomActionBase
+{
+	GENERATED_BODY()
+
+public:
+	virtual const FText& GetLabel() override;
+	virtual const FText& GetTooltip() override;
+
+	virtual bool CanApplyOnAssets(const TArray<FAssetData>& SelectedAssets) override;
+
+	virtual void ApplyOnAssets(const TArray<FAssetData>& SelectedAssets) override;
+
+	virtual bool CanApplyOnActors(const TArray<AActor*>& SelectedActors) override;
+
+	virtual void ApplyOnActors(const TArray<AActor*>& SelectedActors) override;
+};
+
+
+UCLASS(config = Editor, Transient)
+class PARAMETRICSURFACEEXTENSION_API UParametricRetessellateActionOptions : public UObject
+{
+	GENERATED_BODY()
+
+public:
+	UPROPERTY(config, EditAnywhere, BlueprintReadWrite, Category = "NotVisible", meta = (ShowOnlyInnerProperties))
+	FDatasmithRetessellationOptions Options;
+};
diff --git a/Engine/Plugins/Enterprise/DatasmithCADImporter/Source/ParametricSurfaceExtension/Public/ParametricSurfaceBlueprintLibrary.h b/Engine/Plugins/Enterprise/DatasmithCADImporter/Source/ParametricSurfaceExtension/Public/ParametricSurfaceBlueprintLibrary.h
new file mode 100644
--- /dev/null
+++ b/Engine/Plugins/Enterprise/DatasmithCADImporter/Source/ParametricSurfaceExtension/Public/ParametricSurfaceBlueprintLibrary.h
@@ -0,0 +1,44 @@
+// Copyright Epic Games, Inc. All Rights Reserved.
+
+#pragma once
+
+#include "CoreMinimal.h"
+
+#include "DatasmithImportOptions.h"
+
+#include "Kismet/BlueprintFunctionLibrary.h"
+#include "UObject/TextProperty.h"
+
+#include "ParametricSurfaceBlueprintLibrary.generated.h"
+
+class UStaticMesh;
+
+UCLASS(Blueprintable, BlueprintType, meta = (DisplayName = "CAD Surface Operations Library"))
+class PARAMETRICSURFACEEXTENSION_API UParametricSurfaceBlueprintLibrary : public UBlueprintFunctionLibrary
+{
+	GENERATED_BODY()
+
+public:
+	/**
+	 * Re-tessellate LOD 0 of a static mesh if it contains parametric surface data.
+	 * @param	StaticMesh				Static mesh to re-tessellate.
+	 * @param	TessellationSettings	Tessellation settings to use.
+	 * @param	FailureReason			Text describing the reason of failure.
+	 * @return True if successful, false otherwise
+	 */
+	UFUNCTION(BlueprintCallable, Category = "Datasmith | Surface Operations" )
+	static bool RetessellateStaticMesh(UStaticMesh* StaticMesh, const FDatasmithRetessellationOptions& TessellationSettings, FText& FailureReason);
+
+	/**
+	 * Re-tessellate LOD 0 of a static mesh if it contains parametric surface data.
+	 * This implementations allows to skip post edit operations (bApplyChanges=false),
+	 * the caller is then responsible to handle those operations.
+	 * @param	StaticMesh				Static mesh to re-tessellate.
+	 * @param	TessellationSettings	Tessellation settings to use.
+	 * @param	bApplyChanges			Indicates if change must be notified.
+	 * @param	FailureReason			Text describing the reason of failure, or a warning if the operation was successful.
+	 * @return True if successful, false otherwise
+	 */
+	UFUNCTION(Category = "Datasmith | Surface Operations" )
+	static bool RetessellateStaticMeshWithNotification(UStaticMesh* StaticMesh, const FDatasmithRetessellationOptions& TessellationSettings, bool bApplyChanges, FText& FailureReason);
+};
diff --git a/Engine/Plugins/Enterprise/DatasmithCADImporter/Source/ParametricSurfaceExtension/Public/ParametricSurfaceExtensionModule.h b/Engine/Plugins/Enterprise/DatasmithCADImporter/Source/ParametricSurfaceExtension/Public/ParametricSurfaceExtensionModule.h
new file mode 100644
--- /dev/null
+++ b/Engine/Plugins/Enterprise/DatasmithCADImporter/Source/ParametricSurfaceExtension/Public/ParametricSurfaceExtensionModule.h
@@ -0,0 +1,20 @@
+// Copyright Epic Games, Inc. All Rights Reserved.
+
+#pragma once
+
+#include "Modules/ModuleInterface.h"
+
+#define PARAMETRICSURFACEEXTENSION_MODULE_NAME TEXT("ParametricSurfaceExtension")
+
+/**
+ * This module exposes additional features for assets containing CoreTech data.
+ */
+class FParametricSurfaceExtensionModule : public IModuleInterface
+{
+public:
+    static FParametricSurfaceExtensionModule& Get();
+    static bool IsAvailable();
+
+private:
+	virtual void StartupModule() override;
+};
