Two problems:
- first one is that stitiching heal doesn't work. A call a stitching sew is necessary.
- the second is that UnloadModel reset unit setting in kernel IO. A call to ChangeUnit is mandatory before each load file or AddBrep. The modification of unit does not need anymore to shutdown CADKernel.

diff --git a/Engine/Plugins/Enterprise/DatasmithCADImporter/Source/CADInterfaces/Private/CADInterfacesModule.cpp b/Engine/Plugins/Enterprise/DatasmithCADImporter/Source/CADInterfaces/Private/CADInterfacesModule.cpp
--- a/Engine/Plugins/Enterprise/DatasmithCADImporter/Source/CADInterfaces/Private/CADInterfacesModule.cpp
+++ b/Engine/Plugins/Enterprise/DatasmithCADImporter/Source/CADInterfaces/Private/CADInterfacesModule.cpp
@@ -32,8 +32,7 @@ ECADInterfaceAvailability ICADInterfacesModule::GetAvailability()
 {
 	if (FModuleManager::Get().IsModuleLoaded(CADINTERFACES_MODULE_NAME))
 	{
-		double MetricUnit = 0.001;
-		if(CADLibrary::CTKIO_InitializeKernel(MetricUnit))
+		if(CADLibrary::CTKIO_InitializeKernel())
 		{
 			return ECADInterfaceAvailability::Available;
 		}
diff --git a/Engine/Plugins/Enterprise/DatasmithCADImporter/Source/CADInterfaces/Private/CTSession.cpp b/Engine/Plugins/Enterprise/DatasmithCADImporter/Source/CADInterfaces/Private/CTSession.cpp
--- a/Engine/Plugins/Enterprise/DatasmithCADImporter/Source/CADInterfaces/Private/CTSession.cpp
+++ b/Engine/Plugins/Enterprise/DatasmithCADImporter/Source/CADInterfaces/Private/CTSession.cpp
@@ -25,6 +25,12 @@ namespace CADLibrary
 		return CADLibrary::CTKIO_Repair(MainObjectId, ImportParams.StitchingTechnique, SewingToleranceFactor);
 	}
 
+	void FCTSession::SetSceneUnit(double InMetricUnit)
+	{
+		ImportParams.MetricUnit = InMetricUnit;
+		CTKIO_ChangeUnit(InMetricUnit);
+	}
+
 	void FCTSession::SetImportParameters(float ChordTolerance, float MaxEdgeLength, float NormalTolerance, CADLibrary::EStitchingTechnique StitchingTechnique, bool bScaleUVMap)
 	{
 		ImportParams.ChordTolerance = ChordTolerance;
diff --git a/Engine/Plugins/Enterprise/DatasmithCADImporter/Source/CADInterfaces/Private/CoreTechFileParser.cpp b/Engine/Plugins/Enterprise/DatasmithCADImporter/Source/CADInterfaces/Private/CoreTechFileParser.cpp
--- a/Engine/Plugins/Enterprise/DatasmithCADImporter/Source/CADInterfaces/Private/CoreTechFileParser.cpp
+++ b/Engine/Plugins/Enterprise/DatasmithCADImporter/Source/CADInterfaces/Private/CoreTechFileParser.cpp
@@ -29,7 +29,7 @@ namespace CADLibrary
 		: CachePath(InCachePath)
 		, ImportParameters(ImportParams)
 	{
-		CTKIO_InitializeKernel(ImportParameters.MetricUnit, *EnginePluginsPath);
+		CTKIO_InitializeKernel(*EnginePluginsPath);
 	}
 
 	bool FCoreTechFileParser::FindFile(FFileDescription& File)
diff --git a/Engine/Plugins/Enterprise/DatasmithCADImporter/Source/CADInterfaces/Private/CoreTechFileReader.cpp b/Engine/Plugins/Enterprise/DatasmithCADImporter/Source/CADInterfaces/Private/CoreTechFileReader.cpp
--- a/Engine/Plugins/Enterprise/DatasmithCADImporter/Source/CADInterfaces/Private/CoreTechFileReader.cpp
+++ b/Engine/Plugins/Enterprise/DatasmithCADImporter/Source/CADInterfaces/Private/CoreTechFileReader.cpp
@@ -295,6 +295,7 @@ namespace CADLibrary
 			CTReImportOption &= ~CT_LOAD_FLAGS_READ_ASM_STRUCT_ONLY;  // BUG CT -> Ticket 11685
 			Result = CT_KERNEL_IO::LoadFile(*FileDescription.Path, MainId, CTReImportOption, 0, *LoadOption);
 		}
+		CTKIO_ChangeUnit(Context.ImportParameters.MetricUnit);
 
 		// the file is loaded but it's empty, so no data is generate
 		if (Result == IO_ERROR_EMPTY_ASSEMBLY)
diff --git a/Engine/Plugins/Enterprise/DatasmithCADImporter/Source/CADInterfaces/Private/CoreTechInterfaceImpl.cpp b/Engine/Plugins/Enterprise/DatasmithCADImporter/Source/CADInterfaces/Private/CoreTechInterfaceImpl.cpp
--- a/Engine/Plugins/Enterprise/DatasmithCADImporter/Source/CADInterfaces/Private/CoreTechInterfaceImpl.cpp
+++ b/Engine/Plugins/Enterprise/DatasmithCADImporter/Source/CADInterfaces/Private/CoreTechInterfaceImpl.cpp
@@ -30,8 +30,13 @@ const char* CoreTechLicenseKey =
 
 namespace CADLibrary
 {
-	bool FCoreTechInterfaceImpl::InitializeKernel(double Unit, const TCHAR* EnginePluginsPath)
+	bool FCoreTechInterfaceImpl::InitializeKernel(const TCHAR* EnginePluginsPath)
 	{
+		if (bIsInitialize)
+		{
+			return true;
+		}
+
 		FString KernelIOPath;
 		if (FCString::Strlen(EnginePluginsPath))
 		{
@@ -43,23 +48,14 @@ namespace CADLibrary
 			}
 		}
 
-		if (FMath::IsNearlyEqual(Unit, KernelUnit))
-		{
-			// Kernel_IO is already initialized with requested Unit value
-			return true;
-		}
-
-		// Kernel_IO is already initialized, so it is stopped to be able to restart it with the good unit value
-		ShutdownKernel();
-
 		CT_STR AppName = CoreTechLicenseKey;
-		CT_IO_ERROR Status = CT_KERNEL_IO::InitializeKernel(AppName, Unit, 0.00001 / Unit, *KernelIOPath);
+		CT_IO_ERROR Status = CT_KERNEL_IO::InitializeKernel(AppName, 0.001, 0.01, *KernelIOPath);
 
 		if (Status == IO_OK || Status == IO_ERROR_ALREADY_INITIALIZED)
 		{
 			if (Status == IO_OK)
 			{
-				KernelUnit = Unit;
+				bIsInitialize = true;
 			}
 			return true;
 		}
@@ -81,7 +77,7 @@ namespace CADLibrary
 
 	bool FCoreTechInterfaceImpl::ShutdownKernel()
 	{
-		KernelUnit = -1;
+		bIsInitialize = false;
 		return CT_KERNEL_IO::ShutdownKernel() == IO_OK; // ignorable. Just in case CT was not previously stopped
 	}
 
@@ -90,6 +86,15 @@ namespace CADLibrary
 		return CT_KERNEL_IO::UnloadModel() == IO_OK;
 	}
 
+	bool FCoreTechInterfaceImpl::ChangeUnit(double SceneUnit)
+	{
+		if (CT_KERNEL_IO::ChangeUnit(SceneUnit) != IO_OK)
+		{
+			return false;
+		}
+		return CT_KERNEL_IO::ChangeTolerance(0.00001 / SceneUnit) == IO_OK;
+	}
+	
 	bool FCoreTechInterfaceImpl::CreateModel(uint64& OutMainObjectId)
 	{
 		OutMainObjectId = 0;
diff --git a/Engine/Plugins/Enterprise/DatasmithCADImporter/Source/CADInterfaces/Private/CoreTechInterfaceImpl.h b/Engine/Plugins/Enterprise/DatasmithCADImporter/Source/CADInterfaces/Private/CoreTechInterfaceImpl.h
--- a/Engine/Plugins/Enterprise/DatasmithCADImporter/Source/CADInterfaces/Private/CoreTechInterfaceImpl.h
+++ b/Engine/Plugins/Enterprise/DatasmithCADImporter/Source/CADInterfaces/Private/CoreTechInterfaceImpl.h
@@ -34,15 +34,16 @@ namespace CADLibrary
 
 		virtual void SetExternal(bool bSetExternal) override { bIsExternal = bSetExternal; }
 
-		virtual bool InitializeKernel(double MetricUnit, const TCHAR* = TEXT("")) override;
+		virtual bool InitializeKernel(const TCHAR* = TEXT("")) override;
 
 		virtual bool ShutdownKernel() override;
 
 		virtual bool UnloadModel() override;
 
+		virtual bool ChangeUnit(double SceneUnit) override;
+		
 		virtual bool CreateModel(uint64& OutMainObjectId) override;
-
-
+		
 		virtual bool ChangeTesselationParameters(double MaxSag, double MaxLength, double MaxAngle) override;
 
 		virtual bool LoadModel
@@ -120,7 +121,7 @@ namespace CADLibrary
 		void RepairInternal(CT_OBJECT_ID MainId, bool bConnectOpenBody, CT_DOUBLE SewingToleranceFactor);
 
 		bool bIsExternal = false;
-		double KernelUnit = -1;
+		bool bIsInitialize = false;
 	};
 }
 
diff --git a/Engine/Plugins/Enterprise/DatasmithCADImporter/Source/CADInterfaces/Private/CoreTechTypes.cpp b/Engine/Plugins/Enterprise/DatasmithCADImporter/Source/CADInterfaces/Private/CoreTechTypes.cpp
--- a/Engine/Plugins/Enterprise/DatasmithCADImporter/Source/CADInterfaces/Private/CoreTechTypes.cpp
+++ b/Engine/Plugins/Enterprise/DatasmithCADImporter/Source/CADInterfaces/Private/CoreTechTypes.cpp
@@ -29,7 +29,7 @@ namespace CADLibrary
 #endif
 	}
 
-	bool CTKIO_InitializeKernel(double Unit, const TCHAR* InEnginePluginsPath)
+	bool CTKIO_InitializeKernel(const TCHAR* InEnginePluginsPath)
 	{
 		if (!CoreTechInterface.IsValid())
 		{
@@ -43,7 +43,7 @@ namespace CADLibrary
 			EnginePluginsPath = FPaths::EnginePluginsDir();
 		}
 		
-		return CoreTechInterface->InitializeKernel(Unit, *EnginePluginsPath);
+		return CoreTechInterface->InitializeKernel(*EnginePluginsPath);
 	}
 
 	bool CTKIO_ShutdownKernel()
@@ -56,6 +56,11 @@ namespace CADLibrary
 		return CoreTechInterface.IsValid() ? CoreTechInterface->UnloadModel() : false;
 	}
 
+	bool CTKIO_ChangeUnit(double SceneUnit)
+	{
+		return CoreTechInterface.IsValid() ? CoreTechInterface->ChangeUnit(SceneUnit) : false;
+	}
+
 	bool CTKIO_CreateModel(uint64& OutMainObjectId)
 	{
 		OutMainObjectId = 0;
@@ -205,14 +210,15 @@ namespace CADLibrary
 
 	const TCHAR* FCoreTechSessionBase::Owner = nullptr;
 
-	FCoreTechSessionBase::FCoreTechSessionBase(const TCHAR* InOwner, double Unit)
+	FCoreTechSessionBase::FCoreTechSessionBase(const TCHAR* InOwner)
 	{
 		// Lib init
 		ensureAlways(InOwner != nullptr); // arg must be valid
 		ensureAlways(Owner == nullptr); // static Owner must be unset
 
-		CTKIO_ShutdownKernel(); // ignorable. Just in case CT was not previously stopped
-		Owner = CTKIO_InitializeKernel(Unit) ? InOwner : nullptr;
+		Owner = CTKIO_InitializeKernel() ? InOwner : nullptr;
+		// ignorable. Just in case CT was not previously emptied 
+		CTKIO_UnloadModel();
 
 		if (Owner)
 		{
@@ -226,7 +232,6 @@ namespace CADLibrary
 		if (Owner != nullptr)
 		{
 			CTKIO_UnloadModel();
-			CTKIO_ShutdownKernel();
 			Owner = nullptr;
 		}
 	}
diff --git a/Engine/Plugins/Enterprise/DatasmithCADImporter/Source/CADInterfaces/Public/CTSession.h b/Engine/Plugins/Enterprise/DatasmithCADImporter/Source/CADInterfaces/Public/CTSession.h
--- a/Engine/Plugins/Enterprise/DatasmithCADImporter/Source/CADInterfaces/Public/CTSession.h
+++ b/Engine/Plugins/Enterprise/DatasmithCADImporter/Source/CADInterfaces/Public/CTSession.h
@@ -20,14 +20,10 @@ public:
 	 * Make sure CT is initialized, and a main object is ready.
 	 * Handle input file unit and an output unit
 	 * @param InOwner:        text that describe the owner of the session (helps to fix initialization issues)
-	 * @param FileMetricUnit: number of meters per file unit.
-	 * eg. For a file in inches, arg should be 0.0254
 	 */
-	FCTSession(const TCHAR* InOwner, double InFileMetricUnit, double InScaleFactor)
-		: FCoreTechSessionBase(InOwner, InFileMetricUnit)
+	FCTSession(const TCHAR* InOwner)
+		: FCoreTechSessionBase(InOwner)
 	{
-		ImportParams.ScaleFactor = InScaleFactor;
-		ImportParams.MetricUnit = InFileMetricUnit;
 	}
 
 	void ClearData();
@@ -50,6 +46,13 @@ public:
 		ImportParams.ScaleFactor = InScaleFactor;
 	}
 
+	/**
+	 * Handle input file unit
+	 * @param FileMetricUnit: number of meters per file unit.
+	 * eg. For a file in inches, arg should be 0.0254
+	 */
+	void SetSceneUnit(double InMetricUnit);
+	
 	/**
 	 * Set Import parameters,
 	 * Tack care to set scale factor before because import parameters will be scale according to scale factor
diff --git a/Engine/Plugins/Enterprise/DatasmithCADImporter/Source/CADInterfaces/Public/CoreTechTypes.h b/Engine/Plugins/Enterprise/DatasmithCADImporter/Source/CADInterfaces/Public/CoreTechTypes.h
--- a/Engine/Plugins/Enterprise/DatasmithCADImporter/Source/CADInterfaces/Public/CoreTechTypes.h
+++ b/Engine/Plugins/Enterprise/DatasmithCADImporter/Source/CADInterfaces/Public/CoreTechTypes.h
@@ -91,10 +91,15 @@ namespace CADLibrary
 		virtual bool IsExternal() = 0;
 		virtual void SetExternal(bool Value) = 0;
 
-		/*
-		*
-		*/
-		virtual bool InitializeKernel(double MetricUnit, const TCHAR* = TEXT("")) = 0;
+		virtual bool InitializeKernel(const TCHAR* = TEXT("")) = 0;
+
+		/**
+		 * Change Kernel_IO unit
+		 * As to be call after CTKIO_UnloadModel
+		 * This method set also the Tolerance to 0.00001 m whether 0.01 mm
+		 * @param MetricUnit: Length unit express in meter i.e. 0.001 = mm
+		 */
+		virtual bool ChangeUnit(double SceneUnit) = 0;
 
 		virtual bool ShutdownKernel() = 0;
 
@@ -177,12 +182,16 @@ namespace CADLibrary
 		TArray<uint32> TriangleMaterials;
 	};
 
+	CADINTERFACES_API bool CTKIO_InitializeKernel(const TCHAR* = TEXT(""));
+
 	/**
-	 * Mandatory: Kernel_IO has to be initialize with the final MetricUnit, set function doesn't work
-	 * This method set Tolerance to 0.00001 m whether 0.01 mm
-	 * @param MetricUnit: Length unit express in meter
+	 * Change Kernel_IO unit
+	 * As to be call after CTKIO_UnloadModel
+	 * This method set also the Tolerance to 0.00001 m whether 0.01 mm
+	 * @param MetricUnit: Length unit express in meter i.e. 0.001 = mm
 	 */
-	CADINTERFACES_API bool CTKIO_InitializeKernel(double MetricUnit, const TCHAR* = TEXT(""));
+	CADINTERFACES_API bool CTKIO_ChangeUnit(double SceneUnit);
+	
 	CADINTERFACES_API bool CTKIO_ShutdownKernel();
 	CADINTERFACES_API bool CTKIO_UnloadModel();
 	CADINTERFACES_API bool CTKIO_CreateModel(uint64& OutMainObjectId);
@@ -247,11 +256,8 @@ namespace CADLibrary
 	public:
 		/**
 		 * Make sure CT is initialized, and a main object is ready.
-		 * Handle input file unit and an output unit
-		 * @param FileMetricUnit number of meters per file unit.
-		 * eg. For a file in inches, arg should be 0.0254
 		 */
-		FCoreTechSessionBase(const TCHAR* Owner, double Unit);
+		FCoreTechSessionBase(const TCHAR* Owner);
 		bool IsSessionValid() { return Owner != nullptr && MainObjectId != 0; }
 		virtual ~FCoreTechSessionBase();
 
diff --git a/Engine/Plugins/Enterprise/DatasmithCADImporter/Source/CoreTechSurface/Private/CoreTechSurfaceHelper.cpp b/Engine/Plugins/Enterprise/DatasmithCADImporter/Source/CoreTechSurface/Private/CoreTechSurfaceHelper.cpp
--- a/Engine/Plugins/Enterprise/DatasmithCADImporter/Source/CoreTechSurface/Private/CoreTechSurfaceHelper.cpp
+++ b/Engine/Plugins/Enterprise/DatasmithCADImporter/Source/CoreTechSurface/Private/CoreTechSurfaceHelper.cpp
@@ -41,7 +41,7 @@ namespace CoreTechSurface
 
 	bool LoadFile(const FString& FileName, FMeshDescription& MeshDescription, const CADLibrary::FImportParameters& ImportParameters, CADLibrary::FMeshParameters& MeshParameters)
 	{
-		CADLibrary::FCoreTechSessionBase Session(TEXT("CoreTechMeshLoader::LoadFile"), ImportParameters.MetricUnit);
+		CADLibrary::FCoreTechSessionBase Session(TEXT("CoreTechMeshLoader::LoadFile"));
 		if (!Session.IsSessionValid())
 		{
 			return false;
@@ -53,10 +53,11 @@ namespace CoreTechSurface
 			// Something wrong happened during the load, abort
 			return false;
 		}
+		CADLibrary::CTKIO_ChangeUnit(ImportParameters.MetricUnit);
 
 		if (ImportParameters.StitchingTechnique != CADLibrary::EStitchingTechnique::StitchingNone)
 		{
-			CADLibrary::CTKIO_Repair(MainObjectID, CADLibrary::EStitchingTechnique::StitchingHeal);
+			CADLibrary::CTKIO_Repair(MainObjectID, CADLibrary::EStitchingTechnique::StitchingSew);
 		}
 
 		return Tessellate(MainObjectID, ImportParameters, MeshDescription, MeshParameters);
diff --git a/Engine/Plugins/Enterprise/DatasmithCADImporter/Source/DatasmithOpenNurbsTranslator/Private/DatasmithOpenNurbsTranslator.cpp b/Engine/Plugins/Enterprise/DatasmithCADImporter/Source/DatasmithOpenNurbsTranslator/Private/DatasmithOpenNurbsTranslator.cpp
--- a/Engine/Plugins/Enterprise/DatasmithCADImporter/Source/DatasmithOpenNurbsTranslator/Private/DatasmithOpenNurbsTranslator.cpp
+++ b/Engine/Plugins/Enterprise/DatasmithCADImporter/Source/DatasmithOpenNurbsTranslator/Private/DatasmithOpenNurbsTranslator.cpp
@@ -649,7 +649,7 @@ public:
 			TranslationCache = MakeShared<FTranslationCache>();
 		}
 
-		LocalSession = FRhinoCoretechWrapper::GetSharedSession(MetricUnit, ScalingFactor);
+		LocalSession = FRhinoCoretechWrapper::GetSharedSession();
 	}
 
 	~FOpenNurbsTranslatorImpl()
@@ -2327,6 +2327,7 @@ bool FOpenNurbsTranslatorImpl::Read(ON_BinaryFile& Archive, TSharedRef<IDatasmit
 
 	// 	ScalingFactor is defined according to input Rhino file unit. CT don't modify CAD input according to the set file unit.
 	ScalingFactor = 100 / Settings.m_ModelUnitsAndTolerances.Scale(ON::LengthUnitSystem::Meters);
+	MetricUnit = 1 / Settings.m_ModelUnitsAndTolerances.Scale(ON::LengthUnitSystem::Meters);
 	LocalSession->SetScaleFactor(ScalingFactor);
 
 	// Step 4: REQUIRED - Read bitmap table (it can be empty)
@@ -3184,6 +3185,7 @@ bool FOpenNurbsTranslatorImpl::TranslateBRep(ON_Brep* Brep, const ON_3dmObjectAt
 		LocalSession->GetImportParameters().ModelCoordSys = FDatasmithUtils::EModelCoordSystem::ZUp_RightHanded_FBXLegacy;
 
 		LocalSession->ClearData();
+		LocalSession->SetSceneUnit(GetMetricUnit());
 
 		LocalSession->AddBRep(*Brep, Offset);
 
diff --git a/Engine/Plugins/Enterprise/DatasmithCADImporter/Source/DatasmithOpenNurbsTranslator/Private/RhinoCoretechWrapper.cpp b/Engine/Plugins/Enterprise/DatasmithCADImporter/Source/DatasmithOpenNurbsTranslator/Private/RhinoCoretechWrapper.cpp
--- a/Engine/Plugins/Enterprise/DatasmithCADImporter/Source/DatasmithOpenNurbsTranslator/Private/RhinoCoretechWrapper.cpp
+++ b/Engine/Plugins/Enterprise/DatasmithCADImporter/Source/DatasmithOpenNurbsTranslator/Private/RhinoCoretechWrapper.cpp
@@ -451,12 +451,12 @@ bool FRhinoCoretechWrapper::AddBRep(ON_Brep& Brep, const ON_3dVector& Offset)
 
 
 
-TSharedPtr<FRhinoCoretechWrapper> FRhinoCoretechWrapper::GetSharedSession(double SceneUnit, double ScaleFactor)
+TSharedPtr<FRhinoCoretechWrapper> FRhinoCoretechWrapper::GetSharedSession()
 {
 	TSharedPtr<FRhinoCoretechWrapper> Session = SharedSession.Pin();
 	if (!Session.IsValid())
 	{
-		Session = MakeShared<FRhinoCoretechWrapper>(TEXT("Rh2CTSharedSession"), SceneUnit, ScaleFactor);
+		Session = MakeShared<FRhinoCoretechWrapper>(TEXT("Rh2CTSharedSession"));
 		SharedSession = Session;
 	}
 	return Session;
diff --git a/Engine/Plugins/Enterprise/DatasmithCADImporter/Source/DatasmithOpenNurbsTranslator/Private/RhinoCoretechWrapper.h b/Engine/Plugins/Enterprise/DatasmithCADImporter/Source/DatasmithOpenNurbsTranslator/Private/RhinoCoretechWrapper.h
--- a/Engine/Plugins/Enterprise/DatasmithCADImporter/Source/DatasmithOpenNurbsTranslator/Private/RhinoCoretechWrapper.h
+++ b/Engine/Plugins/Enterprise/DatasmithCADImporter/Source/DatasmithOpenNurbsTranslator/Private/RhinoCoretechWrapper.h
@@ -46,8 +46,8 @@ public:
 	 * @param ScaleFactor scale factor to apply to the mesh to be in UnrealEditor unit (cm).
 	 * eg. For a file in inches, arg should be 0.0254
 	 */
-	FRhinoCoretechWrapper(const TCHAR* InOwner, double FileMetricUnit, double ScaleFactor)
-		: FCTSession(InOwner, FileMetricUnit, ScaleFactor)
+	FRhinoCoretechWrapper(const TCHAR* InOwner)
+		: FCTSession(InOwner)
 	{
 	}
 
@@ -59,7 +59,7 @@ public:
 	 */
 	bool AddBRep(ON_Brep& Brep, const ON_3dVector& Offset);
 	
-	static TSharedPtr<FRhinoCoretechWrapper> GetSharedSession(double SceneUnit, double ScaleFactor);
+	static TSharedPtr<FRhinoCoretechWrapper> GetSharedSession();
 
 	bool Tessellate(FMeshDescription& Mesh, CADLibrary::FMeshParameters& MeshParameters);
 
diff --git a/Engine/Plugins/Enterprise/DatasmithCADImporter/Source/DatasmithWireTranslator/Private/AliasCoretechWrapper.h b/Engine/Plugins/Enterprise/DatasmithCADImporter/Source/DatasmithWireTranslator/Private/AliasCoretechWrapper.h
--- a/Engine/Plugins/Enterprise/DatasmithCADImporter/Source/DatasmithWireTranslator/Private/AliasCoretechWrapper.h
+++ b/Engine/Plugins/Enterprise/DatasmithCADImporter/Source/DatasmithWireTranslator/Private/AliasCoretechWrapper.h
@@ -38,9 +38,11 @@ public:
 	 * eg. For a file in inches, arg should be 0.0254
 	 */
 	FAliasCoretechWrapper(const TCHAR* InOwner)
-		: CADLibrary::FCTSession(InOwner, 0.01, 1)
-		// Unit for CoreTech session is set to cm, 0.01, because Wire's unit is cm. Consequently, Scale factor is set to 1.
+		: CADLibrary::FCTSession(InOwner)
 	{
+		// Unit for CoreTech session is set to cm, 0.01, because Wire's unit is cm. Consequently, Scale factor is set to 1.
+		ImportParams.MetricUnit = 0.01;
+		ImportParams.ScaleFactor = 1;
 	}
 
 	bool AddBRep(TArray<AlDagNode*>& DagNodeSet, EAliasObjectReference ObjectReference);
diff --git a/Engine/Plugins/Enterprise/DatasmithCADImporter/Source/DatasmithWireTranslator/Private/DatasmithWireTranslator.cpp b/Engine/Plugins/Enterprise/DatasmithCADImporter/Source/DatasmithWireTranslator/Private/DatasmithWireTranslator.cpp
--- a/Engine/Plugins/Enterprise/DatasmithCADImporter/Source/DatasmithWireTranslator/Private/DatasmithWireTranslator.cpp
+++ b/Engine/Plugins/Enterprise/DatasmithCADImporter/Source/DatasmithWireTranslator/Private/DatasmithWireTranslator.cpp
@@ -1775,6 +1775,10 @@ TOptional<FMeshDescription> FWireTranslatorImpl::MeshDagNodeWithExternalMesher(A
 {
 	LocalSession->ClearData();
 
+
+	// Wire unit is cm
+	LocalSession->SetSceneUnit(0.01);
+	
 	FString Filename = DagNode.name();
 
 	EAliasObjectReference ObjectReference = EAliasObjectReference::LocalReference;
@@ -1809,6 +1813,9 @@ TOptional<FMeshDescription> FWireTranslatorImpl::MeshDagNodeWithExternalMesher(T
 {
 	LocalSession->ClearData();
 
+	// Wire unit is cm
+	LocalSession->SetSceneUnit(0.01);
+	
 	EAliasObjectReference ObjectReference = EAliasObjectReference::LocalReference;
 	if (MeshParameters.bIsSymmetric)
 	{
