A process has been added after the import to avoid faraway points. This process was based on the aabb of the topological vertices of the body. In this case, the aabb of the vertices is very small compare to the body. This process make more trouble than it fixes problems. So it has been removed.

diff --git a/Engine/Plugins/Enterprise/DatasmithCADImporter/Source/CADLibrary/Private/MeshDescriptionHelper.cpp b/Engine/Plugins/Enterprise/DatasmithCADImporter/Source/CADLibrary/Private/MeshDescriptionHelper.cpp
--- a/Engine/Plugins/Enterprise/DatasmithCADImporter/Source/CADLibrary/Private/MeshDescriptionHelper.cpp
+++ b/Engine/Plugins/Enterprise/DatasmithCADImporter/Source/CADLibrary/Private/MeshDescriptionHelper.cpp
@@ -67,13 +67,6 @@ namespace CADLibrary
 		int32 TriangleCount = Body.TriangleCount;
 		TArray<FTessellationData>& FaceTessellationSet = Body.Faces;
 
-		// Add offset to the bounding box to avoid to remove good vertex
-		FVector Size = Body.BBox.GetSize();
-		FBox BBox = Body.BBox.ExpandBy(Size.Size());
-		BBox.IsValid = Body.BBox.IsValid;
-		BBox.Min *= ImportParams.ScaleFactor;
-		BBox.Max *= ImportParams.ScaleFactor;
-
 		TVertexAttributesRef<FVector3f> VertexPositions = MeshDescription.GetVertexPositions();
 
 		// Create a list of vertex Z/index pairs
@@ -115,13 +108,7 @@ namespace CADLibrary
 			int32 Index_i = VertexDataSet[i].Index;
 			IndexOfCoincidentNode[Index_i] = Index_i;
 
-			// Check if mesh vertex is not outside body bbox
 			const FVector& PositionA = VertexDataSet[i].Coordinates;
-			if (BBox.IsValid && !BBox.IsInside(PositionA))
-			{
-				VertexDataSet[i].Index = INDEX_NONE;
-				continue;
-			}
 
 			// only need to search forward, since we add pairs both ways
 			for (int32 j = i + 1; j < VertexDataSet.Num(); j++)
