Improve Import CT

diff --git a/Engine/Plugins/Enterprise/DatasmithCADImporter/Source/CADInterfaces/Private/CoreTechFileReader.cpp b/Engine/Plugins/Enterprise/DatasmithCADImporter/Source/CADInterfaces/Private/CoreTechFileReader.cpp
--- a/Engine/Plugins/Enterprise/DatasmithCADImporter/Source/CADInterfaces/Private/CoreTechFileReader.cpp
+++ b/Engine/Plugins/Enterprise/DatasmithCADImporter/Source/CADInterfaces/Private/CoreTechFileReader.cpp
@@ -287,15 +287,16 @@ namespace CADLibrary
 			}
 		}
 
+		CTKIO_ChangeUnit(Context.ImportParameters.MetricUnit);
 		Result = CT_KERNEL_IO::LoadFile(*FileDescription.Path, MainId, CTImportOption, 0, *LoadOption);
 		if (Result == IO_ERROR_EMPTY_ASSEMBLY)
 		{
 			CT_KERNEL_IO::UnloadModel();
 			CT_FLAGS CTReImportOption = CTImportOption | CT_LOAD_FLAGS_LOAD_EXTERNAL_REF;
 			CTReImportOption &= ~CT_LOAD_FLAGS_READ_ASM_STRUCT_ONLY;  // BUG CT -> Ticket 11685
+			CTKIO_ChangeUnit(Context.ImportParameters.MetricUnit);
 			Result = CT_KERNEL_IO::LoadFile(*FileDescription.Path, MainId, CTReImportOption, 0, *LoadOption);
 		}
-		CTKIO_ChangeUnit(Context.ImportParameters.MetricUnit);
 
 		// the file is loaded but it's empty, so no data is generate
 		if (Result == IO_ERROR_EMPTY_ASSEMBLY)
