Pluggin of CADKernel
By setting r.CADTranslator.EnableKernelIOTessellation CVar to 0 (disabling KernelIO Tessellation) CADKernel is activated.
This first version allow to tessellate CAD file with CADKernel. Retessellate and Sew are not availlable.

diff --git a/Engine/Plugins/Enterprise/DatasmithCADImporter/Source/CADInterfaces/Private/CoreTechBridge.cpp b/Engine/Plugins/Enterprise/DatasmithCADImporter/Source/CADInterfaces/Private/CoreTechBridge.cpp
new file mode 100644
--- /dev/null
+++ b/Engine/Plugins/Enterprise/DatasmithCADImporter/Source/CADInterfaces/Private/CoreTechBridge.cpp
@@ -0,0 +1,1111 @@
+// Copyright Epic Games, Inc. All Rights Reserved.
+
+#include "CoreTechBridge.h"
+
+#pragma warning(push)
+#pragma warning(disable:4996) // unsafe sprintf
+#pragma warning(disable:4828) // illegal character
+#include "kernel_io/object_io/asm_io/component_io/component_io.h"
+#include "kernel_io/object_io/geom_io/curve_io/ccompo_io/ccompo_io.h"
+#include "kernel_io/object_io/geom_io/curve_io/circle_io/circle_io.h"
+#include "kernel_io/object_io/geom_io/curve_io/cnurbs_io/cnurbs_io.h"
+#include "kernel_io/object_io/geom_io/curve_io/curve_io.h"
+#include "kernel_io/object_io/geom_io/curve_io/curveonsurface_io/curveonsurface_io.h"
+#include "kernel_io/object_io/geom_io/curve_io/ellipse_io/ellipse_io.h"
+#include "kernel_io/object_io/geom_io/curve_io/hyperbola_io/hyperbola_io.h"
+#include "kernel_io/object_io/geom_io/curve_io/line_io/line_io.h"
+#include "kernel_io/object_io/geom_io/curve_io/parabola_io/parabola_io.h"
+#include "kernel_io/object_io/geom_io/surface_io/cone_io/cone_io.h"
+#include "kernel_io/object_io/geom_io/surface_io/cylinder_io/cylinder_io.h"
+#include "kernel_io/object_io/geom_io/surface_io/plane_io/plane_io.h"
+#include "kernel_io/object_io/geom_io/surface_io/slineartransfo_io/slineartransfo_io.h"
+#include "kernel_io/object_io/geom_io/surface_io/snurbs_io/snurbs_io.h"
+#include "kernel_io/object_io/geom_io/surface_io/soffset_io/soffset_io.h"
+#include "kernel_io/object_io/geom_io/surface_io/sphere_io/sphere_io.h"
+#include "kernel_io/object_io/geom_io/surface_io/srevol_io/srevol_io.h"
+#include "kernel_io/object_io/geom_io/surface_io/sruled_io/sruled_io.h"
+#include "kernel_io/object_io/geom_io/surface_io/surface_io.h"
+#include "kernel_io/object_io/geom_io/surface_io/torus_io/torus_io.h"
+#include "kernel_io/object_io/topo_io/body_io/body_io.h"
+#include "kernel_io/object_io/topo_io/coedge_io/coedge_io.h"
+#include "kernel_io/object_io/topo_io/face_io/face_io.h"
+#include "kernel_io/object_io/topo_io/loop_io/loop_io.h"
+#include "kernel_io/object_io/topo_io/shell_io/shell_io.h"
+#pragma warning(pop)
+
+#include "CADKernel/Core/MetadataDictionary.h"
+#include "CADKernel/Core/Session.h"
+
+#include "CADKernel/Geo/Curves/BezierCurve.h"
+#include "CADKernel/Geo/Curves/BoundedCurve.h"
+#include "CADKernel/Geo/Curves/CompositeCurve.h"
+#include "CADKernel/Geo/Curves/Curve.h"
+#include "CADKernel/Geo/Curves/EllipseCurve.h"
+#include "CADKernel/Geo/Curves/HyperbolaCurve.h"
+#include "CADKernel/Geo/Curves/NURBSCurve.h"
+#include "CADKernel/Geo/Curves/NURBSCurve.h"
+#include "CADKernel/Geo/Curves/ParabolaCurve.h"
+#include "CADKernel/Geo/Curves/PolylineCurve.h"
+#include "CADKernel/Geo/Curves/RestrictionCurve.h"
+#include "CADKernel/Geo/Curves/RestrictionCurve.h"
+#include "CADKernel/Geo/Curves/SegmentCurve.h"
+#include "CADKernel/Geo/Curves/SurfacicCurve.h"
+
+#include "CADKernel/Geo/Surfaces/BezierSurface.h"
+#include "CADKernel/Geo/Surfaces/CompositeSurface.h"
+#include "CADKernel/Geo/Surfaces/ConeSurface.h"
+#include "CADKernel/Geo/Surfaces/CoonsSurface.h"
+#include "CADKernel/Geo/Surfaces/CylinderSurface.h"
+#include "CADKernel/Geo/Surfaces/NURBSSurface.h"
+#include "CADKernel/Geo/Surfaces/OffsetSurface.h"
+#include "CADKernel/Geo/Surfaces/PlaneSurface.h"
+#include "CADKernel/Geo/Surfaces/RevolutionSurface.h"
+#include "CADKernel/Geo/Surfaces/RuledSurface.h"
+#include "CADKernel/Geo/Surfaces/SphericalSurface.h"
+#include "CADKernel/Geo/Surfaces/Surface.h"
+#include "CADKernel/Geo/Surfaces/TabulatedCylinderSurface.h"
+#include "CADKernel/Geo/Surfaces/TorusSurface.h"
+
+#include "CADKernel/Math/AABB.h"
+#include "CADKernel/Math/Boundary.h"
+
+#include "CADKernel/Topo/Body.h"
+#include "CADKernel/Topo/Shell.h"
+#include "CADKernel/Topo/TopologicalEdge.h"
+#include "CADKernel/Topo/TopologicalFace.h"
+#include "CADKernel/Topo/TopologicalLink.h"
+#include "CADKernel/Topo/TopologicalVertex.h"
+
+#include "CADData.h"
+
+using namespace CADKernel;
+
+FCoreTechBridge::FCoreTechBridge(TSharedRef<FSession>& InSession)
+	: Session(InSession)
+	, GeometricTolerance(Session->GetGeometricTolerance())
+{
+}
+
+TSharedRef<FBody> FCoreTechBridge::AddBody(CT_OBJECT_ID CTBodyId)
+{
+	TSharedRef<FBody> Body = FEntity::MakeShared<FBody>();
+
+	AddMetadata(CTBodyId, Body);
+
+#ifdef CORETECHBRIDGE_DEBUG
+	Body->CtKioId = (FIdent)CTBodyId;
+#endif
+
+	CTIdToEntity.Add((const uint32)CTBodyId, Body);
+
+	CT_LIST_IO CTShellIds;
+	CT_BODY_IO::AskShells(CTBodyId, CTShellIds);
+
+	CT_OBJECT_ID CTShellId;
+	while ((CTShellId = CTShellIds.IteratorIter()) != 0)
+	{
+		TSharedRef<FShell> Shell = FEntity::MakeShared<FShell>();
+		Body->AddShell(Shell);
+
+		AddMetadata(CTShellId, Shell);
+#ifdef CORETECHBRIDGE_DEBUG
+		Shell->CtKioId = (FIdent)CTShellId;
+#endif
+
+		CT_LIST_IO CTFaceIds;
+		CT_SHELL_IO::AskFaces(CTShellId, CTFaceIds);
+
+		CTFaceIds.IteratorInitialize();
+		CT_OBJECT_ID CTFaceId;
+		while ((CTFaceId = CTFaceIds.IteratorIter()) != 0)
+		{
+			AddFace(CTFaceId, Shell);
+		}
+	}
+
+	return Body;
+}
+
+void FCoreTechBridge::AddFace(CT_OBJECT_ID CTFaceId, TSharedRef<FShell>& Shell)
+{
+	FSurfacicBoundary FaceBoundary;
+	Get2DCurvesRange(CTFaceId, FaceBoundary);
+
+	CT_OBJECT_ID CTSurfaceId;
+	CT_ORIENTATION CTOrientation;
+	CT_IO_ERROR Return = CT_FACE_IO::AskSurface(CTFaceId, CTSurfaceId, CTOrientation);
+	if (Return != IO_OK || CTSurfaceId <= 0)
+	{
+		FMessage::Printf(Log, TEXT("The CTFace %d has invalid carrier surface, this face is ignored"), CTFaceId);
+		return;
+	}
+
+	TSharedPtr<FSurface> SurfacePtr = AddSurface(CTSurfaceId, FaceBoundary);
+	if (!SurfacePtr.IsValid())
+	{
+		FMessage::Printf(Log, TEXT("The CTFace %d has invalid carrier surface, this face is ignored"), CTFaceId);
+		return;
+	}
+
+	TSharedRef<FSurface> Surface = SurfacePtr.ToSharedRef();
+
+	const FSurfacicBoundary& SurfaceBounds = Surface->GetBoundary();
+	if (SurfaceBounds.IsDegenerated())
+	{
+		FMessage::Printf(Log, TEXT("The CTFace %d has degenerated carrier surface, this face is ignored ([%f, %f], [%f, %f])"), CTFaceId, SurfaceBounds.UVBoundaries[EIso::IsoU].Min, SurfaceBounds.UVBoundaries[EIso::IsoU].Max, SurfaceBounds.UVBoundaries[EIso::IsoV].Min, SurfaceBounds.UVBoundaries[EIso::IsoV].Max);
+		return;
+	}
+
+	CT_LIST_IO CTLoopIds;
+	Return = CT_FACE_IO::AskLoops(CTFaceId, CTLoopIds);
+	if (Return != IO_OK)
+	{
+		FMessage::Printf(Log, TEXT("The CTFace %d has problem to get its loops, this face is ignored"), CTFaceId);
+		return;
+	}
+
+	TArray<TSharedPtr<FTopologicalLoop>> Loops;
+
+	CTLoopIds.IteratorInitialize();
+	CT_OBJECT_ID CTLoopId;
+	while ((CTLoopId = CTLoopIds.IteratorIter()) != 0)
+	{
+		TSharedPtr<FTopologicalLoop> Loop = AddLoop(CTLoopId, Surface);
+		if (!Loop.IsValid())
+		{
+			continue;
+		}
+
+		TArray<FPoint2D> LoopSampling;
+		Loop->Get2DSampling(LoopSampling);
+		FAABB2D Boundary;
+		Boundary += LoopSampling;
+		Loop->Boundary.Set(Boundary.GetMin(), Boundary.GetMax());
+
+		// Check if the loop is not composed with only degenerated edge
+		bool bDegeneratedLoop = true;
+		for (const FOrientedEdge& Edge : Loop->GetEdges())
+		{
+			if (!Edge.Entity->IsDegenerated())
+			{
+				bDegeneratedLoop = false;
+				break;
+			}
+		}
+		if (bDegeneratedLoop)
+		{
+			continue;
+		}
+
+		Loops.Add(Loop);
+	}
+
+	if (Loops.Num() == 0)
+	{
+		FMessage::Printf(Log, TEXT("The CTFace %d is degenerate, this face is ignored"), CTFaceId);
+		return;
+	}
+
+	TSharedRef<FTopologicalFace> Face = FEntity::MakeShared<FTopologicalFace>(Surface);
+	AddMetadata(CTFaceId, Face);
+	Face->SetPatchId((int32) CTFaceId);
+
+#ifdef CORETECHBRIDGE_DEBUG
+	Face->CtKioId = (FIdent) CTFaceId;
+#endif
+
+	Face->AddLoops(Loops);
+
+	EOrientation Orientation = CTOrientation == CT_FORWARD ? EOrientation::Front : EOrientation::Back;
+
+	Shell->Add(Face, Orientation);
+}
+
+void FCoreTechBridge::Get2DCurvesRange(CT_OBJECT_ID CTFaceId, CADKernel::FSurfacicBoundary& OutBounds)
+{
+	CT_FACE_IO::AskUVminmax(CTFaceId, OutBounds.UVBoundaries[EIso::IsoU].Min, OutBounds.UVBoundaries[EIso::IsoU].Max, OutBounds.UVBoundaries[EIso::IsoV].Min, OutBounds.UVBoundaries[EIso::IsoV].Max);
+}
+
+TSharedPtr<FTopologicalLoop> FCoreTechBridge::AddLoop(CT_OBJECT_ID CTLoopId, TSharedRef<FSurface>& Surface)
+{
+	CT_LIST_IO CTCoedgeIds;
+	CT_LOOP_IO::AskCoedges(CTLoopId, CTCoedgeIds);
+
+	TArray<TSharedPtr<FTopologicalEdge>> Edges;
+	TArray<EOrientation> Directions;
+
+	int32 CoedgeCount = CTCoedgeIds.Count();
+	Edges.Reserve(CoedgeCount);
+	Directions.Reserve(CoedgeCount);
+
+	CTCoedgeIds.IteratorInitialize();
+
+	CT_OBJECT_ID CoEdgeId;
+	while ((CoEdgeId = CTCoedgeIds.IteratorIter()) != 0)
+	{
+		TSharedPtr<FTopologicalEdge> Edge = AddEdge(CoEdgeId, Surface);
+		if (!Edge.IsValid())
+		{
+			continue;
+		}
+
+		Edges.Emplace(Edge);
+		Directions.Emplace(EOrientation::Front);
+	}
+
+	if (Edges.Num() == 0)
+	{
+		return TSharedPtr<FTopologicalLoop>();
+	}
+
+	return FTopologicalLoop::Make(Edges, Directions);
+}
+
+TSharedPtr<FTopologicalEdge> FCoreTechBridge::AddEdge(CT_OBJECT_ID CTCoedgeId, TSharedRef<FSurface>& Surface)
+{
+	// build carrier curve
+	CT_UINT32 Order;
+	CT_UINT32 PoleDim;
+	CT_UINT32 KnotSize;
+	CT_UINT32 PoleNum;
+	CT_COEDGE_IO::AskUVCurveSizeArrays(CTCoedgeId, Order, PoleDim, KnotSize, PoleNum);
+
+	int32 Degre = Order - 1;
+
+	TArray<double> Knots;
+	Knots.SetNum(KnotSize);
+
+	TArray<CT_DOUBLE> RawPoles;
+	RawPoles.SetNum(PoleDim * PoleNum);
+
+	CT_COEDGE_IO::AskUVCurveArrays(CTCoedgeId, Knots.GetData(), RawPoles.GetData());
+
+	TArray<FPoint> Poles;
+	Poles.SetNum(PoleNum);
+
+	int32 Index = 0;
+	for (CT_UINT32 IPole = 0; IPole < PoleNum; ++IPole)
+	{
+		Poles[IPole].Set(RawPoles[Index], RawPoles[Index + 1]);
+		Index += PoleDim;
+	}
+
+	TArray<double> Weights;
+	if (PoleDim == 3)
+	{
+		Weights.SetNum(PoleNum);
+		Index = 2;
+		for (CT_UINT32 IPole = 0; IPole < PoleNum; ++IPole, Index += 3)
+		{
+			Weights[IPole] = RawPoles[Index];
+		}
+	}
+
+	const FSurfacicBoundary& SurfaceBounds = Surface->GetBoundary();
+
+	// Move poles inside the bounds of the carrier surface, otherwise we could have problems between 2d and 3p curve points 
+	// Indeed, if 2d point are outside the carrier surface bounds, the bounds is used...
+	for (CT_UINT32 iPole = 0; iPole < PoleNum; iPole++)
+	{
+		SurfaceBounds.MoveInsideIfNot(Poles[iPole], 0.0);
+	}
+
+	TSharedPtr<FCurve> Curve;
+	if (PoleDim == 2)
+	{
+		Curve = FEntity::MakeShared<FNURBSCurve>(Surface->Get2DTolerance(), Degre, Knots, Poles, 2);
+	}
+	else
+	{
+		Curve = FEntity::MakeShared<FNURBSCurve>(Surface->Get2DTolerance(), Degre, Knots, Poles, Weights, 2);
+	}
+
+	if (!Curve.IsValid())
+	{
+		return TSharedPtr<FTopologicalEdge>();
+	}
+
+	TSharedRef<FRestrictionCurve> RestrictionCurve = FEntity::MakeShared<FRestrictionCurve>(GeometricTolerance, Surface, Curve.ToSharedRef());
+	TSharedPtr<FTopologicalEdge> Edge = FTopologicalEdge::Make(RestrictionCurve);
+	if (!Edge.IsValid())
+	{
+		return TSharedPtr<FTopologicalEdge>();
+	}
+
+	CTIdToEntity.Add((uint32)CTCoedgeId, Edge);
+
+#ifdef CORETECHBRIDGE_DEBUG
+	Edge->CtKioId = (FIdent)CTCoedgeId;
+#endif
+
+	// Link edges
+	CT_OBJECT_ID CTConnectedCoedgeId;
+	CT_COEDGE_IO::AskConnectedCoedge(CTCoedgeId, CTConnectedCoedgeId);
+	if (CTConnectedCoedgeId > 0 && CTConnectedCoedgeId != CTCoedgeId)
+	{
+		TSharedPtr<FEntity>* TwinEdge = CTIdToEntity.Find((uint32)CTConnectedCoedgeId);
+		if (TwinEdge != nullptr)
+		{
+			Edge->Link(StaticCastSharedRef<FTopologicalEdge>(TwinEdge->ToSharedRef()));
+		}
+	}
+	return Edge;
+}
+
+TSharedPtr<FSurface> FCoreTechBridge::AddSurface(CT_OBJECT_ID CTSurfaceId, const FSurfacicBoundary& Boundary)
+{
+	TSharedPtr<FEntity>* SurfacePtr = CTIdToEntity.Find((uint32)CTSurfaceId);
+	if (SurfacePtr != nullptr)
+	{
+		return StaticCastSharedPtr<FSurface>(*SurfacePtr);
+	}
+
+	CT_OBJECT_TYPE CTSurfaceType;
+	CT_OBJECT_IO::AskType(CTSurfaceId, CTSurfaceType);
+
+	TSharedPtr<FSurface> Surface;
+	switch (CTSurfaceType)
+	{
+	case CT_PLANE_TYPE:
+		Surface = AddPlaneSurface(CTSurfaceId);
+		break;
+	case CT_S_NURBS_TYPE:
+		Surface = AddNurbsSurface(CTSurfaceId);
+		break;
+	case CT_S_REVOL_TYPE:
+		Surface = AddRevolutionSurface(CTSurfaceId, Boundary);
+		break;
+	case CT_S_OFFSET_TYPE:
+		Surface = AddOffsetSurface(CTSurfaceId, Boundary);
+		break;
+	case CT_CYLINDER_TYPE:
+		Surface = AddCylinderSurface(CTSurfaceId, Boundary);
+		break;
+	case CT_CONE_TYPE:
+		Surface = AddConeSurface(CTSurfaceId, Boundary);
+		break;
+	case CT_SPHERE_TYPE:
+		Surface = AddSphereSurface(CTSurfaceId, Boundary);
+		break;
+	case CT_TORUS_TYPE:
+		Surface = AddTorusSurface(CTSurfaceId, Boundary);
+		break;
+	case CT_S_RULED_TYPE:
+		Surface = AddRuledSurface(CTSurfaceId);
+		break;
+	case CT_S_LINEARTRANSFO_TYPE:
+		Surface = AddLinearTransfoSurface(CTSurfaceId);
+		break;
+	default:
+		FMessage::Printf(Debug, TEXT("Unknown surface type %d\n"), CTSurfaceType);
+		Surface = 0;
+	}
+
+#ifdef CORETECHBRIDGE_DEBUG
+	Surface->CtKioId = (FIdent)CTSurfaceId;
+#endif
+
+	return Surface;
+}
+
+TSharedPtr<FCurve> FCoreTechBridge::AddCurve(CT_OBJECT_ID CTCurveId, CT_OBJECT_ID CTSurfaceId)
+{
+	TSharedPtr<FEntity>* CurvePtr = CTIdToEntity.Find((uint32)CTCurveId);
+	if (CurvePtr != nullptr)
+	{
+		return StaticCastSharedPtr<FCurve>(*CurvePtr);
+	}
+
+	CT_OBJECT_TYPE CurveType;
+	CT_OBJECT_IO::AskType(CTCurveId, CurveType);
+
+	CT_OBJECT_TYPE CTSurfaceType = CT_SURFACE_TYPE;
+	if (CTSurfaceId != 0)
+	{
+		CT_OBJECT_IO::AskType(CTSurfaceId, CTSurfaceType);
+	}
+
+	TSharedPtr<FCurve> Curve;
+	switch (CurveType)
+	{
+	case CT_C_NURBS_TYPE:
+		Curve = AddNurbsCurve(CTCurveId);
+		break;
+	case CT_LINE_TYPE:
+		Curve = AddLineCurve(CTCurveId, CTSurfaceId);
+		break;
+	case CT_C_COMPO_TYPE:
+		Curve = AddCompositeCurve(CTCurveId);
+		break;
+	case CT_CIRCLE_TYPE:
+		Curve = AddCircleCurve(CTCurveId);
+		break;
+	case CT_PARABOLA_TYPE:
+		Curve = AddParabolaCurve(CTCurveId);
+		break;
+	case CT_HYPERBOLA_TYPE:
+		Curve = AddHyperbolaCurve(CTCurveId);
+		break;
+	case CT_ELLIPSE_TYPE:
+		Curve = AddEllipseCurve(CTCurveId);
+		break;
+	case CT_CURVE_ON_SURFACE_TYPE:
+		Curve = AddCurveOnSurface(CTCurveId);
+		break;
+	default:
+		FMessage::Printf(Debug, TEXT("Unknown curve type %d\n"), CurveType);
+		Curve = 0;
+	}
+
+	if (Curve.IsValid())
+	{
+		CTIdToEntity.Add((uint32)CTCurveId, Curve);
+#ifdef CORETECHBRIDGE_DEBUG
+		Curve->CtKioId = (FIdent)CTCurveId;
+#endif
+
+	}
+	return Curve;
+}
+
+TSharedPtr<FSurface> FCoreTechBridge::AddRuledSurface(CT_OBJECT_ID CTSurfaceId)
+{
+	CT_OBJECT_ID CTGeneratrix1Id;
+	CT_OBJECT_ID CTGeneratrix2Id;
+	CT_SRULED_IO::AskParameters(CTSurfaceId, CTGeneratrix1Id, CTGeneratrix2Id);
+
+	TSharedPtr<FCurve> Curve1 = AddCurve(CTGeneratrix1Id, CTSurfaceId);
+	TSharedPtr<FCurve> Curve2 = AddCurve(CTGeneratrix2Id, CTSurfaceId);
+
+	if (Curve2.IsValid() && Curve1.IsValid())
+	{
+		return  FEntity::MakeShared<FRuledSurface>(GeometricTolerance, Curve1, Curve2);
+
+	}
+	return TSharedPtr<FSurface>();
+}
+
+TSharedPtr<FSurface> FCoreTechBridge::AddTorusSurface(CT_OBJECT_ID CTSurfaceId, const FSurfacicBoundary& Boundary)
+{
+	CT_COORDINATE TorusOrigin;
+	CT_VECTOR Direction;
+	CT_VECTOR UReference;
+	CT_DOUBLE MajorRadius;
+	CT_DOUBLE MinorRadius;
+	CT_TORUS_IO::AskParameters(CTSurfaceId, TorusOrigin, Direction, MajorRadius, MinorRadius, UReference);
+
+	FMatrixH CoordinateSystem = CreateCoordinateSystem(TorusOrigin, Direction, UReference);
+
+	double MajorStartAngle = Boundary.IsValid() ? Boundary.UVBoundaries[EIso::IsoU].Min : 0.0;
+	double MajorEndAngle = Boundary.IsValid() ? Boundary.UVBoundaries[EIso::IsoU].Max : 2.0 * PI;
+	double MinorStartAngle = Boundary.IsValid() ? Boundary.UVBoundaries[EIso::IsoV].Min : 0.0;
+	double MinorEndAngle = Boundary.IsValid() ? Boundary.UVBoundaries[EIso::IsoV].Max : 2.0 * PI;
+
+	return FEntity::MakeShared<FTorusSurface>(GeometricTolerance, CoordinateSystem, MajorRadius, MinorRadius, MajorStartAngle, MajorEndAngle, MinorStartAngle, MinorEndAngle);
+}
+
+TSharedPtr<FSurface> FCoreTechBridge::AddSphereSurface(CT_OBJECT_ID CTSurfaceId, const FSurfacicBoundary& Boundary)
+{
+	CT_COORDINATE SphereOrigin;
+	CT_VECTOR Direction;
+	CT_VECTOR UReference;
+	CT_DOUBLE Radius;
+	CT_SPHERE_IO::AskParameters(CTSurfaceId, SphereOrigin, Direction, Radius, UReference);
+
+	FMatrixH CoordinateSystem = CreateCoordinateSystem(SphereOrigin, Direction, UReference);
+
+	double MeridianStartAngle = Boundary.IsValid() ? Boundary.UVBoundaries[EIso::IsoU].Min : 0.0;
+	double MeridianEndAngle = Boundary.IsValid() ? Boundary.UVBoundaries[EIso::IsoU].Max : 2.0 * PI;
+
+	double ParallelStartAngle = Boundary.IsValid() ? Boundary.UVBoundaries[EIso::IsoV].Min : 0.0;
+	double ParallelEndAngle = Boundary.IsValid() ? Boundary.UVBoundaries[EIso::IsoV].Max : 2.0 * PI;
+
+	return FEntity::MakeShared<FSphericalSurface>(GeometricTolerance, CoordinateSystem, Radius, MeridianStartAngle, MeridianEndAngle, ParallelStartAngle, ParallelEndAngle);
+}
+
+TSharedPtr<FSurface> FCoreTechBridge::AddOffsetSurface(CT_OBJECT_ID CTSurfaceID, const FSurfacicBoundary& Boundary)
+{
+	CT_OBJECT_ID CTBaseSurfaceId;
+	CT_DOUBLE OffsetValue;
+	CT_SOFFSET_IO::AskParameters(CTSurfaceID, CTBaseSurfaceId, OffsetValue);
+
+	CT_OBJECT_TYPE BaseType;
+	CT_SURFACE_IO::AskType(CTBaseSurfaceId, BaseType);
+
+	TSharedPtr<FSurface> BaseSurface = AddSurface(CTBaseSurfaceId, BaseType == CT_CONE_TYPE ? Boundary : FSurfacicBoundary());
+	if (!BaseSurface.IsValid())
+	{
+		return TSharedPtr<FSurface>();
+	}
+
+	return FEntity::MakeShared<FOffsetSurface>(GeometricTolerance, BaseSurface.ToSharedRef(), OffsetValue);
+}
+
+TSharedPtr<FSurface> FCoreTechBridge::AddLinearTransfoSurface(CT_OBJECT_ID CTSurfaceId)
+{
+	ensureCADKernel(false);
+	FMessage::Printf(Log, TEXT("The CTSurface %d is a linear transform surface, this face is ignored"), CTSurfaceId);
+	return TSharedPtr<FSurface>();
+
+	CT_OBJECT_ID CTBaseSurfaceId;
+	CT_DOUBLE CTMatrix[16];
+	CT_SLINEARTRANSFO_IO::AskParameters(CTSurfaceId, CTBaseSurfaceId, CTMatrix);
+
+	FMatrixH Matrix(CTMatrix);
+
+	TSharedPtr<FSurface> BaseSurface = AddSurface(CTBaseSurfaceId, FSurfacicBoundary());
+	TSharedPtr<FSurface> Surface = StaticCastSharedPtr<FSurface>(BaseSurface->ApplyMatrix(Matrix));
+
+#ifdef CADKERNEL_DEV 
+	const FMatrixH& BaseSurfaceMatrix = GetParamSpaceTransform(BaseSurface);
+	if (!BaseSurfaceMatrix.IsId())
+	{
+		SetParamSpaceTransorm(Surface, BaseSurfaceMatrix);
+	}
+#endif
+	return Surface;
+}
+
+TSharedPtr<FSurface> FCoreTechBridge::AddConeSurface(CT_OBJECT_ID CTSurfaceId, const FSurfacicBoundary& Boundary)
+{
+	CT_COORDINATE Origin;
+	CT_VECTOR Direction;
+	CT_VECTOR UReference;
+	CT_DOUBLE Radius;
+	CT_DOUBLE HalfAngle;
+	CT_CONE_IO::AskParameters(CTSurfaceId, Origin, Direction, Radius, HalfAngle, UReference);
+
+	FMatrixH CoordinateSystem = CreateCoordinateSystem(Origin, Direction, UReference);
+
+	double StartRuleLength = Boundary.IsValid() ? Boundary.UVBoundaries[EIso::IsoV].Min : -1e5;
+	double EndRuleLength = Boundary.IsValid() ? Boundary.UVBoundaries[EIso::IsoV].Max : 1e5;
+
+	double StartAngle = Boundary.IsValid() ? Boundary.UVBoundaries[EIso::IsoU].Min : 0.0;
+	double EndAngle = Boundary.IsValid() ? Boundary.UVBoundaries[EIso::IsoU].Max : 2.0 * PI;
+
+	return FEntity::MakeShared<FConeSurface>(GeometricTolerance, CoordinateSystem, Radius, HalfAngle, StartRuleLength, EndRuleLength, StartAngle, EndAngle);
+}
+
+TSharedPtr<FSurface> FCoreTechBridge::AddPlaneSurface(CT_OBJECT_ID CTSurfaceId)
+{
+	CT_COORDINATE Origin;
+	CT_VECTOR Normal;
+	CT_VECTOR UReference;
+	CT_PLANE_IO::AskParameters(CTSurfaceId, Origin, Normal, UReference);
+
+	FMatrixH CoordinateSystem = CreateCoordinateSystem(Origin, Normal, UReference);
+
+	return FEntity::MakeShared<FPlaneSurface>(GeometricTolerance, CoordinateSystem);
+}
+
+TSharedPtr<FSurface> FCoreTechBridge::AddCylinderSurface(CT_OBJECT_ID CTSurfaceId, const FSurfacicBoundary& Boundary)
+{
+	CT_COORDINATE Origin;
+	CT_VECTOR Direction, UReference;
+	CT_DOUBLE Radius;
+
+	CT_CYLINDER_IO::AskParameters(CTSurfaceId, Origin, Direction, Radius, UReference);
+
+	FMatrixH CoordinateSystem = CreateCoordinateSystem(Origin, Direction, UReference);
+
+	double StartLength = Boundary.IsValid() ? Boundary.UVBoundaries[EIso::IsoV].Min : -1e5;
+	double EndLength = Boundary.IsValid() ? Boundary.UVBoundaries[EIso::IsoV].Max : 1e5;
+
+	double StartAngle = Boundary.IsValid() ? Boundary.UVBoundaries[EIso::IsoU].Min : 0.0;
+	double EndAngle = Boundary.IsValid() ? Boundary.UVBoundaries[EIso::IsoU].Max : 2.0 * PI;
+
+	return FEntity::MakeShared<FCylinderSurface>(GeometricTolerance, CoordinateSystem, Radius, StartLength, EndLength, StartAngle, EndAngle);
+}
+
+TSharedPtr<FSurface> FCoreTechBridge::AddRevolutionSurface(CT_OBJECT_ID CTSurfaceId, const FSurfacicBoundary& Boundary)
+{
+	CT_COORDINATE Origin;
+	CT_VECTOR Direction;
+	CT_OBJECT_ID CTGeneratrixId;
+	CT_SREVOL_IO::AskParameters(CTSurfaceId, Origin, Direction, CTGeneratrixId);
+
+	TSharedPtr<FCurve> Generatrix = AddCurve(CTGeneratrixId, CTSurfaceId);
+	if (!Generatrix.IsValid())
+	{
+		return TSharedPtr<FSurface>();
+	}
+
+	FPoint Point1, Point2;
+	for (int32 i = 0; i < 3; i++)
+	{
+		Point1[i] = Origin.xyz[i];
+		Point2[i] = Origin.xyz[i] + Direction.xyz[i];
+	}
+	TSharedRef<FSegmentCurve> Axe = FEntity::MakeShared<FSegmentCurve>(GeometricTolerance, Point1, Point2);
+
+	double MinAngle = Boundary.IsValid() ? Boundary.UVBoundaries[EIso::IsoV].Min : 0.0;
+	double MaxAngle = Boundary.IsValid() ? Boundary.UVBoundaries[EIso::IsoV].Max : 2.0 * PI;
+
+	return FEntity::MakeShared<FRevolutionSurface>(GeometricTolerance, Axe, Generatrix.ToSharedRef(), MinAngle, MaxAngle);
+}
+
+TSharedPtr<FSurface> FCoreTechBridge::AddNurbsSurface(CT_OBJECT_ID CTSurfaceId)
+{
+	CT_UINT32 OrderU;
+	CT_UINT32 OrderV;
+	CT_UINT32 KnotUSize;
+	CT_UINT32 KnotVSize;
+	CT_UINT32 PoleUNum;
+	CT_UINT32 PoleVNum;
+	CT_UINT32 PolesDim;
+
+	CT_SNURBS_IO::AskNurbsSurfaceSizeArrays(CTSurfaceId, OrderU, OrderV, KnotUSize, KnotVSize, PoleUNum, PoleVNum, PolesDim, CT_FALSE);
+
+	CT_UINT32 DegreU = OrderU - 1;
+	CT_UINT32 DegreV = OrderV - 1;
+
+	CT_UINT32 PoleNum = PoleUNum * PoleVNum;
+
+	TArray<double> KnotsU;
+	KnotsU.SetNum(KnotUSize);
+
+	TArray<double> KnotsV;
+	KnotsV.SetNum(KnotVSize);
+
+	TArray<CT_DOUBLE> RawPoles;
+	RawPoles.SetNum(PolesDim * PoleNum);
+
+	CT_SNURBS_IO::AskNurbsSurfaceArrays(CTSurfaceId, KnotsU.GetData(), KnotsV.GetData(), RawPoles.GetData());
+
+	TArray<FPoint> Poles;
+	Poles.SetNum(PoleNum);
+
+	double* RawPolesPtr = RawPoles.GetData();
+	for (uint32 Undex = 0, RawPolesOffset = 0; Undex < PoleUNum; ++Undex)
+	{
+		for (uint32 Vndex = 0; Vndex < PoleVNum; ++Vndex, RawPolesOffset += PolesDim)
+		{
+			uint32 Index = (Vndex * PoleUNum + Undex);
+			Poles[Index].Set(RawPolesPtr + RawPolesOffset);
+		}
+	}
+
+	int32 PolesUNum = KnotsU.Num() - DegreU - 1;
+	int32 PolesVNum = KnotsV.Num() - DegreV - 1;
+	if(PolesDim == 4)
+	{
+		TArray<double> Weights;
+		Weights.SetNum(PoleNum);
+		for (uint32 Undex = 0, RawPolesOffset = 3; Undex < PoleUNum; ++Undex)
+		{
+			uint32 Index = Undex;
+			for (uint32 Vndex = 0; Vndex < PoleVNum; ++Vndex, RawPolesOffset += PolesDim, Index += PoleUNum)
+			{
+				Weights[Index] = RawPoles[RawPolesOffset];
+			}
+		}
+		return FEntity::MakeShared<FNURBSSurface>(GeometricTolerance, PolesUNum, PolesVNum, DegreU, DegreV, KnotsU, KnotsV, Poles, Weights);
+	}
+
+	return FEntity::MakeShared<FNURBSSurface>(GeometricTolerance, PolesUNum, PolesVNum, DegreU, DegreV, KnotsU, KnotsV, Poles);
+}
+
+TSharedPtr<FCurve> FCoreTechBridge::AddNurbsCurve(CT_OBJECT_ID CTCurveId)
+{
+	CT_UINT32 Order;
+	CT_UINT32 KnotSize;
+	CT_UINT32 PoleNum;
+	CT_UINT32 PoleDim;
+
+	CT_CNURBS_IO::AskNurbsCurveSizeArrays(CTCurveId, Order, PoleDim, KnotSize, PoleNum, CT_FALSE);
+	ensureCADKernel(PoleDim >= 3);
+
+	int32 Degre = Order - 1;
+
+	TArray<double> Knots;
+	Knots.SetNum(KnotSize);
+	TArray<CT_DOUBLE> RawPoles;
+	RawPoles.SetNum(PoleDim * PoleNum);
+
+	CT_CNURBS_IO::AskNurbsCurveArrays(CTCurveId, Knots.GetData(), RawPoles.GetData());
+
+	TArray<FPoint> Poles;
+	Poles.SetNum(PoleNum);
+
+	int32 Index = 0;
+	for (CT_UINT32 IPole = 0; IPole < PoleNum; IPole++, Index += PoleDim)
+	{
+		Poles[IPole].Set(RawPoles.GetData() + Index);
+	}
+
+	if (PoleDim > 3)
+	{
+		TArray<double> Weights;
+		Weights.SetNum(PoleNum);
+		Index = 3;
+		for (CT_UINT32 IPole = 0; IPole < PoleNum; IPole++, Index += 4)
+		{
+			Weights[IPole] = RawPoles[Index];
+		}
+		return FEntity::MakeShared<FNURBSCurve>(GeometricTolerance, Degre, Knots, Poles, Weights);
+	}
+
+	return FEntity::MakeShared<FNURBSCurve>(GeometricTolerance, Degre, Knots, Poles);
+}
+
+TSharedPtr<FCurve> FCoreTechBridge::AddLineCurve(CT_OBJECT_ID CTCurveId, CT_OBJECT_ID CTSurfaceId)
+{
+	CT_COORDINATE Origin;
+	CT_VECTOR Direction;
+	CT_DOUBLE Start;
+	CT_DOUBLE End;
+
+	CT_LINE_IO::AskParameters(CTCurveId, Origin, Direction, Start, End);
+
+	CT_OBJECT_TYPE CTSurfaceType = CT_SURFACE_TYPE;
+	if (CTSurfaceId != 0)
+	{
+		CT_OBJECT_IO::AskType(CTSurfaceId, CTSurfaceType);
+	}
+
+	FPoint Point1, Point2;
+	for (int32 Index = 0; Index < 3; Index++)
+	{
+		Point1[Index] = Origin.xyz[Index] + Start * Direction.xyz[Index];
+		Point2[Index] = Origin.xyz[Index] + End * Direction.xyz[Index];
+	}
+
+	if (CTSurfaceType == CT_S_REVOL_TYPE)
+	{
+		ensureCADKernel(false);
+		CT_DOUBLE TMin, TMax;
+		CT_OBJECT_ID CTGeneratrix;
+		CT_SREVOL_IO::AskParameters(CTSurfaceId, Origin, Direction, CTGeneratrix);
+		CT_CURVE_IO::AskParametricInterval(CTGeneratrix, TMin, TMax);
+		Point1[0] = (Point1[0] - TMin) / (TMax - TMin);
+		Point2[0] = (Point2[0] - TMin) / (TMax - TMin);
+	}
+
+	return FEntity::MakeShared<FSegmentCurve>(GeometricTolerance, Point1, Point2, 3);
+}
+
+TSharedPtr<FCurve> FCoreTechBridge::AddCompositeCurve(CT_OBJECT_ID CTCurveId)
+{
+	CT_INT32 CurveNum;
+	CT_CCOMPO_IO::AskParameters(CTCurveId, CurveNum);
+
+	TArray <CT_OBJECT_ID> CurvesArray;
+	CurvesArray.SetNum(CurveNum);
+	CT_CCOMPO_IO::AskParameters(CTCurveId, CurveNum, CurvesArray.GetData());
+
+	TArray<TSharedPtr<FCurve>> Curves;
+	Curves.Reserve(CurveNum);
+	for (int32 i = 0; i < CurveNum; i++)
+	{
+		TSharedPtr<FCurve> Curve = AddCurve(CurvesArray[i]);
+		if (Curve.IsValid())
+		{
+			Curves.Emplace(Curve);
+		}
+		else
+		{
+			FMessage::Printf(Log, TEXT("The CT Composite curve %d has an invalid curve, this curve is ignored"), CTCurveId);
+			return TSharedPtr<FCurve>();
+		}
+	}
+
+	return FEntity::MakeShared<FCompositeCurve>(GeometricTolerance, Curves);
+}
+
+TSharedPtr<FCurve> FCoreTechBridge::AddCircleCurve(CT_OBJECT_ID CTCurveId)
+{
+	CT_COORDINATE Origin;
+	CT_VECTOR Direction, UReference;
+	double Radius;
+	double StartAngle;
+	double EndAngle;
+
+	CT_CIRCLE_IO::AskParameters(CTCurveId, Origin, Direction, Radius, UReference, StartAngle, EndAngle);
+
+	FMatrixH CoordinateSystem = CreateCoordinateSystem(Origin, Direction, UReference);
+
+	if (FMath::IsNearlyZero(EndAngle))
+	{
+		EndAngle = 2.0 * PI;
+	}
+	return FEntity::MakeShared<FEllipseCurve>(GeometricTolerance, CoordinateSystem, Radius, Radius, FLinearBoundary(StartAngle, EndAngle));
+}
+
+TSharedPtr<FCurve> FCoreTechBridge::AddParabolaCurve(CT_OBJECT_ID CTCurveId)
+{
+	CT_COORDINATE Origin;
+	CT_VECTOR Direction;
+	CT_VECTOR UReference;
+	CT_DOUBLE Focal;
+	CT_DOUBLE StartAlpha;
+	CT_DOUBLE EndAlpha;
+
+	CT_PARABOLA_IO::AskParameters(CTCurveId, Origin, Direction, Focal, UReference, StartAlpha, EndAlpha);
+
+	FMatrixH CoordinateSystem = CreateCoordinateSystem(Origin, Direction, UReference);
+
+	return FEntity::MakeShared<FParabolaCurve>(GeometricTolerance, CoordinateSystem, Focal, FLinearBoundary(StartAlpha, EndAlpha), 3);
+}
+
+TSharedPtr<FCurve> FCoreTechBridge::AddEllipseCurve(CT_OBJECT_ID CTCurveId)
+{
+	CT_COORDINATE Origin;
+	CT_VECTOR Direction, UReference;
+	CT_DOUBLE Radius1, Radius2, StartAlpha, EndAlpha;
+
+	CT_ELLIPSE_IO::AskParameters(CTCurveId, Origin, Direction, Radius1, Radius2, UReference, StartAlpha, EndAlpha);
+
+	FMatrixH CoordinateSystem = CreateCoordinateSystem(Origin, Direction, UReference);
+
+	if (FMath::IsNearlyZero(EndAlpha))
+	{
+		StartAlpha *= -1.;
+	}
+
+	return  FEntity::MakeShared<FEllipseCurve>(GeometricTolerance, CoordinateSystem, Radius1, Radius2, FLinearBoundary(StartAlpha, EndAlpha));
+}
+
+TSharedPtr<FCurve> FCoreTechBridge::AddCurveOnSurface(CT_OBJECT_ID CTCurveId)
+{
+	CT_OBJECT_ID CTBaseSurfaceId;
+	CT_OBJECT_ID CTParametricCurveId;
+	CT_CURVEONSURFACE_IO::AskParameters(CTCurveId, CTBaseSurfaceId, CTParametricCurveId);
+
+	TSharedPtr<FSurface> Surface = AddSurface(CTBaseSurfaceId, FSurfacicBoundary());
+	if (!Surface.IsValid())
+	{
+		return TSharedPtr<FCurve>();
+	}
+
+	CT_OBJECT_TYPE CTSurfaceType;
+	CT_OBJECT_IO::AskType(CTBaseSurfaceId, CTSurfaceType);
+
+	if (CTSurfaceType == CT_S_RULED_TYPE && CTSurfaceType == CT_CYLINDER_TYPE)
+	{
+		ensureCADKernel(false);
+		FMessage::Printf(Debug, TEXT("case : ruled surface with curve on cylinder surface, curveId %d\n"), CTCurveId);
+		return TSharedPtr<FCurve>();
+	}
+
+#ifdef CADKERNEL_DEV
+	// Curves on cylinders have UV param. and Z axis inversed:
+	// Z'=-Z and {UV}' = {VU}
+	if (CTSurfaceType == CT_CYLINDER_TYPE)
+	{
+		ensureCADKernel(false);
+		FMatrixH MatrixTransformZandUV;
+		MatrixTransformZandUV.SetIdentity();
+		MatrixTransformZandUV[0] = 0;
+		MatrixTransformZandUV[1] = -1;
+		MatrixTransformZandUV[4] = 1;
+		MatrixTransformZandUV[5] = 0;
+		SetParamSpaceTransorm(Surface, MatrixTransformZandUV);
+	}
+#endif
+
+	TSharedPtr<FCurve> Curve2D = AddCurve(CTParametricCurveId, CTBaseSurfaceId);
+	if (!Curve2D.IsValid())
+	{
+		return TSharedPtr<FCurve>();
+	}
+
+#ifdef CADKERNEL_DEV
+	FMatrixH Matrix = GetParamSpaceTransform(Surface);
+	if (!Matrix.IsId())
+	{
+		Curve2D = StaticCastSharedPtr<FCurve>(Curve2D->ApplyMatrix(Matrix));
+	}
+
+	if (!Curve2D.IsValid())
+	{
+		return TSharedPtr<FCurve>();
+	}
+#endif
+
+	return FEntity::MakeShared<FSurfacicCurve>(GeometricTolerance, Curve2D.ToSharedRef(), Surface.ToSharedRef());
+}
+
+TSharedPtr<FCurve> FCoreTechBridge::AddHyperbolaCurve(CT_OBJECT_ID CTCurveId)
+{
+	CT_COORDINATE Origin;
+	CT_VECTOR Direction;
+	CT_VECTOR UReference;
+	CT_DOUBLE HalfAxis1;
+	CT_DOUBLE HalfAxis2;
+	CT_DOUBLE StartAlpha;
+	CT_DOUBLE EndAlpha;
+
+	CT_HYPERBOLA_IO::AskParameters(CTCurveId, Origin, Direction, HalfAxis1, HalfAxis2, UReference, StartAlpha, EndAlpha);
+
+	FMatrixH CoordinateSystem = CreateCoordinateSystem(Origin, Direction, UReference);
+
+	return FEntity::MakeShared<FHyperbolaCurve>(GeometricTolerance, CoordinateSystem, HalfAxis1, HalfAxis2, FLinearBoundary(StartAlpha, EndAlpha));
+}
+
+FMatrixH FCoreTechBridge::CreateCoordinateSystem(const CT_COORDINATE& InOrigin, const CT_VECTOR& InDirection, const CT_VECTOR& InUReference)
+{
+	FPoint Origin(InOrigin.xyz);
+	FPoint Ox(InUReference.xyz);
+	FPoint Oz(InDirection.xyz);
+ 	FPoint Oy = Oz ^ Ox;
+	return FMatrixH(Origin, Ox, Oy, Oz);
+}
+
+enum class EValueType : uint8
+{
+	Unknown = 0,
+
+	Integer,
+	Double,
+	String,
+};
+
+void GetAttributeValue(CT_ATTRIB_TYPE AttributType, int IthField, EValueType ValueType, FString& StringValue, int32& IntegerValue, double& DoubleValue)
+{
+	CT_STR               FieldName;
+	CT_ATTRIB_FIELD_TYPE FieldType;
+
+	StringValue = TEXT("");
+
+	if (CT_ATTRIB_DEFINITION_IO::AskFieldDefinition(AttributType, IthField, FieldType, FieldName) != IO_OK)
+	{
+		return;
+	}
+
+	switch (FieldType) {
+	case CT_ATTRIB_FIELD_UNKNOWN:
+	{
+		ValueType = EValueType::Unknown;
+		break;
+	}
+
+	case CT_ATTRIB_FIELD_INTEGER:
+	{
+		ValueType = EValueType::Integer;
+		if (CT_CURRENT_ATTRIB_IO::AskIntField(IthField, IntegerValue) != IO_OK)
+		{
+			ValueType = EValueType::Unknown;
+			break;
+		}
+		break;
+	}
+
+	case CT_ATTRIB_FIELD_DOUBLE:
+	{
+		ValueType = EValueType::Double;
+		if (CT_CURRENT_ATTRIB_IO::AskDblField(IthField, DoubleValue) != IO_OK)
+		{
+			ValueType = EValueType::Unknown;
+			break;
+		}
+		break;
+	}
+
+	case CT_ATTRIB_FIELD_STRING:
+	{
+		ValueType = EValueType::String;
+		CT_STR StrValue;
+		if (CT_CURRENT_ATTRIB_IO::AskStrField(IthField, StrValue) != IO_OK)
+		{
+			ValueType = EValueType::Unknown;
+			break;
+		}
+		StringValue = FCoreTechBridge::AsFString(StrValue);
+		break;
+	}
+
+	default:
+	{
+		ValueType = EValueType::Unknown;
+		break;
+	}
+	}
+}
+
+void FCoreTechBridge::AddMetadata(CT_OBJECT_ID CTNodeId, TSharedRef<FMetadataDictionary> Entity)
+{
+	Entity->SetHostId(CTNodeId);
+
+	CT_UINT32 IthAttrib = 0;
+	while (CT_OBJECT_IO::SearchAttribute(CTNodeId, CT_ATTRIB_ALL, IthAttrib++) == IO_OK)
+	{
+		// Get the current attribute type
+		CT_ATTRIB_TYPE       AttributeType;
+
+		if (CT_CURRENT_ATTRIB_IO::AskAttributeType(AttributeType) != IO_OK)
+		{
+			continue;
+		}
+
+		switch (AttributeType) 
+		{
+		case CT_ATTRIB_NAME:
+		{
+			CT_STR NameStrValue;
+			if (CT_CURRENT_ATTRIB_IO::AskStrField(ITH_NAME_VALUE, NameStrValue) == IO_OK)
+			{
+				Entity->SetName(AsFString(NameStrValue));
+			}
+			break;
+		}
+
+		case CT_ATTRIB_ORIGINAL_NAME:
+		{
+			CT_STR NameStrValue;
+			if (CT_CURRENT_ATTRIB_IO::AskStrField(ITH_NAME_VALUE, NameStrValue) == IO_OK)
+			{
+				Entity->SetName(AsFString(NameStrValue));
+			}
+			break;
+		}
+
+		case CT_ATTRIB_LAYERID:
+		{
+			int32 LayerId = 0;
+			if (CT_CURRENT_ATTRIB_IO::AskIntField(ITH_LAYERID_VALUE, LayerId) != IO_OK)
+			{
+				LayerId = 0;
+			}
+			Entity->SetLayer(LayerId);
+			break;
+
+		}
+
+		case CT_ATTRIB_COLORID:
+		{
+			int32 ColorId = 0;
+			if (CT_CURRENT_ATTRIB_IO::AskIntField(ITH_COLORID_VALUE, ColorId) != IO_OK)
+			{
+				break;
+			}
+
+			uint8 Alpha = 255;
+			if (CT_OBJECT_IO::SearchAttribute(CTNodeId, CT_ATTRIB_TRANSPARENCY) == IO_OK)
+			{
+				CT_DOUBLE Transparency = 0.;
+				if (CT_CURRENT_ATTRIB_IO::AskDblField(0, Transparency) == IO_OK)
+				{
+					Alpha = (uint8)(FMath::Max((1. - Transparency), Transparency) * 255.);
+				}
+			}
+
+			uint32 ColorHId = CADLibrary::BuildColorId(ColorId, Alpha);
+			Entity->SetColorId(ColorHId);
+		}
+		break;
+
+		case CT_ATTRIB_MATERIALID:
+		{
+			int32 MaterialId = 0;
+			if (CT_CURRENT_ATTRIB_IO::AskIntField(ITH_MATERIALID_VALUE, MaterialId) == IO_OK)
+			{
+				Entity->SetMaterialId(MaterialId);
+			}
+			break;
+		}
+
+		default:
+			break;
+		}
+	}
+}
\ No newline at end of file
diff --git a/Engine/Plugins/Enterprise/DatasmithCADImporter/Source/CADInterfaces/Private/CoreTechBridge.h b/Engine/Plugins/Enterprise/DatasmithCADImporter/Source/CADInterfaces/Private/CoreTechBridge.h
new file mode 100644
--- /dev/null
+++ b/Engine/Plugins/Enterprise/DatasmithCADImporter/Source/CADInterfaces/Private/CoreTechBridge.h
@@ -0,0 +1,112 @@
+// Copyright Epic Games, Inc. All Rights Reserved.
+#pragma once
+#include "CADKernel/Core/Types.h"
+
+#ifdef CADKERNEL_DEV
+#include "CADKernel/Geo/Surfaces/Surface.h"
+#endif
+#include "CADKernel/Math/MatrixH.h"
+
+#pragma warning(push)
+#pragma warning(disable:4996) // unsafe sprintf
+#pragma warning(disable:4828) // illegal character
+#include "kernel_io/kernel_io.h"
+#include "kernel_io/object_io/object_io.h"
+#pragma warning(pop)
+
+namespace CADKernel
+{
+	struct FSurfacicBoundary;
+
+	class FBody;
+	class FCriterion;
+	class FCurve;
+	class FEntity;
+	class FFaceMesh;
+	class FMetadataDictionary;
+	class FModel;
+	class FPoint;
+	class FSession;
+	class FShell;
+	class FSurface;
+	class FTopologicalEdge;
+	class FTopologicalLoop;
+
+	class FCoreTechBridge
+	{
+	private:
+
+		TSharedRef<FSession> Session;
+		const double GeometricTolerance;
+
+		TMap<const uint32, TSharedPtr<FEntity>> CTIdToEntity;
+
+		TArray<TSharedPtr<FCriterion>> Criteria;
+
+	public:
+		FCoreTechBridge(TSharedRef<FSession>& InSession);
+
+		TSharedRef<FBody> AddBody(CT_OBJECT_ID CTBodyId);
+
+		static FString AsFString(CT_STR CtName)
+		{
+			return CtName.IsEmpty() ? FString() : CtName.toUnicode();
+		};
+
+	private:
+		void AddFace(CT_OBJECT_ID CTFaceId, TSharedRef<FShell>& Shell);
+
+		TSharedPtr<FTopologicalLoop> AddLoop(CT_OBJECT_ID CTLoopId, TSharedRef<FSurface>& Surface);
+		TSharedPtr<FTopologicalEdge> AddEdge(CT_OBJECT_ID CTCoedgeId, TSharedRef<FSurface>& Surface);
+
+		TSharedPtr<FSurface> AddSurface(CT_OBJECT_ID CTSurfaceId, const FSurfacicBoundary& OutBounds);
+		TSharedPtr<FSurface> AddPlaneSurface(CT_OBJECT_ID CTSurfaceId);
+		TSharedPtr<FSurface> AddNurbsSurface(CT_OBJECT_ID CTSurfaceId);
+		TSharedPtr<FSurface> AddRevolutionSurface(CT_OBJECT_ID CTSurfaceId, const FSurfacicBoundary& Bounds);
+		TSharedPtr<FSurface> AddConeSurface(CT_OBJECT_ID CTSurfaceId, const FSurfacicBoundary& Bounds);
+		TSharedPtr<FSurface> AddOffsetSurface(CT_OBJECT_ID CTSurfaceId, const FSurfacicBoundary& Bounds);
+		TSharedPtr<FSurface> AddCylinderSurface(CT_OBJECT_ID CTSurfaceId, const FSurfacicBoundary& Bounds);
+		TSharedPtr<FSurface> AddSphereSurface(CT_OBJECT_ID CTSurfaceId, const FSurfacicBoundary& Bounds);
+		TSharedPtr<FSurface> AddTorusSurface(CT_OBJECT_ID CTSurfaceId, const FSurfacicBoundary& Bounds);
+		TSharedPtr<FSurface> AddRuledSurface(CT_OBJECT_ID CTSurfaceId);
+		TSharedPtr<FSurface> AddLinearTransfoSurface(CT_OBJECT_ID CTSurfaceId);
+
+		TSharedPtr<FCurve> AddCircleCurve(CT_OBJECT_ID CTCurveId);
+		TSharedPtr<FCurve> AddCompositeCurve(CT_OBJECT_ID CTCurveId);
+		TSharedPtr<FCurve> AddCurve(CT_OBJECT_ID CTCurveId, CT_OBJECT_ID CTSurfaceId = 0);
+		TSharedPtr<FCurve> AddCurveOnSurface(CT_OBJECT_ID CTCurveId);
+		TSharedPtr<FCurve> AddEllipseCurve(CT_OBJECT_ID CTCurveId);
+		TSharedPtr<FCurve> AddHyperbolaCurve(CT_OBJECT_ID CTCurveId);
+		TSharedPtr<FCurve> AddLineCurve(CT_OBJECT_ID CTCurveId, CT_OBJECT_ID CTSurfaceId = 0);
+		TSharedPtr<FCurve> AddNurbsCurve(CT_OBJECT_ID CTCurveId);
+		TSharedPtr<FCurve> AddParabolaCurve(CT_OBJECT_ID CTCurveId);
+
+		void Get2DCurvesRange(CT_OBJECT_ID CTFaceId, FSurfacicBoundary& OutBounds);
+
+		FMatrixH CreateCoordinateSystem(const CT_COORDINATE& InOrigin, const CT_VECTOR& InDirection, const CT_VECTOR& InURef);
+
+		void AddMetadata(CT_OBJECT_ID NodeId, TSharedRef<FMetadataDictionary> Entity);
+
+		template<typename EntityType>
+		void AddMetadata(CT_OBJECT_ID NodeId, TSharedRef<EntityType> Entity)
+		{
+			AddMetadata(NodeId, (TSharedRef<FMetadataDictionary>) Entity);
+		}
+
+#ifdef CADKERNEL_DEV
+		TMap<const uint32, FMatrixH> SurfaceToMatrix;
+
+		const FMatrixH& GetParamSpaceTransform(TSharedPtr<FSurface>& Surface)
+		{
+			FMatrixH* Matrix = SurfaceToMatrix.Find(Surface->GetId());
+			return Matrix == nullptr ? FMatrixH::Identity : *Matrix;
+		}
+
+		void SetParamSpaceTransorm(TSharedPtr<FSurface>& Surface, const FMatrixH Matrix)
+		{
+			SurfaceToMatrix.Add(Surface->GetId(), Matrix);
+		}
+#endif
+
+	};
+}
diff --git a/Engine/Plugins/Enterprise/DatasmithCADImporter/Source/CADInterfaces/Private/CoreTechFileReader.cpp b/Engine/Plugins/Enterprise/DatasmithCADImporter/Source/CADInterfaces/Private/CoreTechFileReader.cpp
--- a/Engine/Plugins/Enterprise/DatasmithCADImporter/Source/CADInterfaces/Private/CoreTechFileReader.cpp
+++ b/Engine/Plugins/Enterprise/DatasmithCADImporter/Source/CADInterfaces/Private/CoreTechFileReader.cpp
@@ -22,6 +22,24 @@
 #include "CADData.h"
 #include "CADOptions.h"
 #include "DatasmithUtils.h"
+
+#include "CADKernel/Core/Entity.h"
+#include "CADKernel/Core/MetadataDictionary.h"
+#include "CADKernel/Core/Session.h"
+#include "CADKernel/Core/Types.h"
+
+#include "CADKernel/Mesh/Criteria/Criterion.h"
+#include "CADKernel/Mesh/Meshers/ParametricMesher.h"
+#include "CADKernel/Mesh/Structure/FaceMesh.h"
+#include "CADKernel/Mesh/Structure/ModelMesh.h"
+
+#include "CADKernel/Topo/Model.h"
+#include "CADKernel/Topo/Body.h"
+#include "CADKernel/Topo/Shell.h"
+#include "CADKernel/Topo/TopologicalEdge.h"
+#include "CADKernel/Topo/TopologicalFace.h"
+#include "CADKernel/Topo/TopologicalVertex.h"
+
 #include "HAL/FileManager.h"
 #include "Internationalization/Text.h"
 #include "Templates/TypeHash.h"
@@ -549,10 +567,10 @@ namespace CADLibrary
 			}
 			else
 			{
-				//if (ReadBody(BodyId, ComponentId, DefaultMaterialHash, false))
-				//{
-				//	Context.SceneGraphArchive.ComponentSet[Index].Children.Add(BodyId);
-				//}
+				if (ReadBody(BodyId, ComponentId, DefaultMaterialHash, false))
+				{
+					Context.SceneGraphArchive.ComponentSet[Index].Children.Add(BodyId);
+				}
 			}
 		}
 
@@ -710,6 +728,86 @@ namespace CADLibrary
 		return true;
 	}
 
+#ifdef CORETECHBRIDGE_DEBUG
+	static int32 BodyIndex = 0;
+#endif
+	bool FCoreTechFileReader::ReadBody(CT_OBJECT_ID BodyId, CT_OBJECT_ID ParentId, uint32 DefaultMaterialHash, bool bNeedRepair)
+	{
+		// Is this body a constructive geometry ?
+		CT_LIST_IO FaceList;
+		CT_BODY_IO::AskFaces(BodyId, FaceList);
+		if (1 == FaceList.Count())
+		{
+			FaceList.IteratorInitialize();
+			FString Value;
+			GetStringMetaDataValue(FaceList.IteratorIter(), TEXT("Constructive Plane"), Value);
+			if (Value == TEXT("true"))
+			{
+				return false;
+			}
+		}
+
+		int32 Index = Context.SceneGraphArchive.BodySet.Emplace(BodyId);
+		Context.SceneGraphArchive.CADIdToBodyIndex.Add(BodyId, Index);
+		ReadNodeMetaData(BodyId, Context.SceneGraphArchive.BodySet[Index].MetaData);
+
+		int32 BodyMeshIndex = Context.BodyMeshes.Emplace(BodyId);
+
+		if (uint32 MaterialHash = GetObjectMaterial(Context.SceneGraphArchive.BodySet[Index]))
+		{
+			DefaultMaterialHash = MaterialHash;
+		}
+
+		FBodyMesh& BodyMesh = Context.BodyMeshes[BodyMeshIndex];
+		Context.SceneGraphArchive.BodySet[Index].MeshActorName = CoreTechFileReaderUtils::GetStaticMeshUuid(*Context.SceneGraphArchive.ArchiveFileName, BodyId);
+		BodyMesh.MeshActorName = Context.SceneGraphArchive.BodySet[Index].MeshActorName;
+
+		{
+			TSharedRef<CADKernel::FSession> CADKernelSession = MakeShared<CADKernel::FSession>(0.00001 / Context.ImportParameters.MetricUnit);
+
+			TSharedRef<CADKernel::FModel> CADKernelModel = CADKernelSession->GetModel();
+
+			CADKernel::FCoreTechBridge CoreTechBridge(CADKernelSession);
+
+			TSharedRef<CADKernel::FBody> CADKernelBody = CoreTechBridge.AddBody(BodyId);
+			CADKernelModel->Add(CADKernelBody);
+
+#ifdef CORETECHBRIDGE_DEBUG
+			FString FolderName = FPaths::GetCleanFilename(FileDescription.MainCadFilePath);
+			CADKernelSession->SaveDatabase(*FPaths::Combine(Context.CachePath, TEXT("CADKernel"), FolderName, FString::Printf(TEXT("%06d_"), BodyIndex++) + FileDescription.Name + TEXT(".ugeom")));
+#endif
+
+			// Save Body in CADKernelArchive file for re-tessellation
+			if (!Context.CachePath.IsEmpty())
+			{
+				FString BodyFile = FString::Printf(TEXT("UEx%08x"), Context.SceneGraphArchive.BodySet[Index].MeshActorName);
+				CADKernelSession->SaveDatabase(*FPaths::Combine(Context.CachePath, TEXT("body"), BodyFile + TEXT(".ugeom")), CADKernelBody);
+			}
+
+			// Tesselate the body
+			TSharedRef<CADKernel::FModelMesh> CADKernelModelMesh = CADKernel::FEntity::MakeShared<CADKernel::FModelMesh>();
+
+			DefineMeshCriteria(CADKernelModelMesh);
+
+			CADKernel::FParametricMesher Mesher(CADKernelModelMesh);
+			Mesher.MeshEntity(CADKernelModel);
+
+			TFunction<void(FObjectDisplayDataId, FObjectDisplayDataId, int32)> ProcessFace;
+			ProcessFace = [&](FObjectDisplayDataId FaceMaterial, FObjectDisplayDataId BodyMaterial, int32 Index)
+			{
+				SetFaceMainMaterial(FaceMaterial, BodyMaterial, BodyMesh, Index);
+			};
+
+			CADKernelUtils::GetBodyTessellation(CADKernelModelMesh, CADKernelBody, BodyMesh, DefaultMaterialHash, ProcessFace);
+
+			CADKernelSession->Clear();
+		}
+
+		Context.SceneGraphArchive.BodySet[Index].ColorFaceSet = BodyMesh.ColorSet;
+		Context.SceneGraphArchive.BodySet[Index].MaterialFaceSet = BodyMesh.MaterialSet;
+		return true;
+	}
+
 	void FCoreTechFileReader::GetAttributeValue(CT_ATTRIB_TYPE AttributType, int IthField, FString& Value)
 	{
 		CT_STR               FieldName;
@@ -1287,6 +1385,32 @@ namespace CADLibrary
 		}
 	}
 
+	void FCoreTechFileReader::DefineMeshCriteria(TSharedRef<CADKernel::FModelMesh>& MeshModel)
+	{
+		{
+			TSharedPtr<CADKernel::FCriterion> CurvatureCriterion = CADKernel::FCriterion::CreateCriterion(CADKernel::ECriterion::CADCurvature);
+			MeshModel->AddCriterion(CurvatureCriterion);
+		}
+
+		if (Context.ImportParameters.MaxEdgeLength > SMALL_NUMBER)
+		{
+			TSharedPtr<CADKernel::FCriterion> MaxSizeCriterion = CADKernel::FCriterion::CreateCriterion(CADKernel::ECriterion::MaxSize, Context.ImportParameters.MaxEdgeLength);
+			MeshModel->AddCriterion(MaxSizeCriterion);
+		}
+
+		if (Context.ImportParameters.ChordTolerance > SMALL_NUMBER)
+		{
+			TSharedPtr<CADKernel::FCriterion> ChordCriterion = CADKernel::FCriterion::CreateCriterion(CADKernel::ECriterion::Sag, Context.ImportParameters.ChordTolerance);
+			MeshModel->AddCriterion(ChordCriterion);
+		}
+
+		if (Context.ImportParameters.MaxNormalAngle > SMALL_NUMBER)
+		{
+			TSharedPtr<CADKernel::FCriterion> MaxNormalAngleCriterion = CADKernel::FCriterion::CreateCriterion(CADKernel::ECriterion::Angle, Context.ImportParameters.MaxNormalAngle);
+			MeshModel->AddCriterion(MaxNormalAngleCriterion);
+		}
+	}
+
 	namespace CoreTechFileReaderUtils
 	{
 		template<typename ValueType>
@@ -1884,6 +2008,109 @@ namespace CADLibrary
 		}
 
 	}
+
+	namespace CADKernelUtils
+	{
+		uint32 GetFaceTessellation(const TSharedRef<CADKernel::FFaceMesh>& FaceMesh, FBodyMesh& OutBodyMesh)
+		{
+			// Something wrong happened, either an error or no data to collect
+			if (FaceMesh->TrianglesVerticesIndex.Num() == 0)
+			{
+				return 0;
+			}
+
+			FTessellationData& Tessellation = OutBodyMesh.Faces.Emplace_GetRef();
+			
+			const TSharedRef<CADKernel::FMetadataDictionary>& HaveMetadata = StaticCastSharedRef<CADKernel::FMetadataDictionary>(StaticCastSharedRef<CADKernel::FTopologicalFace>(FaceMesh->GetGeometricEntity()));
+			Tessellation.PatchId = HaveMetadata->GetPatchId();
+
+			Tessellation.PositionIndices = FaceMesh->VerticesGlobalIndex;
+			Tessellation.VertexIndices = FaceMesh->TrianglesVerticesIndex;
+
+			Tessellation.NormalArray.Reserve(FaceMesh->Normals.Num());
+			for (const CADKernel::FPoint& Normal : FaceMesh->Normals)
+			{
+				Tessellation.NormalArray.Emplace((float)Normal.X, (float)Normal.Y, (float)Normal.Z);
+			}
+
+			Tessellation.TexCoordArray.Reserve(FaceMesh->UVMap.Num());
+			for (const CADKernel::FPoint2D& TexCoord : FaceMesh->UVMap)
+			{
+				Tessellation.TexCoordArray.Emplace((float)TexCoord.U, (float)TexCoord.V);
+			}
+
+			return Tessellation.VertexIndices.Num() / 3;
+		}
+
+		template<class ClassType>
+		void GetDisplayDataIds(const TSharedRef<ClassType>& Entity, FObjectDisplayDataId& DisplayDataId)
+		{
+			const TSharedRef<CADKernel::FMetadataDictionary>& HaveMetadata = StaticCastSharedRef<CADKernel::FMetadataDictionary>(Entity);
+			DisplayDataId.Color = HaveMetadata->GetColorId();
+			DisplayDataId.Material = HaveMetadata->GetMaterialId();
+		}
+
+		void GetBodyTessellation(const TSharedRef<CADKernel::FModelMesh>& ModelMesh, const TSharedRef<CADKernel::FBody>& Body, FBodyMesh& OutBodyMesh, uint32 DefaultMaterialHash, TFunction<void(FObjectDisplayDataId, FObjectDisplayDataId, int32)> SetFaceMainMaterial)
+		{
+			ModelMesh->GetNodeCoordinates(OutBodyMesh.VertexArray);
+
+			uint32 FaceSize = Body->FaceCount();
+
+			// Allocate memory space for tessellation data
+			OutBodyMesh.Faces.Reserve(FaceSize);
+			OutBodyMesh.ColorSet.Reserve(FaceSize);
+			OutBodyMesh.MaterialSet.Reserve(FaceSize);
+
+			FObjectDisplayDataId BodyMaterial;
+			BodyMaterial.DefaultMaterialName = DefaultMaterialHash;
+
+			GetDisplayDataIds(Body, BodyMaterial);
+
+			// Loop through the face of bodies and collect all tessellation data
+			int32 FaceIndex = 0;
+			for (const TSharedPtr<CADKernel::FShell>& Shell : Body->GetShells())
+			{
+				if (!Shell.IsValid())
+				{
+					continue;
+				}
+
+				FObjectDisplayDataId ShellMaterial = BodyMaterial;
+				GetDisplayDataIds(Shell.ToSharedRef(), ShellMaterial);
+
+				for (const CADKernel::FOrientedFace& Face : Shell->GetFaces())
+				{
+					if (!Face.Entity.IsValid())
+					{
+						continue;
+					}
+
+					if (!Face.Entity->HasTesselation())
+					{
+						continue;
+					}
+
+					FObjectDisplayDataId FaceMaterial;
+					GetDisplayDataIds(Face.Entity.ToSharedRef(), FaceMaterial);
+
+					uint32 TriangleNum = GetFaceTessellation(Face.Entity->GetMesh(), OutBodyMesh);
+
+					if (TriangleNum == 0)
+					{
+						continue;
+					}
+
+					OutBodyMesh.TriangleCount += TriangleNum;
+
+					if (SetFaceMainMaterial)
+					{
+						SetFaceMainMaterial(FaceMaterial, BodyMaterial, FaceIndex);
+					}
+					FaceIndex++;
+				}
+			}
+		}
+	}
 }
 
 #endif // USE_KERNEL_IO_SDK
diff --git a/Engine/Plugins/Enterprise/DatasmithCADImporter/Source/CADInterfaces/Private/CoreTechFileReader.h b/Engine/Plugins/Enterprise/DatasmithCADImporter/Source/CADInterfaces/Private/CoreTechFileReader.h
--- a/Engine/Plugins/Enterprise/DatasmithCADImporter/Source/CADInterfaces/Private/CoreTechFileReader.h
+++ b/Engine/Plugins/Enterprise/DatasmithCADImporter/Source/CADInterfaces/Private/CoreTechFileReader.h
@@ -10,6 +10,8 @@
 #include "CADSceneGraph.h"
 #include "CoreTechTypes.h"
 
+#include "CoreTechBridge.h"
+
 #include "Containers/Map.h"
 #include "Containers/Queue.h"
 #include "Misc/Paths.h"
@@ -27,6 +29,14 @@ struct FFileDescription;
 #include "kernel_io/material_io/material_io.h"
 #pragma warning(pop)
 
+namespace CADKernel
+{
+	class FBody;
+	class FBodyMesh;
+	class FFaceMesh;
+	class FModelMesh;
+}
+
 namespace CADLibrary
 {
 	namespace CoreTechFileReaderUtils
@@ -36,6 +46,14 @@ namespace CADLibrary
 		void GetBodyTessellation(CT_OBJECT_ID BodyId, FBodyMesh& OutBodyMesh, TFunction<void(CT_OBJECT_ID, int32, FTessellationData&)> = TFunction<void(CT_OBJECT_ID, int32, FTessellationData&)>());
 	}
 
+	namespace CADKernelUtils
+	{
+		void GetBodyTessellation(const TSharedRef<CADKernel::FModelMesh>& ModelMesh, const TSharedRef<CADKernel::FBody>& Body, FBodyMesh& OutBodyMesh, uint32 DefaultMaterialHash, TFunction<void(FObjectDisplayDataId, FObjectDisplayDataId, int32)> SetFaceMainMaterial);
+
+		uint32 GetFaceTessellation(const TSharedRef<CADKernel::FFaceMesh>& FaceMesh, FBodyMesh& OutBodyMesh);
+	}
+
+
 	class CADINTERFACES_API FCoreTechFileReader
 	{
 	public:
@@ -83,6 +101,7 @@ namespace CADLibrary
 		bool ReadUnloadedComponent(CT_OBJECT_ID NodeId);
 
 		bool ReadKioBody(CT_OBJECT_ID NodeId, CT_OBJECT_ID ParentId, uint32 ParentMaterialHash, bool bNeedRepair);
+		bool ReadBody(CT_OBJECT_ID NodeId, CT_OBJECT_ID ParentId, uint32 ParentMaterialHash, bool bNeedRepair);
 
 		bool FindFile(FFileDescription& FileDescription);
 
@@ -100,6 +119,8 @@ namespace CADLibrary
 		CT_FLAGS SetCoreTechImportOption();
 		void GetAttributeValue(CT_ATTRIB_TYPE attrib_type, int ith_field, FString& value);
 
+		void DefineMeshCriteria(TSharedRef<CADKernel::FModelMesh>& MeshModel);
+
 	protected:
 		CADLibrary::FFileDescription FileDescription;
 
diff --git a/Engine/Plugins/Enterprise/DatasmithCADImporter/Source/CADLibrary/Private/MeshDescriptionHelper.cpp b/Engine/Plugins/Enterprise/DatasmithCADImporter/Source/CADLibrary/Private/MeshDescriptionHelper.cpp
--- a/Engine/Plugins/Enterprise/DatasmithCADImporter/Source/CADLibrary/Private/MeshDescriptionHelper.cpp
+++ b/Engine/Plugins/Enterprise/DatasmithCADImporter/Source/CADLibrary/Private/MeshDescriptionHelper.cpp
@@ -205,6 +205,90 @@ namespace CADLibrary
 		}
 	}
 
+	void FillVertexPosition(const FImportParameters& ImportParams, const FMeshParameters& MeshParameters, FBodyMesh& Body, FMeshDescription& MeshDescription)
+	{
+		int32 TriangleCount = Body.TriangleCount;
+		TArray<FTessellationData>& FaceTessellationSet = Body.Faces;
+
+		// Add offset to the bounding box to avoid to remove good vertex
+		FVector Size = Body.BBox.GetSize();
+		FBox BBox = Body.BBox.ExpandBy(Size.Size());
+		BBox.IsValid = Body.BBox.IsValid;
+		BBox.Min *= ImportParams.ScaleFactor;
+		BBox.Max *= ImportParams.ScaleFactor;
+
+		TVertexAttributesRef<FVector> VertexPositions = MeshDescription.GetVertexPositions();
+
+		FVector Position;
+
+		TArray<FVector>& VertexArray = Body.VertexArray;
+
+		for (FVector& Vertex : VertexArray)
+		{
+			Vertex *= ImportParams.ScaleFactor;
+			}
+
+		TArray<int32>& VertexIdSet = Body.VertexIds;
+		int32 VertexCount = VertexArray.Num();
+		VertexIdSet.SetNumZeroed(VertexCount);
+
+		if (BBox.IsValid)
+		{
+			int32 VertexIndex = 0;
+			for (const FVector& Vertex : VertexArray)
+			{
+				if (!BBox.IsInside(Vertex))
+				{
+					VertexIdSet[VertexIndex] = INDEX_NONE;
+				}
+				VertexIndex++;
+			}
+		}
+
+		// Make MeshDescription.VertexPositions and VertexID
+		MeshDescription.ReserveNewVertices(MeshParameters.bIsSymmetric ? VertexCount * 2 : VertexCount);
+
+		int32 VertexIndex = -1;
+		for (const FVector& Vertex : VertexArray)
+		{
+			VertexIndex++;
+
+			// Vertex is outside bbox
+			if (VertexIdSet[VertexIndex] == INDEX_NONE)
+			{
+				continue;
+			}
+
+			FVertexID VertexID = MeshDescription.CreateVertex();
+			VertexPositions[VertexID] = FDatasmithUtils::ConvertVector((FDatasmithUtils::EModelCoordSystem) ImportParams.ModelCoordSys, Vertex);
+			VertexIdSet[VertexIndex] = VertexID;
+		}
+
+		// if Symmetric mesh, the symmetric side of the mesh have to be generated
+		if (MeshParameters.bIsSymmetric)
+		{
+			FMatrix SymmetricMatrix = FDatasmithUtils::GetSymmetricMatrix(MeshParameters.SymmetricOrigin, MeshParameters.SymmetricNormal);
+
+			TArray<int32>& SymmetricVertexIds = Body.SymmetricVertexIds;
+			SymmetricVertexIds.SetNum(VertexArray.Num());
+
+			VertexIndex = 0;
+			for (const FVector& Vertex : VertexArray)
+			{
+				if (VertexIdSet[VertexIndex] == INDEX_NONE)
+				{
+					SymmetricVertexIds[VertexIndex++] = INDEX_NONE;
+					continue;
+				}
+
+				FVertexID VertexID = MeshDescription.CreateVertex();
+				VertexPositions[VertexID] = FDatasmithUtils::ConvertVector((FDatasmithUtils::EModelCoordSystem) ImportParams.ModelCoordSys, Vertex);
+				VertexPositions[VertexID] = SymmetricMatrix.TransformPosition(VertexPositions[VertexID]);
+				SymmetricVertexIds[VertexIndex++] = VertexID;
+			}
+		}
+	}
+
 	// PolygonAttributes name used into modeling tools (ExtendedMeshAttribute::PolyTriGroups)
 	const FName PolyTriGroups("PolyTriGroups");
 
@@ -495,7 +579,7 @@ namespace CADLibrary
 		}
 		else
 		{
-			//FillVertexPosition(ImportParams, MeshParameters, Body, MeshDescription);
+			FillVertexPosition(ImportParams, MeshParameters, Body, MeshDescription);
 		}
 		
 		if (!FillMesh(MeshParameters, ImportParams, Body, MeshDescription))
diff --git a/Engine/Plugins/Enterprise/DatasmithCADImporter/Source/DatasmithCADTranslator/Private/DatasmithCADTranslator.cpp b/Engine/Plugins/Enterprise/DatasmithCADImporter/Source/DatasmithCADTranslator/Private/DatasmithCADTranslator.cpp
--- a/Engine/Plugins/Enterprise/DatasmithCADImporter/Source/DatasmithCADTranslator/Private/DatasmithCADTranslator.cpp
+++ b/Engine/Plugins/Enterprise/DatasmithCADImporter/Source/DatasmithCADTranslator/Private/DatasmithCADTranslator.cpp
@@ -81,7 +81,7 @@ bool FDatasmithCADTranslator::LoadScene(TSharedRef<IDatasmithScene> DatasmithSce
 {
 	ImportParameters.bEnableKernelIOTessellation = (CVarStaticCADTranslatorEnableKernelIOTessellation.GetValueOnAnyThread() != 0);;
 #ifdef CADKERNEL_DEBUG
-	//ImportParameters.bEnableKernelIOTessellation = false;
+	ImportParameters.bEnableKernelIOTessellation = false;
 #endif
 
 	ImportParameters.MetricUnit = 0.001;
@@ -150,16 +150,16 @@ bool FDatasmithCADTranslator::LoadScene(TSharedRef<IDatasmithScene> DatasmithSce
 			DatasmithDispatcher::FDatasmithDispatcher Dispatcher(ImportParameters, CachePath, NumCores, CADFileToUEFileMap, CADFileToUEGeomMap);
 			Dispatcher.AddTask(FileDescription);
 
-			Dispatcher.Process(bWithProcessor);
-		}
+		Dispatcher.Process(bWithProcessor);
+	}
 
-		FDatasmithSceneGraphBuilder SceneGraphBuilder(CADFileToUEFileMap, CachePath, DatasmithScene, GetSource(), ImportParameters);
-		SceneGraphBuilder.Build();
+	FDatasmithSceneGraphBuilder SceneGraphBuilder(CADFileToUEFileMap, CachePath, DatasmithScene, GetSource(), ImportParameters);
+	SceneGraphBuilder.Build();
 
-		MeshBuilderPtr = MakeUnique<FDatasmithMeshBuilder>(CADFileToUEGeomMap, CachePath, ImportParameters);
+	MeshBuilderPtr = MakeUnique<FDatasmithMeshBuilder>(CADFileToUEGeomMap, CachePath, ImportParameters);
 
-		return true;
-	}
+	return true;
+}
 
 	ImportParameters.bEnableCacheUsage = false;
 
diff --git a/Engine/Source/Runtime/Datasmith/CADKernel/Public/CADKernel/Core/MetadataDictionary.h b/Engine/Source/Runtime/Datasmith/CADKernel/Public/CADKernel/Core/MetadataDictionary.h
--- a/Engine/Source/Runtime/Datasmith/CADKernel/Public/CADKernel/Core/MetadataDictionary.h
+++ b/Engine/Source/Runtime/Datasmith/CADKernel/Public/CADKernel/Core/MetadataDictionary.h
@@ -13,9 +13,10 @@ namespace CADKernel
 
 		FString Name;
 		int32 HostId = 0;
+		int32 PatchId = 0;
+ 		int32 LayerId = -1; 
 		uint32 ColorId = 0;
 		uint32 MaterialId = 0;
-		int32 LayerId = -1; 
 
 	public:
 
@@ -41,10 +42,30 @@ namespace CADKernel
 			ColorId = InColorId;
 		}
 
+		uint32 GetColorId()
+		{
+			return ColorId;
+		}
+
 		void SetMaterialId(const uint32& InMaterialId)
 		{
 			MaterialId = InMaterialId;
 		}
+
+		uint32 GetMaterialId()
+		{
+			return MaterialId;
+		}
+
+		void SetPatchId(int32 InPatchId)
+		{
+			PatchId = InPatchId;
+		}
+
+		int32 GetPatchId()
+		{
+			return PatchId;
+		}
 	};
 }
 
