Fix Issue 75454: Compile errors in CoreTechBridge.h
As CoreTechFileReader, CoreTechBridge can only be used if USE_KERNEL_IO_SDK is set

diff --git a/Engine/Plugins/Enterprise/DatasmithCADImporter/Source/CADInterfaces/Private/CoreTechBridge.cpp b/Engine/Plugins/Enterprise/DatasmithCADImporter/Source/CADInterfaces/Private/CoreTechBridge.cpp
--- a/Engine/Plugins/Enterprise/DatasmithCADImporter/Source/CADInterfaces/Private/CoreTechBridge.cpp
+++ b/Engine/Plugins/Enterprise/DatasmithCADImporter/Source/CADInterfaces/Private/CoreTechBridge.cpp
@@ -2,6 +2,8 @@
 
 #include "CoreTechBridge.h"
 
+#ifdef USE_KERNEL_IO_SDK
+
 #pragma warning(push)
 #pragma warning(disable:4996) // unsafe sprintf
 #pragma warning(disable:4828) // illegal character
@@ -1108,4 +1110,6 @@ void FCoreTechBridge::AddMetadata(CT_OBJECT_ID CTNodeId, TSharedRef<FMetadataDic
 			break;
 		}
 	}
-}
\ No newline at end of file
+}
+
+#endif // USE_KERNEL_IO_SDK
\ No newline at end of file
diff --git a/Engine/Plugins/Enterprise/DatasmithCADImporter/Source/CADInterfaces/Private/CoreTechBridge.h b/Engine/Plugins/Enterprise/DatasmithCADImporter/Source/CADInterfaces/Private/CoreTechBridge.h
--- a/Engine/Plugins/Enterprise/DatasmithCADImporter/Source/CADInterfaces/Private/CoreTechBridge.h
+++ b/Engine/Plugins/Enterprise/DatasmithCADImporter/Source/CADInterfaces/Private/CoreTechBridge.h
@@ -2,6 +2,8 @@
 #pragma once
 #include "CADKernel/Core/Types.h"
 
+#ifdef USE_KERNEL_IO_SDK 
+
 #ifdef CADKERNEL_DEV
 #include "CADKernel/Geo/Surfaces/Surface.h"
 #endif
@@ -110,3 +112,5 @@ namespace CADKernel
 
 	};
 }
+
+#endif // USE_KERNEL_IO_SDK
