Fix Issue 75481: Compile warnings

diff --git a/Engine/Plugins/Enterprise/DatasmithCADImporter/Source/CADInterfaces/Private/CoreTechBridge.cpp b/Engine/Plugins/Enterprise/DatasmithCADImporter/Source/CADInterfaces/Private/CoreTechBridge.cpp
--- a/Engine/Plugins/Enterprise/DatasmithCADImporter/Source/CADInterfaces/Private/CoreTechBridge.cpp
+++ b/Engine/Plugins/Enterprise/DatasmithCADImporter/Source/CADInterfaces/Private/CoreTechBridge.cpp
@@ -882,7 +882,7 @@ TSharedPtr<FCurve> FCoreTechBridge::AddCurveOnSurface(CT_OBJECT_ID CTCurveId)
 	CT_OBJECT_TYPE CTSurfaceType;
 	CT_OBJECT_IO::AskType(CTBaseSurfaceId, CTSurfaceType);
 
-	if (CTSurfaceType == CT_S_RULED_TYPE && CTSurfaceType == CT_CYLINDER_TYPE)
+	if (CTSurfaceType == CT_S_RULED_TYPE || CTSurfaceType == CT_CYLINDER_TYPE)
 	{
 		ensureCADKernel(false);
 		FMessage::Printf(Debug, TEXT("case : ruled surface with curve on cylinder surface, curveId %d\n"), CTCurveId);
