Datasmith : add AP242 Xml Step, XPDM, SolidEdge formats

diff --git a/Engine/Plugins/Enterprise/DatasmithCADImporter/Source/DatasmithCADTranslator/Private/DatasmithCADTranslator.cpp b/Engine/Plugins/Enterprise/DatasmithCADImporter/Source/DatasmithCADTranslator/Private/DatasmithCADTranslator.cpp
--- a/Engine/Plugins/Enterprise/DatasmithCADImporter/Source/DatasmithCADTranslator/Private/DatasmithCADTranslator.cpp
+++ b/Engine/Plugins/Enterprise/DatasmithCADImporter/Source/DatasmithCADTranslator/Private/DatasmithCADTranslator.cpp
@@ -64,13 +64,17 @@ void FDatasmithCADTranslator::Initialize(FDatasmithTranslatorCapabilities& OutCa
 
 	OutCapabilities.SupportedFileFormats.Add(FFileFormatInfo{ TEXT("step"), TEXT("Step files") });
 	OutCapabilities.SupportedFileFormats.Add(FFileFormatInfo{ TEXT("stp"), TEXT("Step files") });
+	OutCapabilities.SupportedFileFormats.Add(FFileFormatInfo{ TEXT("xml"), TEXT("AP242 Xml Step files, XPDM files") });
 
 	OutCapabilities.SupportedFileFormats.Add(FFileFormatInfo{ TEXT("x_t"), TEXT("Parasolid files (Text format)") });
 	OutCapabilities.SupportedFileFormats.Add(FFileFormatInfo{ TEXT("x_b"), TEXT("Parasolid files (Binary format)") });
 
-	OutCapabilities.SupportedFileFormats.Add(FFileFormatInfo{ TEXT("asm"), TEXT("Unigraphics Assembly, NX files") });
+	OutCapabilities.SupportedFileFormats.Add(FFileFormatInfo{ TEXT("asm"), TEXT("Unigraphics, NX, SolidEdge Assembly files") });
 	OutCapabilities.SupportedFileFormats.Add(FFileFormatInfo{ TEXT("prt"), TEXT("Unigraphics, NX Part files") });
 
+	OutCapabilities.SupportedFileFormats.Add(FFileFormatInfo{ TEXT("par"), TEXT("SolidEdge Part files") });
+	OutCapabilities.SupportedFileFormats.Add(FFileFormatInfo{ TEXT("psm"), TEXT("SolidEdge Part files") });
+
 	OutCapabilities.SupportedFileFormats.Add(FFileFormatInfo{ TEXT("dwg"), TEXT("AutoCAD, Model files") });
 	OutCapabilities.SupportedFileFormats.Add(FFileFormatInfo{ TEXT("dgn"), TEXT("MicroStation files") });
 
diff --git a/Engine/Plugins/Enterprise/DatasmithCADImporter/Source/DatasmithPLMXMLTranslator/Private/DatasmithPlmXmlTranslator.cpp b/Engine/Plugins/Enterprise/DatasmithCADImporter/Source/DatasmithPLMXMLTranslator/Private/DatasmithPlmXmlTranslator.cpp
--- a/Engine/Plugins/Enterprise/DatasmithCADImporter/Source/DatasmithPLMXMLTranslator/Private/DatasmithPlmXmlTranslator.cpp
+++ b/Engine/Plugins/Enterprise/DatasmithCADImporter/Source/DatasmithPLMXMLTranslator/Private/DatasmithPlmXmlTranslator.cpp
@@ -43,6 +43,16 @@ void FDatasmithPlmXmlTranslator::Initialize(FDatasmithTranslatorCapabilities& Ou
 	OutCapabilities.bIsEnabled = false;
 }
 
+bool FDatasmithPlmXmlTranslator::IsSourceSupported(const FDatasmithSceneSource& Source)
+{
+	if (Source.GetSourceFileExtension() != TEXT("xml"))
+	{
+		return true;
+	}
+
+	return Datasmith::CheckXMLFileSchema(Source.GetSourceFile(), TEXT("PLMXML"));
+}
+
 bool FDatasmithPlmXmlTranslator::LoadScene(TSharedRef<IDatasmithScene> OutScene)
 {
 	OutScene->SetHost(TEXT("PlmXmlTranslator"));
diff --git a/Engine/Plugins/Enterprise/DatasmithCADImporter/Source/DatasmithPLMXMLTranslator/Private/DatasmithPlmXmlTranslator.h b/Engine/Plugins/Enterprise/DatasmithCADImporter/Source/DatasmithPLMXMLTranslator/Private/DatasmithPlmXmlTranslator.h
--- a/Engine/Plugins/Enterprise/DatasmithCADImporter/Source/DatasmithPLMXMLTranslator/Private/DatasmithPlmXmlTranslator.h
+++ b/Engine/Plugins/Enterprise/DatasmithCADImporter/Source/DatasmithPLMXMLTranslator/Private/DatasmithPlmXmlTranslator.h
@@ -26,6 +26,7 @@ public:
 
 	virtual void GetSceneImportOptions(TArray<TStrongObjectPtr<UDatasmithOptionsBase>>& Options) override;
 	virtual void SetSceneImportOptions(TArray<TStrongObjectPtr<UDatasmithOptionsBase>>& Options) override;
+	virtual bool IsSourceSupported(const FDatasmithSceneSource& Source) override;
 	//~ End IDatasmithTranslator interface
 
 private:
diff --git a/Engine/Plugins/Enterprise/DatasmithCADImporter/Source/ParametricSurface/Private/ParametricSurfaceTranslator.cpp b/Engine/Plugins/Enterprise/DatasmithCADImporter/Source/ParametricSurface/Private/ParametricSurfaceTranslator.cpp
--- a/Engine/Plugins/Enterprise/DatasmithCADImporter/Source/ParametricSurface/Private/ParametricSurfaceTranslator.cpp
+++ b/Engine/Plugins/Enterprise/DatasmithCADImporter/Source/ParametricSurface/Private/ParametricSurfaceTranslator.cpp
@@ -30,6 +30,16 @@ void FParametricSurfaceTranslator::SetSceneImportOptions(TArray<TStrongObjectPtr
 	}
 }
 
+bool FParametricSurfaceTranslator::IsSourceSupported(const FDatasmithSceneSource& Source)
+{
+	if (Source.GetSourceFileExtension() != TEXT("xml"))
+	{
+		return true;
+	}
+
+	return Datasmith::CheckXMLFileSchema(Source.GetSourceFile(), TEXT("XPDMXML"), TEXT("ns3:Uos"));
+}
+
 
 
 
diff --git a/Engine/Plugins/Enterprise/DatasmithCADImporter/Source/ParametricSurface/Public/ParametricSurfaceTranslator.h b/Engine/Plugins/Enterprise/DatasmithCADImporter/Source/ParametricSurface/Public/ParametricSurfaceTranslator.h
--- a/Engine/Plugins/Enterprise/DatasmithCADImporter/Source/ParametricSurface/Public/ParametricSurfaceTranslator.h
+++ b/Engine/Plugins/Enterprise/DatasmithCADImporter/Source/ParametricSurface/Public/ParametricSurfaceTranslator.h
@@ -19,6 +19,7 @@ public:
 	// Begin IDatasmithTranslator overrides
 	virtual void GetSceneImportOptions(TArray<TStrongObjectPtr<UDatasmithOptionsBase>>& Options) override;
 	virtual void SetSceneImportOptions(TArray<TStrongObjectPtr<UDatasmithOptionsBase>>& Options) override;
+	virtual bool IsSourceSupported(const FDatasmithSceneSource& Source) override;
 	// End IDatasmithTranslator overrides
 
 protected:
diff --git a/Engine/Plugins/Enterprise/DatasmithImporter/Source/DatasmithTranslator/DatasmithTranslator.Build.cs b/Engine/Plugins/Enterprise/DatasmithImporter/Source/DatasmithTranslator/DatasmithTranslator.Build.cs
--- a/Engine/Plugins/Enterprise/DatasmithImporter/Source/DatasmithTranslator/DatasmithTranslator.Build.cs
+++ b/Engine/Plugins/Enterprise/DatasmithImporter/Source/DatasmithTranslator/DatasmithTranslator.Build.cs
@@ -13,6 +13,7 @@ namespace UnrealBuildTool.Rules
 					"CoreUObject",
 					"DatasmithCore",
 					"DatasmithContent",
+					"XmlParser",
 				}
 			);
 
diff --git a/Engine/Plugins/Enterprise/DatasmithImporter/Source/DatasmithTranslator/Private/DatasmithTranslator.cpp b/Engine/Plugins/Enterprise/DatasmithImporter/Source/DatasmithTranslator/Private/DatasmithTranslator.cpp
--- a/Engine/Plugins/Enterprise/DatasmithImporter/Source/DatasmithTranslator/Private/DatasmithTranslator.cpp
+++ b/Engine/Plugins/Enterprise/DatasmithImporter/Source/DatasmithTranslator/Private/DatasmithTranslator.cpp
@@ -3,6 +3,8 @@
 #include "DatasmithTranslator.h"
 #include "DatasmithTranslatorManager.h"
 
+#include "XmlFile.h"
+
 #define LOCTEXT_NAMESPACE "DatasmithTranslator"
 
 namespace Datasmith
@@ -18,6 +20,24 @@ namespace Datasmith
 		FDatasmithTranslatorManager::Get().Unregister(Info.TranslatorName);
 	}
 
+	FString GetXMLFileSchema(const FString& XmlFilePath)
+	{
+		FXmlFile XmlFile;
+
+		if (!XmlFile.LoadFile(XmlFilePath))
+		{
+			return FString();
+		}
+
+		const FXmlNode* RootNode = XmlFile.GetRootNode();
+		if (!RootNode)
+		{
+			return FString();
+		}
+		FString Tag = RootNode->GetTag();
+		return Tag;
+	}
+
 } // ns Datasmith
 
 
diff --git a/Engine/Plugins/Enterprise/DatasmithImporter/Source/DatasmithTranslator/Public/DatasmithTranslator.h b/Engine/Plugins/Enterprise/DatasmithImporter/Source/DatasmithTranslator/Public/DatasmithTranslator.h
--- a/Engine/Plugins/Enterprise/DatasmithImporter/Source/DatasmithTranslator/Public/DatasmithTranslator.h
+++ b/Engine/Plugins/Enterprise/DatasmithImporter/Source/DatasmithTranslator/Public/DatasmithTranslator.h
@@ -170,4 +170,32 @@ namespace Datasmith
 		TStrongObjectPtr<UOptionClass> Option(MakeOptionsPtr<UOptionClass>());
 		return Option;
 	}
+
+	FString DATASMITHTRANSLATOR_API GetXMLFileSchema(const FString& XmlFilePath);
+
+	template<typename ValueType>
+	bool CheckXMLSchema(const FString& FileXmlSchema, ValueType Format)
+	{
+		return (FileXmlSchema == Format);
+	}
+
+	template<typename ValueType, class... ValueTypes>
+	bool CheckXMLSchema(const FString& FileXmlSchema, ValueType Format, ValueTypes... OtherFormats)
+	{
+		if (FileXmlSchema == Format)
+		{
+			return true;
+		}
+		return CheckXMLSchema(FileXmlSchema, OtherFormats...);
+	}
+
+	/**
+	 * Check if the xml file schema is in the list of supported schema
+	 */
+	template<class... ValueTypes>
+	bool CheckXMLFileSchema(const FString& XmlFilePath, ValueTypes&... SupportedSchema)
+	{
+		FString FileXmlSchema = GetXMLFileSchema(XmlFilePath);
+		return CheckXMLSchema(FileXmlSchema, SupportedSchema...);
+	}
 }
