CADKernel Source Code v0.02
The review of the source code is in progress.

diff --git a/Engine/Plugins/Enterprise/DatasmithCADImporter/Source/CADInterfaces/Private/CoreTechBridge.cpp b/Engine/Plugins/Enterprise/DatasmithCADImporter/Source/CADInterfaces/Private/CoreTechBridge.cpp
--- a/Engine/Plugins/Enterprise/DatasmithCADImporter/Source/CADInterfaces/Private/CoreTechBridge.cpp
+++ b/Engine/Plugins/Enterprise/DatasmithCADImporter/Source/CADInterfaces/Private/CoreTechBridge.cpp
@@ -85,11 +85,23 @@ using namespace CADKernel;
 FCoreTechBridge::FCoreTechBridge(TSharedRef<FSession>& InSession)
 	: Session(InSession)
 	, GeometricTolerance(Session->GetGeometricTolerance())
+	, SquareGeometricTolerance(FMath::Square(Session->GetGeometricTolerance()))
+	, SquareJoiningVertexTolerance(SquareGeometricTolerance * 2)
 {
 }
 
 TSharedRef<FBody> FCoreTechBridge::AddBody(CT_OBJECT_ID CTBodyId)
 {
+	TSharedPtr<FEntity>* EntityPtr = CTIdToEntity.Find((uint32)CTBodyId);
+	if (EntityPtr)
+	{
+		TSharedPtr<FBody> Body = StaticCastSharedPtr<FBody>(*EntityPtr);
+		if (Body.IsValid())
+		{
+			return Body.ToSharedRef();
+		}
+	}
+
 	TSharedRef<FBody> Body = FEntity::MakeShared<FBody>();
 
 	AddMetadata(CTBodyId, Body);
@@ -100,6 +112,25 @@ TSharedRef<FBody> FCoreTechBridge::AddBody(CT_OBJECT_ID CTBodyId)
 
 	CTIdToEntity.Add((const uint32)CTBodyId, Body);
 
+#ifdef CORETECHBRIDGE_DEBUG_AABB
+	{
+		// Compute Body BBox based on CAD data
+		uint32 VerticesSize;
+		CT_BODY_IO::AskVerticesSizeArray(CTBodyId, VerticesSize);
+
+		FAABB BBox;
+		TArray<CT_COORDINATE> VerticesArray;
+		VerticesArray.SetNum(VerticesSize);
+		CT_BODY_IO::AskVerticesArray(CTBodyId, VerticesArray.GetData());
+		for (const CT_COORDINATE& Point : VerticesArray)
+		{
+			BBox += Point.xyz;
+		}
+		F3DDebugSession BBoxSession(FString::Printf(TEXT("CTBody %d"), CTBodyId));
+		DisplayAABB(BBox, (FIdent)CTBodyId);
+	}
+#endif
+
 	CT_LIST_IO CTShellIds;
 	CT_BODY_IO::AskShells(CTBodyId, CTShellIds);
 
@@ -136,25 +167,36 @@ void FCoreTechBridge::AddFace(CT_OBJECT_ID CTFaceId, TSharedRef<FShell>& Shell)
 	CT_OBJECT_ID CTSurfaceId;
 	CT_ORIENTATION CTOrientation;
 	CT_IO_ERROR Return = CT_FACE_IO::AskSurface(CTFaceId, CTSurfaceId, CTOrientation);
+
 	if (Return != IO_OK || CTSurfaceId <= 0)
 	{
-		FMessage::Printf(Log, TEXT("The CTFace %d has invalid carrier surface, this face is ignored"), CTFaceId);
+		FMessage::Printf(Log, TEXT("The CTFace %d has invalid carrier surface, this face is ignored\n"), CTFaceId);
 		return;
 	}
 
 	TSharedPtr<FSurface> SurfacePtr = AddSurface(CTSurfaceId, FaceBoundary);
 	if (!SurfacePtr.IsValid())
 	{
-		FMessage::Printf(Log, TEXT("The CTFace %d has invalid carrier surface, this face is ignored"), CTFaceId);
+		FMessage::Printf(Log, TEXT("The CTFace %d has invalid carrier surface, this face is ignored\n"), CTFaceId);
 		return;
 	}
 
 	TSharedRef<FSurface> Surface = SurfacePtr.ToSharedRef();
 
+#ifdef CADKERNEL_DEV
+	const FSurfacicTolerance& Tolerance = Surface->GetIsoTolerances();
+	if (Tolerance[0] < SMALL_NUMBER || Tolerance[1] < SMALL_NUMBER)
+	{
+		FMessage::Printf(Log, TEXT("KIO Surface %d (of Body) has surfacic tolerance probleme: %f %f\n"), CTFaceId, Tolerance[0], Tolerance[1]);
+		return;
+	}
+  	ensureCADKernel(Tolerance[0] > SMALL_NUMBER && Tolerance[1] > SMALL_NUMBER);
+#endif
+
 	const FSurfacicBoundary& SurfaceBounds = Surface->GetBoundary();
 	if (SurfaceBounds.IsDegenerated())
 	{
-		FMessage::Printf(Log, TEXT("The CTFace %d has degenerated carrier surface, this face is ignored ([%f, %f], [%f, %f])"), CTFaceId, SurfaceBounds.UVBoundaries[EIso::IsoU].Min, SurfaceBounds.UVBoundaries[EIso::IsoU].Max, SurfaceBounds.UVBoundaries[EIso::IsoV].Min, SurfaceBounds.UVBoundaries[EIso::IsoV].Max);
+		FMessage::Printf(Log, TEXT("The CTFace %d has degenerated carrier surface, this face is ignored ([%f, %f], [%f, %f])\n"), CTFaceId, SurfaceBounds.UVBoundaries[EIso::IsoU].Min, SurfaceBounds.UVBoundaries[EIso::IsoU].Max, SurfaceBounds.UVBoundaries[EIso::IsoV].Min, SurfaceBounds.UVBoundaries[EIso::IsoV].Max);
 		return;
 	}
 
@@ -162,7 +204,7 @@ void FCoreTechBridge::AddFace(CT_OBJECT_ID CTFaceId, TSharedRef<FShell>& Shell)
 	Return = CT_FACE_IO::AskLoops(CTFaceId, CTLoopIds);
 	if (Return != IO_OK)
 	{
-		FMessage::Printf(Log, TEXT("The CTFace %d has problem to get its loops, this face is ignored"), CTFaceId);
+		FMessage::Printf(Log, TEXT("The CTFace %d has problem to get its loops, this face is ignored\n"), CTFaceId);
 		return;
 	}
 
@@ -204,7 +246,7 @@ void FCoreTechBridge::AddFace(CT_OBJECT_ID CTFaceId, TSharedRef<FShell>& Shell)
 
 	if (Loops.Num() == 0)
 	{
-		FMessage::Printf(Log, TEXT("The CTFace %d is degenerate, this face is ignored"), CTFaceId);
+		FMessage::Printf(Log, TEXT("The CTFace %d is degenerate, this face is ignored\n"), CTFaceId);
 		return;
 	}
 
@@ -213,7 +255,7 @@ void FCoreTechBridge::AddFace(CT_OBJECT_ID CTFaceId, TSharedRef<FShell>& Shell)
 	Face->SetPatchId((int32) CTFaceId);
 
 #ifdef CORETECHBRIDGE_DEBUG
-	Face->CtKioId = (FIdent) CTFaceId;
+	Face->CtKioId = (FIdent)CTFaceId;
 #endif
 
 	Face->AddLoops(Loops);
@@ -260,7 +302,7 @@ TSharedPtr<FTopologicalLoop> FCoreTechBridge::AddLoop(CT_OBJECT_ID CTLoopId, TSh
 		return TSharedPtr<FTopologicalLoop>();
 	}
 
-	return FTopologicalLoop::Make(Edges, Directions);
+	return FTopologicalLoop::Make(Edges, Directions, GeometricTolerance);
 }
 
 TSharedPtr<FTopologicalEdge> FCoreTechBridge::AddEdge(CT_OBJECT_ID CTCoedgeId, TSharedRef<FSurface>& Surface)
@@ -315,11 +357,11 @@ TSharedPtr<FTopologicalEdge> FCoreTechBridge::AddEdge(CT_OBJECT_ID CTCoedgeId, T
 	TSharedPtr<FCurve> Curve;
 	if (PoleDim == 2)
 	{
-		Curve = FEntity::MakeShared<FNURBSCurve>(Surface->Get2DTolerance(), Degre, Knots, Poles, 2);
+		Curve = FEntity::MakeShared<FNURBSCurve>(Degre, Knots, Poles, 2);
 	}
 	else
 	{
-		Curve = FEntity::MakeShared<FNURBSCurve>(Surface->Get2DTolerance(), Degre, Knots, Poles, Weights, 2);
+		Curve = FEntity::MakeShared<FNURBSCurve>(Degre, Knots, Poles, Weights, 2);
 	}
 
 	if (!Curve.IsValid())
@@ -327,7 +369,7 @@ TSharedPtr<FTopologicalEdge> FCoreTechBridge::AddEdge(CT_OBJECT_ID CTCoedgeId, T
 		return TSharedPtr<FTopologicalEdge>();
 	}
 
-	TSharedRef<FRestrictionCurve> RestrictionCurve = FEntity::MakeShared<FRestrictionCurve>(GeometricTolerance, Surface, Curve.ToSharedRef());
+	TSharedRef<FRestrictionCurve> RestrictionCurve = FEntity::MakeShared<FRestrictionCurve>(Surface, Curve.ToSharedRef());
 	TSharedPtr<FTopologicalEdge> Edge = FTopologicalEdge::Make(RestrictionCurve);
 	if (!Edge.IsValid())
 	{
@@ -348,7 +390,7 @@ TSharedPtr<FTopologicalEdge> FCoreTechBridge::AddEdge(CT_OBJECT_ID CTCoedgeId, T
 		TSharedPtr<FEntity>* TwinEdge = CTIdToEntity.Find((uint32)CTConnectedCoedgeId);
 		if (TwinEdge != nullptr)
 		{
-			Edge->Link(StaticCastSharedRef<FTopologicalEdge>(TwinEdge->ToSharedRef()));
+			Edge->Link(StaticCastSharedRef<FTopologicalEdge>(TwinEdge->ToSharedRef()), SquareJoiningVertexTolerance);
 		}
 	}
 	return Edge;
@@ -369,7 +411,7 @@ TSharedPtr<FSurface> FCoreTechBridge::AddSurface(CT_OBJECT_ID CTSurfaceId, const
 	switch (CTSurfaceType)
 	{
 	case CT_PLANE_TYPE:
-		Surface = AddPlaneSurface(CTSurfaceId);
+		Surface = AddPlaneSurface(CTSurfaceId, Boundary);
 		break;
 	case CT_S_NURBS_TYPE:
 		Surface = AddNurbsSurface(CTSurfaceId);
@@ -546,7 +588,7 @@ TSharedPtr<FSurface> FCoreTechBridge::AddOffsetSurface(CT_OBJECT_ID CTSurfaceID,
 TSharedPtr<FSurface> FCoreTechBridge::AddLinearTransfoSurface(CT_OBJECT_ID CTSurfaceId)
 {
 	ensureCADKernel(false);
-	FMessage::Printf(Log, TEXT("The CTSurface %d is a linear transform surface, this face is ignored"), CTSurfaceId);
+	FMessage::Printf(Log, TEXT("The CTSurface %d is a linear transform surface, this face is ignored\n"), CTSurfaceId);
 	return TSharedPtr<FSurface>();
 
 	CT_OBJECT_ID CTBaseSurfaceId;
@@ -588,7 +630,7 @@ TSharedPtr<FSurface> FCoreTechBridge::AddConeSurface(CT_OBJECT_ID CTSurfaceId, c
 	return FEntity::MakeShared<FConeSurface>(GeometricTolerance, CoordinateSystem, Radius, HalfAngle, StartRuleLength, EndRuleLength, StartAngle, EndAngle);
 }
 
-TSharedPtr<FSurface> FCoreTechBridge::AddPlaneSurface(CT_OBJECT_ID CTSurfaceId)
+TSharedPtr<FSurface> FCoreTechBridge::AddPlaneSurface(CT_OBJECT_ID CTSurfaceId, const FSurfacicBoundary& InBoundary)
 {
 	CT_COORDINATE Origin;
 	CT_VECTOR Normal;
@@ -597,7 +639,7 @@ TSharedPtr<FSurface> FCoreTechBridge::AddPlaneSurface(CT_OBJECT_ID CTSurfaceId)
 
 	FMatrixH CoordinateSystem = CreateCoordinateSystem(Origin, Normal, UReference);
 
-	return FEntity::MakeShared<FPlaneSurface>(GeometricTolerance, CoordinateSystem);
+	return FEntity::MakeShared<FPlaneSurface>(GeometricTolerance, CoordinateSystem, InBoundary);
 }
 
 TSharedPtr<FSurface> FCoreTechBridge::AddCylinderSurface(CT_OBJECT_ID CTSurfaceId, const FSurfacicBoundary& Boundary)
@@ -638,7 +680,7 @@ TSharedPtr<FSurface> FCoreTechBridge::AddRevolutionSurface(CT_OBJECT_ID CTSurfac
 		Point1[i] = Origin.xyz[i];
 		Point2[i] = Origin.xyz[i] + Direction.xyz[i];
 	}
-	TSharedRef<FSegmentCurve> Axe = FEntity::MakeShared<FSegmentCurve>(GeometricTolerance, Point1, Point2);
+	TSharedRef<FSegmentCurve> Axe = FEntity::MakeShared<FSegmentCurve>(Point1, Point2);
 
 	double MinAngle = Boundary.IsValid() ? Boundary.UVBoundaries[EIso::IsoV].Min : 0.0;
 	double MaxAngle = Boundary.IsValid() ? Boundary.UVBoundaries[EIso::IsoV].Max : 2.0 * PI;
@@ -744,10 +786,10 @@ TSharedPtr<FCurve> FCoreTechBridge::AddNurbsCurve(CT_OBJECT_ID CTCurveId)
 		{
 			Weights[IPole] = RawPoles[Index];
 		}
-		return FEntity::MakeShared<FNURBSCurve>(GeometricTolerance, Degre, Knots, Poles, Weights);
+		return FEntity::MakeShared<FNURBSCurve>(Degre, Knots, Poles, Weights);
 	}
 
-	return FEntity::MakeShared<FNURBSCurve>(GeometricTolerance, Degre, Knots, Poles);
+	return FEntity::MakeShared<FNURBSCurve>(Degre, Knots, Poles);
 }
 
 TSharedPtr<FCurve> FCoreTechBridge::AddLineCurve(CT_OBJECT_ID CTCurveId, CT_OBJECT_ID CTSurfaceId)
@@ -783,7 +825,7 @@ TSharedPtr<FCurve> FCoreTechBridge::AddLineCurve(CT_OBJECT_ID CTCurveId, CT_OBJE
 		Point2[0] = (Point2[0] - TMin) / (TMax - TMin);
 	}
 
-	return FEntity::MakeShared<FSegmentCurve>(GeometricTolerance, Point1, Point2, 3);
+	return FEntity::MakeShared<FSegmentCurve>(Point1, Point2, 3);
 }
 
 TSharedPtr<FCurve> FCoreTechBridge::AddCompositeCurve(CT_OBJECT_ID CTCurveId)
@@ -806,12 +848,12 @@ TSharedPtr<FCurve> FCoreTechBridge::AddCompositeCurve(CT_OBJECT_ID CTCurveId)
 		}
 		else
 		{
-			FMessage::Printf(Log, TEXT("The CT Composite curve %d has an invalid curve, this curve is ignored"), CTCurveId);
+			FMessage::Printf(Log, TEXT("The CT Composite curve %d has an invalid curve, this curve is ignored\n"), CTCurveId);
 			return TSharedPtr<FCurve>();
 		}
 	}
 
-	return FEntity::MakeShared<FCompositeCurve>(GeometricTolerance, Curves);
+	return FEntity::MakeShared<FCompositeCurve>(Curves);
 }
 
 TSharedPtr<FCurve> FCoreTechBridge::AddCircleCurve(CT_OBJECT_ID CTCurveId)
@@ -830,7 +872,7 @@ TSharedPtr<FCurve> FCoreTechBridge::AddCircleCurve(CT_OBJECT_ID CTCurveId)
 	{
 		EndAngle = 2.0 * PI;
 	}
-	return FEntity::MakeShared<FEllipseCurve>(GeometricTolerance, CoordinateSystem, Radius, Radius, FLinearBoundary(StartAngle, EndAngle));
+	return FEntity::MakeShared<FEllipseCurve>(CoordinateSystem, Radius, Radius, FLinearBoundary(StartAngle, EndAngle));
 }
 
 TSharedPtr<FCurve> FCoreTechBridge::AddParabolaCurve(CT_OBJECT_ID CTCurveId)
@@ -846,7 +888,7 @@ TSharedPtr<FCurve> FCoreTechBridge::AddParabolaCurve(CT_OBJECT_ID CTCurveId)
 
 	FMatrixH CoordinateSystem = CreateCoordinateSystem(Origin, Direction, UReference);
 
-	return FEntity::MakeShared<FParabolaCurve>(GeometricTolerance, CoordinateSystem, Focal, FLinearBoundary(StartAlpha, EndAlpha), 3);
+	return FEntity::MakeShared<FParabolaCurve>(CoordinateSystem, Focal, FLinearBoundary(StartAlpha, EndAlpha), 3);
 }
 
 TSharedPtr<FCurve> FCoreTechBridge::AddEllipseCurve(CT_OBJECT_ID CTCurveId)
@@ -864,7 +906,7 @@ TSharedPtr<FCurve> FCoreTechBridge::AddEllipseCurve(CT_OBJECT_ID CTCurveId)
 		StartAlpha *= -1.;
 	}
 
-	return  FEntity::MakeShared<FEllipseCurve>(GeometricTolerance, CoordinateSystem, Radius1, Radius2, FLinearBoundary(StartAlpha, EndAlpha));
+	return  FEntity::MakeShared<FEllipseCurve>(CoordinateSystem, Radius1, Radius2, FLinearBoundary(StartAlpha, EndAlpha));
 }
 
 TSharedPtr<FCurve> FCoreTechBridge::AddCurveOnSurface(CT_OBJECT_ID CTCurveId)
@@ -924,7 +966,7 @@ TSharedPtr<FCurve> FCoreTechBridge::AddCurveOnSurface(CT_OBJECT_ID CTCurveId)
 	}
 #endif
 
-	return FEntity::MakeShared<FSurfacicCurve>(GeometricTolerance, Curve2D.ToSharedRef(), Surface.ToSharedRef());
+	return FEntity::MakeShared<FSurfacicCurve>(Curve2D.ToSharedRef(), Surface.ToSharedRef());
 }
 
 TSharedPtr<FCurve> FCoreTechBridge::AddHyperbolaCurve(CT_OBJECT_ID CTCurveId)
@@ -941,7 +983,7 @@ TSharedPtr<FCurve> FCoreTechBridge::AddHyperbolaCurve(CT_OBJECT_ID CTCurveId)
 
 	FMatrixH CoordinateSystem = CreateCoordinateSystem(Origin, Direction, UReference);
 
-	return FEntity::MakeShared<FHyperbolaCurve>(GeometricTolerance, CoordinateSystem, HalfAxis1, HalfAxis2, FLinearBoundary(StartAlpha, EndAlpha));
+	return FEntity::MakeShared<FHyperbolaCurve>(CoordinateSystem, HalfAxis1, HalfAxis2, FLinearBoundary(StartAlpha, EndAlpha));
 }
 
 FMatrixH FCoreTechBridge::CreateCoordinateSystem(const CT_COORDINATE& InOrigin, const CT_VECTOR& InDirection, const CT_VECTOR& InUReference)
@@ -1031,9 +1073,7 @@ void FCoreTechBridge::AddMetadata(CT_OBJECT_ID CTNodeId, TSharedRef<FMetadataDic
 	CT_UINT32 IthAttrib = 0;
 	while (CT_OBJECT_IO::SearchAttribute(CTNodeId, CT_ATTRIB_ALL, IthAttrib++) == IO_OK)
 	{
-		// Get the current attribute type
 		CT_ATTRIB_TYPE       AttributeType;
-
 		if (CT_CURRENT_ATTRIB_IO::AskAttributeType(AttributeType) != IO_OK)
 		{
 			continue;
@@ -1043,16 +1083,19 @@ void FCoreTechBridge::AddMetadata(CT_OBJECT_ID CTNodeId, TSharedRef<FMetadataDic
 		{
 		case CT_ATTRIB_NAME:
 		{
-			CT_STR NameStrValue;
-			if (CT_CURRENT_ATTRIB_IO::AskStrField(ITH_NAME_VALUE, NameStrValue) == IO_OK)
+			if (Entity->GetName() == nullptr) // otherwise the name has been set with ORIGINAL_NAME that is prefered in all case
 			{
-				Entity->SetName(AsFString(NameStrValue));
+				CT_STR NameStrValue;
+				if (CT_CURRENT_ATTRIB_IO::AskStrField(ITH_NAME_VALUE, NameStrValue) == IO_OK)
+				{
+					Entity->SetName(AsFString(NameStrValue));
+				}
+				break;
 			}
-			break;
 		}
 
 		case CT_ATTRIB_ORIGINAL_NAME:
-		{
+			{
 			CT_STR NameStrValue;
 			if (CT_CURRENT_ATTRIB_IO::AskStrField(ITH_NAME_VALUE, NameStrValue) == IO_OK)
 			{
@@ -1094,7 +1137,6 @@ void FCoreTechBridge::AddMetadata(CT_OBJECT_ID CTNodeId, TSharedRef<FMetadataDic
 			uint32 ColorHId = CADLibrary::BuildColorId(ColorId, Alpha);
 			Entity->SetColorId(ColorHId);
 		}
-		break;
 
 		case CT_ATTRIB_MATERIALID:
 		{
@@ -1108,8 +1150,8 @@ void FCoreTechBridge::AddMetadata(CT_OBJECT_ID CTNodeId, TSharedRef<FMetadataDic
 
 		default:
 			break;
-		}
-	}
-}
+			}
+			}
+			}
 
 #endif // USE_KERNEL_IO_SDK
\ No newline at end of file
diff --git a/Engine/Plugins/Enterprise/DatasmithCADImporter/Source/CADInterfaces/Private/CoreTechBridge.h b/Engine/Plugins/Enterprise/DatasmithCADImporter/Source/CADInterfaces/Private/CoreTechBridge.h
--- a/Engine/Plugins/Enterprise/DatasmithCADImporter/Source/CADInterfaces/Private/CoreTechBridge.h
+++ b/Engine/Plugins/Enterprise/DatasmithCADImporter/Source/CADInterfaces/Private/CoreTechBridge.h
@@ -40,6 +40,8 @@ namespace CADKernel
 
 		TSharedRef<FSession> Session;
 		const double GeometricTolerance;
+		const double SquareGeometricTolerance;
+		const double SquareJoiningVertexTolerance;
 
 		TMap<const uint32, TSharedPtr<FEntity>> CTIdToEntity;
 
@@ -61,15 +63,15 @@ namespace CADKernel
 		TSharedPtr<FTopologicalLoop> AddLoop(CT_OBJECT_ID CTLoopId, TSharedRef<FSurface>& Surface);
 		TSharedPtr<FTopologicalEdge> AddEdge(CT_OBJECT_ID CTCoedgeId, TSharedRef<FSurface>& Surface);
 
-		TSharedPtr<FSurface> AddSurface(CT_OBJECT_ID CTSurfaceId, const FSurfacicBoundary& OutBounds);
-		TSharedPtr<FSurface> AddPlaneSurface(CT_OBJECT_ID CTSurfaceId);
+		TSharedPtr<FSurface> AddSurface(CT_OBJECT_ID CTSurfaceId, const FSurfacicBoundary& OutBoundary);
+		TSharedPtr<FSurface> AddPlaneSurface(CT_OBJECT_ID CTSurfaceId, const FSurfacicBoundary& InBoundary);
 		TSharedPtr<FSurface> AddNurbsSurface(CT_OBJECT_ID CTSurfaceId);
-		TSharedPtr<FSurface> AddRevolutionSurface(CT_OBJECT_ID CTSurfaceId, const FSurfacicBoundary& Bounds);
-		TSharedPtr<FSurface> AddConeSurface(CT_OBJECT_ID CTSurfaceId, const FSurfacicBoundary& Bounds);
-		TSharedPtr<FSurface> AddOffsetSurface(CT_OBJECT_ID CTSurfaceId, const FSurfacicBoundary& Bounds);
-		TSharedPtr<FSurface> AddCylinderSurface(CT_OBJECT_ID CTSurfaceId, const FSurfacicBoundary& Bounds);
-		TSharedPtr<FSurface> AddSphereSurface(CT_OBJECT_ID CTSurfaceId, const FSurfacicBoundary& Bounds);
-		TSharedPtr<FSurface> AddTorusSurface(CT_OBJECT_ID CTSurfaceId, const FSurfacicBoundary& Bounds);
+		TSharedPtr<FSurface> AddRevolutionSurface(CT_OBJECT_ID CTSurfaceId, const FSurfacicBoundary& InBoundary);
+		TSharedPtr<FSurface> AddConeSurface(CT_OBJECT_ID CTSurfaceId, const FSurfacicBoundary& InBoundary);
+		TSharedPtr<FSurface> AddOffsetSurface(CT_OBJECT_ID CTSurfaceId, const FSurfacicBoundary& InBoundary);
+		TSharedPtr<FSurface> AddCylinderSurface(CT_OBJECT_ID CTSurfaceId, const FSurfacicBoundary& InBoundary);
+		TSharedPtr<FSurface> AddSphereSurface(CT_OBJECT_ID CTSurfaceId, const FSurfacicBoundary& InBoundary);
+		TSharedPtr<FSurface> AddTorusSurface(CT_OBJECT_ID CTSurfaceId, const FSurfacicBoundary& InBoundary);
 		TSharedPtr<FSurface> AddRuledSurface(CT_OBJECT_ID CTSurfaceId);
 		TSharedPtr<FSurface> AddLinearTransfoSurface(CT_OBJECT_ID CTSurfaceId);
 
@@ -83,7 +85,7 @@ namespace CADKernel
 		TSharedPtr<FCurve> AddNurbsCurve(CT_OBJECT_ID CTCurveId);
 		TSharedPtr<FCurve> AddParabolaCurve(CT_OBJECT_ID CTCurveId);
 
-		void Get2DCurvesRange(CT_OBJECT_ID CTFaceId, FSurfacicBoundary& OutBounds);
+		void Get2DCurvesRange(CT_OBJECT_ID CTFaceId, FSurfacicBoundary& OutBoundary);
 
 		FMatrixH CreateCoordinateSystem(const CT_COORDINATE& InOrigin, const CT_VECTOR& InDirection, const CT_VECTOR& InURef);
 
diff --git a/Engine/Source/Runtime/Datasmith/CADKernel/Private/CADKernel/Core/Database.cpp b/Engine/Source/Runtime/Datasmith/CADKernel/Private/CADKernel/Core/Database.cpp
--- a/Engine/Source/Runtime/Datasmith/CADKernel/Private/CADKernel/Core/Database.cpp
+++ b/Engine/Source/Runtime/Datasmith/CADKernel/Private/CADKernel/Core/Database.cpp
@@ -27,6 +27,8 @@ FDatabase::FDatabase()
 	DatabaseEntities.Reserve(DataBaseInitialSize);
 	AvailableIdents.Reserve(DataBaseInitialSize);
 	DatabaseEntities.Add(TSharedPtr<FEntity>());
+
+	AddEntity(Model.ToSharedRef());
 }
 
 TSharedRef<FModel> FDatabase::GetModel()
@@ -34,6 +36,7 @@ TSharedRef<FModel> FDatabase::GetModel()
 	if (!Model.IsValid())
 	{
 		Model = FEntity::MakeShared<FModel>();
+		AddEntity(Model.ToSharedRef());
 	}
 	return Model.ToSharedRef();
 }
@@ -52,6 +55,33 @@ FIdent FDatabase::CreateId()
 	}
 }
 
+void FDatabase::RemoveEntity(FEntity& Entity)
+{
+	FIdent EntityId = Entity.GetId();
+	if (EntityId == 0)
+	{
+		FMessage::Printf(Debug, TEXT("Warning in FDataBasse::RemoveEntity: the entity is not in the database.\n"), EntityId);
+	}
+	else if (EntityId >= (FIdent)DatabaseEntities.Num())
+	{
+		FMessage::Printf(Debug, TEXT("Warning in FDataBasse::RemoveEntity: the Id=%d is not yet defined, the entity has never exist.\n"), EntityId);
+	}
+	else if (DatabaseEntities[EntityId].IsValid())
+	{
+		if (DatabaseEntities[EntityId].Get() != &Entity)
+		{
+			FMessage::Printf(Debug, TEXT("Warning in FDataBasse::RemoveEntity: the entity with Id=%d is not in the database.\n"), EntityId);
+			return;
+		}
+		DatabaseEntities[EntityId] = TSharedPtr<FEntity>();
+		AvailableIdents.Add(EntityId);
+	}
+	else
+	{
+		FMessage::Printf(Debug, TEXT("Warning in FDataBasse::RemoveEntity: the entity with Id=%d is already deleted.\n"), EntityId);
+	}
+}
+
 void FDatabase::RemoveEntity(FIdent EntityId)
 {
 	if (EntityId >= (FIdent)DatabaseEntities.Num())
@@ -137,9 +167,30 @@ void FDatabase::SerializeSelection(FCADKernelArchive& Ar, const TArray<FIdent>&
 
 	EntitiesToBeSerialized.Empty();
 	NotYetSerialized.Empty();
+
+	// An archive must gather entity under a single Model. 
+	TSharedRef<FModel> TmpModel = FEntity::MakeShared<FModel>();
+	TmpModel->SpawnIdent(*this);
+	AddEntityToSave(TmpModel->GetId());
+
 	for (FIdent EntityId : SelectionIds)
 	{
-		AddEntityToSave(EntityId);
+		TSharedPtr<FEntity> Entity = GetEntity(EntityId);
+		switch(Entity->GetEntityType())
+		{
+		case EEntity::Model:
+			TmpModel->Copy(StaticCastSharedRef<FModel>(Entity.ToSharedRef()));
+			break;
+		case EEntity::Body:
+			TmpModel->Add(StaticCastSharedRef<FBody>(Entity.ToSharedRef()));
+			break;
+		case EEntity::TopologicalFace:
+			TmpModel->Add(StaticCastSharedRef<FTopologicalFace>(Entity.ToSharedRef()));
+			break;
+		default:
+			AddEntityToSave(EntityId);
+			break;
+		}
 	}
 
 	int32 Index = 0;
@@ -162,6 +213,8 @@ void FDatabase::SerializeSelection(FCADKernelArchive& Ar, const TArray<FIdent>&
 	NotYetSerialized.Empty();
 	bIsRecursiveSerialization = false;
 
+	RemoveEntity(TmpModel->GetId());
+
 	FMessage::Printf(Log, TEXT("End Serialisation of %d entities\n"), Index);
 }
 
@@ -209,8 +262,8 @@ void FDatabase::Deserialize(FCADKernelArchive& Ar)
 
 	int32 ArchiveSize = 0;
 	Ar << ArchiveSize;
-
 	DatabaseEntities.Reserve((int32)(DatabaseEntities.Num() + ArchiveSize * 1.5));
+	ArchiveSize += 10;
 	ArchiveIdToWaitingPointers.SetNum(ArchiveSize);
 	ArchiveIdToWaitingWeakPointers.SetNum(ArchiveSize);
 	ArchiveEntities.Init(TSharedPtr<FEntity>(), ArchiveSize);
@@ -242,6 +295,8 @@ void FDatabase::Deserialize(FCADKernelArchive& Ar)
 
 void FDatabase::CleanArchiveEntities()
 {
+	int32 FaceCount = 0;
+	int32 ShellCount = 0;
 	for (TSharedPtr<FEntity> Entity : ArchiveEntities)
 	{
 		if (Entity.IsValid())
@@ -258,9 +313,86 @@ void FDatabase::CleanArchiveEntities()
 				StaticCastSharedPtr<FVertexLink>(Entity)->CleanLink();
 				break;
 			}
+			case EEntity::Body:
+			{
+				TSharedPtr<FBody> Body = StaticCastSharedPtr<FBody>(Entity);
+				for (const TSharedPtr<FShell>& Shell : Body->GetShells())
+				{
+					Shell->SetMarker1();
+				}
+				Model->Add(StaticCastSharedPtr<FBody>(Entity));
+				break;
+			}
+			case EEntity::Shell:
+			{
+				TSharedPtr<FShell> Shell = StaticCastSharedPtr<FShell>(Entity);
+				for (const FOrientedFace& Face : Shell->GetFaces())
+				{
+					Face.Entity->SetMarker1();
+				}
+				ShellCount++;
+				break;
+			}
+			case EEntity::TopologicalFace:
+			{
+				FaceCount++;
+			}
+			default:
+				break;
+			}
+		}
+	}
+
+
+	TArray<TSharedPtr<FTopologicalFace>> IndependantFaces;
+	IndependantFaces.Reserve(FaceCount);
+	TArray<TSharedPtr<FBody>> NewBodies;
+	NewBodies.Reserve(ShellCount);
+
+	// find independent faces and shells
+	for (TSharedPtr<FEntity> Entity : ArchiveEntities)
+	{
+		if (Entity.IsValid())
+		{
+			switch (Entity->GetEntityType())
+			{
+			case EEntity::Shell:
+			{
+				TSharedPtr<FShell> Shell = StaticCastSharedPtr<FShell>(Entity);
+				if (Shell->HasMarker1())
+				{
+					Shell->ResetMarker1();
+				}
+				else
+				{
+					// for independent shell, make a new body to contain it
+					TSharedRef<FBody> Body = FEntity::MakeShared<FBody>();
+					Body->AddShell(Shell.ToSharedRef());
+					NewBodies.Emplace(Body);
+				}
+				break;
+			}
+			case EEntity::TopologicalFace:
+			{
+				TSharedPtr<FTopologicalFace> Face= StaticCastSharedPtr<FTopologicalFace>(Entity);
+				if (Face->HasMarker1())
+				{
+					Face->ResetMarker1();
+				}
+				else
+				{
+					IndependantFaces.Emplace(Face);
+				}
+				break;
+			}
+			default:
+				break;
 			}
 		}
 	}
+
+	Model->Append(IndependantFaces);
+	Model->Append(NewBodies);
 }
 
 void FDatabase::Empty()
@@ -357,7 +489,7 @@ void FDatabase::ExpandSelection(TSharedPtr<FEntity> Entity, const TSet<EEntity>&
 	case EEntity::Model:
 	{
 		TSharedPtr<FModel> InModel = StaticCastSharedPtr<FModel>(Entity);
-		for (const TSharedPtr<FBody>& Body : InModel->GetBodyList())
+		for (const TSharedPtr<FBody>& Body : InModel->GetBodies())
 		{
 			ExpandSelection(Body, Filter, Selection);
 		}
diff --git a/Engine/Source/Runtime/Datasmith/CADKernel/Private/CADKernel/Core/Entity.cpp b/Engine/Source/Runtime/Datasmith/CADKernel/Private/CADKernel/Core/Entity.cpp
--- a/Engine/Source/Runtime/Datasmith/CADKernel/Private/CADKernel/Core/Entity.cpp
+++ b/Engine/Source/Runtime/Datasmith/CADKernel/Private/CADKernel/Core/Entity.cpp
@@ -105,6 +105,7 @@ namespace CADKernel
 	}
 
 #ifdef CADKERNEL_DEV
+
 	void FEntity::InfoEntity() const
 	{
 		FInfoEntity Info;
@@ -120,13 +121,19 @@ namespace CADKernel
 		return Info;
 	}
 
-#ifdef CADKERNEL_DEV
 	FInfoEntity& FEntityGeom::GetInfo(FInfoEntity& Info) const
 	{
 		return FEntity::GetInfo(Info)
 			.Add(TEXT("Kio"), CtKioId);
 	}
-#endif
+
+	void FEntity::AddEntityInDatabase(TSharedRef<FEntity> Entity)
+	{
+		if (FSession::Session.IsValid())
+		{
+			FSession::Session->GetDatabase().AddEntity(Entity);
+		}
+	}
 
 #endif
 
diff --git a/Engine/Source/Runtime/Datasmith/CADKernel/Private/CADKernel/Core/Session.cpp b/Engine/Source/Runtime/Datasmith/CADKernel/Private/CADKernel/Core/Session.cpp
--- a/Engine/Source/Runtime/Datasmith/CADKernel/Private/CADKernel/Core/Session.cpp
+++ b/Engine/Source/Runtime/Datasmith/CADKernel/Private/CADKernel/Core/Session.cpp
@@ -1,10 +1,14 @@
 // Copyright Epic Games, Inc. All Rights Reserved.
 #include "CADKernel/Core/Session.h"
 
+#include "CADKernel/Topo/Model.h"
 #include "CADKernel/UI/Message.h"
 
 using namespace CADKernel;
 
+#ifdef CADKERNEL_DEV
+TSharedPtr<FSession> FSession::Session = TSharedPtr<FSession>();
+#endif
 
 void FSession::SaveDatabase(const TCHAR* FileName)
 {
@@ -50,5 +54,15 @@ void FSession::LoadDatabase(const TCHAR* FilePath)
 		FMessage::Printf(Log, TEXT("The archive file %s is corrupted\n"), FilePath);
 	}
 
+	TSharedRef<FModel> SessionModel = GetModel();
+
 	Database.Deserialize(*Archive.Get());
+	FModel* ArchiveModel = Archive->ArchiveModel;
+
+	if(ArchiveModel != nullptr)
+	{
+		ArchiveModel->Empty();
+		Database.RemoveEntity(*ArchiveModel);
+	}
+	Archive->Close();
 }
diff --git a/Engine/Source/Runtime/Datasmith/CADKernel/Private/CADKernel/Geo/Curves/BezierCurve.cpp b/Engine/Source/Runtime/Datasmith/CADKernel/Private/CADKernel/Geo/Curves/BezierCurve.cpp
--- a/Engine/Source/Runtime/Datasmith/CADKernel/Private/CADKernel/Geo/Curves/BezierCurve.cpp
+++ b/Engine/Source/Runtime/Datasmith/CADKernel/Private/CADKernel/Geo/Curves/BezierCurve.cpp
@@ -51,7 +51,7 @@ TSharedPtr<FEntityGeom> FBezierCurve::ApplyMatrix(const FMatrixH& InMatrix) cons
 		TransformedPoles.Emplace(InMatrix.Multiply(Pole));
 	}
 
-	return FEntity::MakeShared<FBezierCurve>(Tolerance, TransformedPoles);
+	return FEntity::MakeShared<FBezierCurve>(TransformedPoles);
 }
 
 #ifdef CADKERNEL_DEV
diff --git a/Engine/Source/Runtime/Datasmith/CADKernel/Private/CADKernel/Geo/Curves/BoundedCurve.cpp b/Engine/Source/Runtime/Datasmith/CADKernel/Private/CADKernel/Geo/Curves/BoundedCurve.cpp
--- a/Engine/Source/Runtime/Datasmith/CADKernel/Private/CADKernel/Geo/Curves/BoundedCurve.cpp
+++ b/Engine/Source/Runtime/Datasmith/CADKernel/Private/CADKernel/Geo/Curves/BoundedCurve.cpp
@@ -45,7 +45,7 @@ TSharedPtr<FCurve> FBoundedCurve::ReboundCurve(const FLinearBoundary& InBoundary
 		return FEntity::MakeShared<FBoundedCurve>(*this);
 	}
 
-	return FEntity::MakeShared<FBoundedCurve>(Tolerance, Curve.ToSharedRef(), NewBoundary, Dimension);
+	return FEntity::MakeShared<FBoundedCurve>(Curve.ToSharedRef(), NewBoundary, Dimension);
 }
 
 TSharedPtr<FEntityGeom> FBoundedCurve::ApplyMatrix(const FMatrixH& InMatrix) const
@@ -56,7 +56,7 @@ TSharedPtr<FEntityGeom> FBoundedCurve::ApplyMatrix(const FMatrixH& InMatrix) con
 		return TSharedPtr<FEntityGeom>();
 	}
 
-	return FEntity::MakeShared<FBoundedCurve>(Tolerance, TransformedCurve.ToSharedRef(), Boundary, Dimension);
+	return FEntity::MakeShared<FBoundedCurve>(TransformedCurve.ToSharedRef(), Boundary, Dimension);
 }
 
 #ifdef CADKERNEL_DEV
diff --git a/Engine/Source/Runtime/Datasmith/CADKernel/Private/CADKernel/Geo/Curves/CompositeCurve.cpp b/Engine/Source/Runtime/Datasmith/CADKernel/Private/CADKernel/Geo/Curves/CompositeCurve.cpp
--- a/Engine/Source/Runtime/Datasmith/CADKernel/Private/CADKernel/Geo/Curves/CompositeCurve.cpp
+++ b/Engine/Source/Runtime/Datasmith/CADKernel/Private/CADKernel/Geo/Curves/CompositeCurve.cpp
@@ -182,7 +182,7 @@ TSharedPtr<FEntityGeom> FCompositeCurve::ApplyMatrix(const FMatrixH& InMatrix) c
 		return TSharedPtr<FEntityGeom>();
 	}
 
-	return FEntity::MakeShared<FCompositeCurve>(Tolerance, TransformedCurves);
+	return FEntity::MakeShared<FCompositeCurve>(TransformedCurves);
 }
 
 #ifdef CADKERNEL_DEV
@@ -207,7 +207,7 @@ FInfoEntity& FCompositeCurve::GetInfo(FInfoEntity& Info) const
 // =========================================================================================================================================================================================================
 // =========================================================================================================================================================================================================
 
-FCompositeCurve::FCompositeCurve(const double InTolerance, const TArray<TSharedPtr<FCurve>>& CurveList, bool bDoInversions)
+FCompositeCurve::FCompositeCurve(const TArray<TSharedPtr<FCurve>>& CurveList, bool bDoInversions)
 {
 	ensureCADKernel(false);
 
diff --git a/Engine/Source/Runtime/Datasmith/CADKernel/Private/CADKernel/Geo/Curves/Curve.cpp b/Engine/Source/Runtime/Datasmith/CADKernel/Private/CADKernel/Geo/Curves/Curve.cpp
--- a/Engine/Source/Runtime/Datasmith/CADKernel/Private/CADKernel/Geo/Curves/Curve.cpp
+++ b/Engine/Source/Runtime/Datasmith/CADKernel/Private/CADKernel/Geo/Curves/Curve.cpp
@@ -13,31 +13,24 @@
 
 namespace CADKernel
 {
-	double FCurve::GetLength() const
+	double FCurve::GetLength(double Tolerance) const
 	{
 		if (!GlobalLength.IsValid())
 		{
-			GlobalLength = ComputeLength(Boundary);
-		}
-		return GlobalLength;
-	}
-
-	double FCurve::GetParametricTolerance() const
-	{
-		if (!Tolerance1D.IsValid())
-		{
-			double Tolerance2DMax = (GetUMax() - GetUMin()) * 0.01;
-			double Length = GetLength();
-			if (Length < SMALL_NUMBER)
+			switch (Dimension)
 			{
-				Tolerance1D = Tolerance2DMax;
-				return Tolerance1D;
+			case 3:
+				GlobalLength = ComputeLength(Boundary, Tolerance);
+				break;
+			case 2:
+				GlobalLength = ComputeLength2D(Boundary, Tolerance);
+				break;
+			default:
+				GlobalLength = 0;
+				break;
 			}
-
-			Tolerance1D = Tolerance * (GetUMax() - GetUMin()) / Length;
-			Tolerance1D = FMath::Max(Tolerance2DMax, (double)Tolerance1D);
 		}
-		return Tolerance1D;
+		return GlobalLength;
 	}
 
 	void FCurve::EvaluatePoints(const TArray<double>& Coordinates, TArray<FCurvePoint>& OutPoints, int32 DerivativeOrder) const
@@ -89,7 +82,7 @@ namespace CADKernel
 			.Add(TEXT("Curve type"), CurvesTypesNames[(uint8)GetCurveType()])
 			.Add(TEXT("Dimension"), (int32) Dimension)
 			.Add(TEXT("Boundary"), Boundary)
-			.Add(TEXT("Length"), GetLength());
+			.Add(TEXT("Length"), GetLength(0.01));
 	}
 #endif
 
@@ -127,10 +120,10 @@ namespace CADKernel
 			return TSharedPtr<FCurve>();
 		}
 
-		return FEntity::MakeShared<FBoundedCurve>(Tolerance, StaticCastSharedRef<FCurve>(AsShared()), NewBoundary, Dimension);
+		return FEntity::MakeShared<FBoundedCurve>(StaticCastSharedRef<FCurve>(AsShared()), NewBoundary, Dimension);
 	}
 
-	double FCurve::ComputeLength(const FLinearBoundary& InBoundary) const
+	double FCurve::ComputeLength(const FLinearBoundary& InBoundary, double Tolerance) const
 	{
 		FPolyline3D Polyline;
 		FCurveSamplerOnChord Sampler(*this, Boundary, Tolerance, Polyline);
@@ -138,7 +131,15 @@ namespace CADKernel
 		return Polyline.GetLength(Boundary);
 	}
 
-	void FCurve::Presample(const FLinearBoundary& InBoundary, TArray<double>& OutSampling) const
+	double FCurve::ComputeLength2D(const FLinearBoundary& InBoundary, double Tolerance) const
+	{
+		FPolyline2D Polyline;
+		FCurve2DSamplerOnChord Sampler(*this, Boundary, Tolerance, Polyline);
+		Sampler.Sample();
+		return Polyline.GetLength(Boundary);
+	}
+
+	void FCurve::Presample(const FLinearBoundary& InBoundary, double Tolerance, TArray<double>& OutSampling) const
 	{
 		FPolyline3D Presampling;
 		FCurveSamplerOnParam Sampler(*this, Boundary, Tolerance * 10., Tolerance, Presampling);
diff --git a/Engine/Source/Runtime/Datasmith/CADKernel/Private/CADKernel/Geo/Curves/EllipseCurve.cpp b/Engine/Source/Runtime/Datasmith/CADKernel/Private/CADKernel/Geo/Curves/EllipseCurve.cpp
--- a/Engine/Source/Runtime/Datasmith/CADKernel/Private/CADKernel/Geo/Curves/EllipseCurve.cpp
+++ b/Engine/Source/Runtime/Datasmith/CADKernel/Private/CADKernel/Geo/Curves/EllipseCurve.cpp
@@ -8,7 +8,7 @@ using namespace CADKernel;
 TSharedPtr<FEntityGeom> FEllipseCurve::ApplyMatrix(const FMatrixH& InMatrix) const
 {
 	FMatrixH NewMatrix = InMatrix * Matrix;
-	return FEntity::MakeShared<FEllipseCurve>(Tolerance, NewMatrix, RadiusU, RadiusV, Boundary);
+	return FEntity::MakeShared<FEllipseCurve>(NewMatrix, RadiusU, RadiusV, Boundary);
 }
 
 #ifdef CADKERNEL_DEV
diff --git a/Engine/Source/Runtime/Datasmith/CADKernel/Private/CADKernel/Geo/Curves/HyperbolaCurve.cpp b/Engine/Source/Runtime/Datasmith/CADKernel/Private/CADKernel/Geo/Curves/HyperbolaCurve.cpp
--- a/Engine/Source/Runtime/Datasmith/CADKernel/Private/CADKernel/Geo/Curves/HyperbolaCurve.cpp
+++ b/Engine/Source/Runtime/Datasmith/CADKernel/Private/CADKernel/Geo/Curves/HyperbolaCurve.cpp
@@ -7,7 +7,7 @@ using namespace CADKernel;
 TSharedPtr<FEntityGeom> FHyperbolaCurve::ApplyMatrix(const FMatrixH& InMatrix) const
 {
 	FMatrixH NewMatrix = InMatrix * Matrix;
-	return FEntity::MakeShared<FHyperbolaCurve>(Tolerance, NewMatrix, SemiMajorAxis, SemiImaginaryAxis, Boundary);
+	return FEntity::MakeShared<FHyperbolaCurve>(NewMatrix, SemiMajorAxis, SemiImaginaryAxis, Boundary);
 }
 
 #ifdef CADKERNEL_DEV
diff --git a/Engine/Source/Runtime/Datasmith/CADKernel/Private/CADKernel/Geo/Curves/NURBSCurve.cpp b/Engine/Source/Runtime/Datasmith/CADKernel/Private/CADKernel/Geo/Curves/NURBSCurve.cpp
--- a/Engine/Source/Runtime/Datasmith/CADKernel/Private/CADKernel/Geo/Curves/NURBSCurve.cpp
+++ b/Engine/Source/Runtime/Datasmith/CADKernel/Private/CADKernel/Geo/Curves/NURBSCurve.cpp
@@ -12,8 +12,8 @@
 
 using namespace CADKernel;
 
-FNURBSCurve::FNURBSCurve(const double InTolerance, int32 InDegre, const TArray<double>& InNodalVector, const TArray<FPoint>& InPoles, int8 InDimension)
-	: FCurve(InTolerance, InDimension)
+FNURBSCurve::FNURBSCurve(int32 InDegre, const TArray<double>& InNodalVector, const TArray<FPoint>& InPoles, int8 InDimension)
+	: FCurve(InDimension)
 	, Degree(InDegre)
 	, NodalVector(InNodalVector)
 	, Poles(InPoles)
@@ -23,8 +23,8 @@ FNURBSCurve::FNURBSCurve(const double InTolerance, int32 InDegre, const TArray<d
 	Finalize();
 }
 
-FNURBSCurve::FNURBSCurve(const double InTolerance, int32 InDegre, const TArray<double>& InNodalVector, const TArray<FPoint>& InPoles, const TArray<double>& InWeights, int8 InDimension)
-	: FCurve(InTolerance, InDimension)
+FNURBSCurve::FNURBSCurve(int32 InDegre, const TArray<double>& InNodalVector, const TArray<FPoint>& InPoles, const TArray<double>& InWeights, int8 InDimension)
+	: FCurve(InDimension)
 	, Degree(InDegre)
 	, NodalVector(InNodalVector)
 	, Weights(InWeights)
@@ -34,6 +34,11 @@ FNURBSCurve::FNURBSCurve(const double InTolerance, int32 InDegre, const TArray<d
 	Finalize();
 }
 
+FNURBSCurve::FNURBSCurve(const TSharedRef<FNURBSCurve>& Nurbs)
+	: FNURBSCurve(Nurbs->Degree, Nurbs->NodalVector, Nurbs->Poles, Nurbs->Weights, Nurbs->Dimension)
+{
+}
+
 TSharedPtr<FEntityGeom> FNURBSCurve::ApplyMatrix(const FMatrixH& InMatrix) const
 {
 	TArray<FPoint> TransformedPoles;
@@ -44,7 +49,7 @@ TSharedPtr<FEntityGeom> FNURBSCurve::ApplyMatrix(const FMatrixH& InMatrix) const
 		TransformedPoles.Emplace(InMatrix.Multiply(Pole));
 	}
 
-	return FEntity::MakeShared<FNURBSCurve>(Tolerance, Degree, NodalVector, TransformedPoles, Weights, Dimension);
+	return FEntity::MakeShared<FNURBSCurve>(Degree, NodalVector, TransformedPoles, Weights, Dimension);
 }
 
 #ifdef CADKERNEL_DEV
@@ -137,3 +142,53 @@ void FNURBSCurve::ExtendTo(const FPoint& Point)
 	PolylineTools::ExtendTo(Poles, Point);
 	Finalize();
 }
+
+
+void FNURBSCurve::SetStartNodalCoordinate(double NewStartBoundary)
+{
+	double Offset = NewStartBoundary - NodalVector[0];
+	for (double& NodalValue : NodalVector)
+	{
+		NodalValue += Offset;
+	}
+
+	Boundary.Set(NodalVector[Degree], NodalVector[NodalVector.Num() - 1 - Degree]);
+}
+
+void FNURBSCurve::Invert()
+{
+	{
+		TArray<FPoint> NewPoles;
+		NewPoles.Reserve(Poles.Num());
+		for (int32 Index = Poles.Num() - 1; Index >= 0; --Index)
+		{
+			NewPoles.Emplace(Poles[Index]);
+		}
+		Swap(Poles, NewPoles);
+	}
+
+	{
+		TArray<double> NewNodalVector;
+		NewNodalVector.Reserve(NodalVector.Num());
+		double LastNodalValue = NodalVector.Last();
+		for (int32 Index = NodalVector.Num() - 1; Index >= 0; --Index)
+		{
+			NewNodalVector.Emplace(LastNodalValue - NodalVector[Index]);
+		}
+		Swap(NodalVector, NewNodalVector);
+
+	}
+
+	if (IsRational())
+	{
+		TArray<double> NewWeights;
+		NewWeights.Reserve(Weights.Num());
+		for (int32 Index = Weights.Num() - 1; Index >= 0; --Index)
+		{
+			NewWeights.Emplace(Weights[Index]);
+		}
+		Swap(Weights, NewWeights);
+	}
+
+	Finalize();
+}
\ No newline at end of file
diff --git a/Engine/Source/Runtime/Datasmith/CADKernel/Private/CADKernel/Geo/Curves/ParabolaCurve.cpp b/Engine/Source/Runtime/Datasmith/CADKernel/Private/CADKernel/Geo/Curves/ParabolaCurve.cpp
--- a/Engine/Source/Runtime/Datasmith/CADKernel/Private/CADKernel/Geo/Curves/ParabolaCurve.cpp
+++ b/Engine/Source/Runtime/Datasmith/CADKernel/Private/CADKernel/Geo/Curves/ParabolaCurve.cpp
@@ -11,7 +11,7 @@ using namespace CADKernel;
 TSharedPtr<FEntityGeom> FParabolaCurve::ApplyMatrix(const FMatrixH& InMatrix) const
 {
 	FMatrixH NewMatrix = InMatrix * Matrix;
-	return FEntity::MakeShared<FParabolaCurve>(Tolerance, NewMatrix, FocalDistance, Boundary, Dimension);
+	return FEntity::MakeShared<FParabolaCurve>(NewMatrix, FocalDistance, Boundary, Dimension);
 }
 
 #ifdef CADKERNEL_DEV
diff --git a/Engine/Source/Runtime/Datasmith/CADKernel/Private/CADKernel/Geo/Curves/RestrictionCurve.cpp b/Engine/Source/Runtime/Datasmith/CADKernel/Private/CADKernel/Geo/Curves/RestrictionCurve.cpp
--- a/Engine/Source/Runtime/Datasmith/CADKernel/Private/CADKernel/Geo/Curves/RestrictionCurve.cpp
+++ b/Engine/Source/Runtime/Datasmith/CADKernel/Private/CADKernel/Geo/Curves/RestrictionCurve.cpp
@@ -16,15 +16,8 @@ FInfoEntity& FRestrictionCurve::GetInfo(FInfoEntity& Info) const
 }
 #endif
 
-void FRestrictionCurve::ExtendTo2DPoint(const FPoint2D& Point)
+void FRestrictionCurve::ExtendTo(const FPoint2D& Point)
 {
 	Curve2D->ExtendTo(Point);
-	if(Polyline.bWithTangent)
-	{
-		EvaluateSurfacicPolylineWithNormalAndTangent(Polyline);
-	}
-	else
-	{
-		EvaluateSurfacicPolyline(Polyline);
-	}
+	EvaluateSurfacicPolyline(Polyline);
 }
diff --git a/Engine/Source/Runtime/Datasmith/CADKernel/Private/CADKernel/Geo/Curves/SegmentCurve.cpp b/Engine/Source/Runtime/Datasmith/CADKernel/Private/CADKernel/Geo/Curves/SegmentCurve.cpp
--- a/Engine/Source/Runtime/Datasmith/CADKernel/Private/CADKernel/Geo/Curves/SegmentCurve.cpp
+++ b/Engine/Source/Runtime/Datasmith/CADKernel/Private/CADKernel/Geo/Curves/SegmentCurve.cpp
@@ -10,7 +10,7 @@ TSharedPtr<FEntityGeom> FSegmentCurve::ApplyMatrix(const FMatrixH& InMatrix) con
 	FPoint TransformedStartPoint = InMatrix.Multiply(StartPoint);
 	FPoint TransformedEndPoint = InMatrix.Multiply(EndPoint);
 
-	return FEntity::MakeShared<FSegmentCurve>(Tolerance, TransformedStartPoint, TransformedEndPoint, Dimension);
+	return FEntity::MakeShared<FSegmentCurve>(TransformedStartPoint, TransformedEndPoint, Dimension);
 }
 
 #ifdef CADKERNEL_DEV
diff --git a/Engine/Source/Runtime/Datasmith/CADKernel/Private/CADKernel/Geo/Curves/SurfacicCurve.cpp b/Engine/Source/Runtime/Datasmith/CADKernel/Private/CADKernel/Geo/Curves/SurfacicCurve.cpp
--- a/Engine/Source/Runtime/Datasmith/CADKernel/Private/CADKernel/Geo/Curves/SurfacicCurve.cpp
+++ b/Engine/Source/Runtime/Datasmith/CADKernel/Private/CADKernel/Geo/Curves/SurfacicCurve.cpp
@@ -40,16 +40,18 @@ void FSurfacicCurve::EvaluatePoints(const TArray<double>& InCoordinates, TArray<
 void FSurfacicCurve::EvaluateSurfacicPolylineWithNormalAndTangent(FSurfacicPolyline& OutPolyline) const
 {
 	OutPolyline.bWithNormals = true;
-	OutPolyline.bWithTangent = true;
 
 	TArray<FCurvePoint2D> Points2D;
 	Curve2D->Evaluate2DPoints(OutPolyline.Coordinates, Points2D, 1);
 	CarrierSurface->EvaluatePoints(Points2D, OutPolyline);
 }
 
-void FSurfacicCurve::EvaluateSurfacicPolyline(FSurfacicPolyline& OutPolyline, bool bInWithNormals) const
+void FSurfacicCurve::EvaluateSurfacicPolyline(FSurfacicPolyline& OutPolyline) const
 {
-	OutPolyline.bWithNormals = bInWithNormals;
+	if (OutPolyline.bWithTangent)
+	{
+		return EvaluateSurfacicPolylineWithNormalAndTangent(OutPolyline);
+	}
 
 	Curve2D->Evaluate2DPoints(OutPolyline.Coordinates, OutPolyline.Points2D);
 	CarrierSurface->EvaluatePoints(OutPolyline);
@@ -60,7 +62,7 @@ TSharedPtr<FEntityGeom> FSurfacicCurve::ApplyMatrix(const FMatrixH& InMatrix) co
 	TSharedPtr<FSurface> transformedSurface = StaticCastSharedPtr<FSurface>(CarrierSurface->ApplyMatrix(InMatrix));
 	if (transformedSurface.IsValid())
 	{
-		return FEntity::MakeShared<FSurfacicCurve>(Tolerance, Curve2D.ToSharedRef(), transformedSurface.ToSharedRef());
+		return FEntity::MakeShared<FSurfacicCurve>(Curve2D.ToSharedRef(), transformedSurface.ToSharedRef());
 	}
 	return TSharedPtr<FEntityGeom>();
 }
@@ -86,6 +88,7 @@ void FSurfacicCurve::FindNotDerivableCoordinates(const FLinearBoundary& InBounda
 		return;
 	}
 
+	double Tolerance = CarrierSurface->Get3DTolerance();
 	FSurfacicPolyline SurfacicPolyline(CarrierSurface.ToSharedRef(), Curve2D.ToSharedRef(), Tolerance, Tolerance, false, false);
 
 	FPolyline2D Polyline;
@@ -93,7 +96,7 @@ void FSurfacicCurve::FindNotDerivableCoordinates(const FLinearBoundary& InBounda
 	Swap(Polyline.Points, SurfacicPolyline.Points2D);
 
 	TArray<double> NotDerivableCurveCoordinates;
-	for (int32 Iso = 0; Iso <= EIso::UndefinedIso; ++Iso)
+	for (int32 Iso = 0; Iso < EIso::UndefinedIso; ++Iso)
 	{
 		for (const double& NotDerivableCoord : SurfaceNotDerivableCoordinates[(EIso) Iso])
 		{
@@ -101,5 +104,7 @@ void FSurfacicCurve::FindNotDerivableCoordinates(const FLinearBoundary& InBounda
 		}
 	}
 	Algo::Sort(NotDerivableCurveCoordinates);
-	ArrayUtils::InsertInside(OutNotDerivableCoordinates, NotDerivableCurveCoordinates, GetParametricTolerance());
+
+	double ToleranceGeoEdge = GetBoundary().ComputeMinimalTolerance();
+	ArrayUtils::Complete(OutNotDerivableCoordinates, NotDerivableCurveCoordinates, ToleranceGeoEdge);
 }
diff --git a/Engine/Source/Runtime/Datasmith/CADKernel/Private/CADKernel/Geo/Sampling/SurfacicPolyline.cpp b/Engine/Source/Runtime/Datasmith/CADKernel/Private/CADKernel/Geo/Sampling/SurfacicPolyline.cpp
--- a/Engine/Source/Runtime/Datasmith/CADKernel/Private/CADKernel/Geo/Sampling/SurfacicPolyline.cpp
+++ b/Engine/Source/Runtime/Datasmith/CADKernel/Private/CADKernel/Geo/Sampling/SurfacicPolyline.cpp
@@ -7,6 +7,7 @@
 #include "CADKernel/Geo/Curves/Curve.h"
 #include "CADKernel/Geo/Surfaces/Surface.h"
 #include "CADKernel/Geo/Sampler/SamplerOnParam.h"
+#include "CADKernel/Geo/Sampling/PolylineTools.h"
 
 using namespace CADKernel;
 
@@ -18,6 +19,11 @@ FSurfacicPolyline::FSurfacicPolyline(TSharedRef<FSurface> InCarrierSurface, TSha
 	Sampler.Sample();
 }
 
+FSurfacicPolyline::FSurfacicPolyline(TSharedRef<FSurface> InCarrierSurface, TSharedRef<FCurve> Curve2D)
+	: FSurfacicPolyline(InCarrierSurface, Curve2D, InCarrierSurface->Get3DTolerance())
+{
+}
+
 FSurfacicPolyline::FSurfacicPolyline(TSharedRef<FSurface> InCarrierSurface, TSharedRef<FCurve> Curve2D, const double ChordTolerance, const double ParamTolerance, bool bInWithNormals, bool bInWithTangents)
 	: bWithNormals(bInWithNormals)
 	, bWithTangent(bInWithTangents)
@@ -26,30 +32,56 @@ FSurfacicPolyline::FSurfacicPolyline(TSharedRef<FSurface> InCarrierSurface, TSha
 	Sampler.Sample();
 }
 
-void FSurfacicPolyline::CheckIfDegenerated(const double Tolerance3D, const double Tolerance2D, const FLinearBoundary& Boudary, bool& bDegeneration2D, bool& bDegeneration3D, double& Length3D) const
+void FSurfacicPolyline::CheckIfDegenerated(const double Tolerance3D, const FSurfacicTolerance& ToleranceIso, const FLinearBoundary& Boudary, bool& bDegeneration2D, bool& bDegeneration3D, double& Length3D) const
 {
-	TPolylineApproximator<FPoint2D> Approximator(Coordinates, Points2D);
-	double Length2D = Approximator.ComputeLengthOfSubPolyline(Boudary);
-	if (FMath::IsNearlyZero(Length2D, FMath::Max((double)KINDA_SMALL_NUMBER, Tolerance2D)))
+	TPolylineApproximator<FPoint> Approximator(Coordinates, Points3D);
+	int32 BoundaryIndices[2];
+	Approximator.GetStartEndIndex(Boudary, BoundaryIndices);
+
+	TPolylineApproximator<FPoint> Approximator3D(Coordinates, Points3D);
+	Length3D = Approximator3D.ComputeLengthOfSubPolyline(BoundaryIndices, Boudary);
+
+	if (!FMath::IsNearlyZero(Length3D, Tolerance3D))
 	{
-		bDegeneration2D = true;
-		bDegeneration3D = true;
-		Length3D = 0.;
+		bDegeneration3D = false;
+		bDegeneration2D = false;
 		return;
 	}
 
-	bDegeneration2D = false;
+	bDegeneration3D = true;
+	Length3D = 0.;
 
-	TPolylineApproximator<FPoint> Approximator3D(Coordinates, Points3D);
-	Length3D = Approximator3D.ComputeLengthOfSubPolyline(Boudary);
+	// Tolerance along Iso U/V is very costly to compute and not accurate.
+	// To test if a curve is degenerated, its 2d bounding box is compute and its compare to the surface boundary along U and along V
+	// Indeed, defining a Tolerance2D has no sense has the boundary length along an Iso could be very very huge compare to the boundary length along the other Iso like [[0, 1000] [0, 1]]
+	// The tolerance along ans iso is the length of the boundary along this iso divide by 100 000: if the curve length in 3d is 10m, the tolerance is 0.01mm
 
-	if (FMath::IsNearlyZero(Length3D, FMath::Max((double)KINDA_SMALL_NUMBER, Tolerance3D)))
+	FAABB2D Aabb;
+	TPolylineApproximator<FPoint2D> Approximator2D(Coordinates, Points2D);
+	Approximator2D.ComputeBoundingBox(BoundaryIndices, Boudary, Aabb);
+
+	bDegeneration2D = (Aabb.GetSize(0) < ToleranceIso[EIso::IsoU] && Aabb.GetSize(1) < ToleranceIso[EIso::IsoV]);
+}
+
+void FSurfacicPolyline::GetExtremities(const FLinearBoundary& InBoundary, const double Tolerance3D, const FSurfacicTolerance& MinToleranceIso, FSurfacicCurveExtremity& Extremities) const
+{
+	FDichotomyFinder Finder(Coordinates);
+	const int32 StartIndex = Finder.Find(InBoundary.Min);
+	const int32 EndIndex = Finder.Find(InBoundary.Max);
+
+	Extremities[0].Point2D = PolylineTools::ComputePoint(Coordinates, Points2D, StartIndex, InBoundary.Min);
+	Extremities[0].Point = PolylineTools::ComputePoint(Coordinates, Points3D, StartIndex, InBoundary.Min);
+	Extremities[0].Tolerance = ComputeTolerance(Tolerance3D, MinToleranceIso, StartIndex);
+
+	Extremities[1].Point2D = PolylineTools::ComputePoint(Coordinates, Points2D, EndIndex, InBoundary.Max);
+	Extremities[1].Point = PolylineTools::ComputePoint(Coordinates, Points3D, EndIndex, InBoundary.Max);
+	if (EndIndex == StartIndex)
 	{
-		bDegeneration3D = true;
+		Extremities[1].Tolerance = Extremities[0].Tolerance;
 	}
 	else
 	{
-		bDegeneration3D = false;
+		Extremities[1].Tolerance = ComputeTolerance(Tolerance3D, MinToleranceIso, EndIndex);
 	}
 }
 
diff --git a/Engine/Source/Runtime/Datasmith/CADKernel/Private/CADKernel/Geo/Surfaces/CompositeSurface.cpp b/Engine/Source/Runtime/Datasmith/CADKernel/Private/CADKernel/Geo/Surfaces/CompositeSurface.cpp
--- a/Engine/Source/Runtime/Datasmith/CADKernel/Private/CADKernel/Geo/Surfaces/CompositeSurface.cpp
+++ b/Engine/Source/Runtime/Datasmith/CADKernel/Private/CADKernel/Geo/Surfaces/CompositeSurface.cpp
@@ -23,6 +23,7 @@ FCompositeSurface::FCompositeSurface(const double InToleranceGeometric, int32 US
 		const FSurfacicBoundary& Bounds = Surface->GetBoundary();
 		NativeUVBoundaries.Emplace(Bounds);
 	}
+	SetMinToleranceIso();
 }
 
 void FCompositeSurface::InitBoundary() const
diff --git a/Engine/Source/Runtime/Datasmith/CADKernel/Private/CADKernel/Geo/Surfaces/ConeSurface.cpp b/Engine/Source/Runtime/Datasmith/CADKernel/Private/CADKernel/Geo/Surfaces/ConeSurface.cpp
--- a/Engine/Source/Runtime/Datasmith/CADKernel/Private/CADKernel/Geo/Surfaces/ConeSurface.cpp
+++ b/Engine/Source/Runtime/Datasmith/CADKernel/Private/CADKernel/Geo/Surfaces/ConeSurface.cpp
@@ -115,9 +115,9 @@ FInfoEntity& FConeSurface::GetInfo(FInfoEntity& Info) const
 		.Add(TEXT("Matrix"), Matrix)
 		.Add(TEXT("StartRadius"), StartRadius)
 		.Add(TEXT("ConeAngle"), ConeAngle)
-		.Add(TEXT("StartRuleLength"), Boundary[EIso::IsoU].Min)
-		.Add(TEXT("EndRuleLength"), Boundary[EIso::IsoU].Max)
-		.Add(TEXT("StartAngle"), Boundary[EIso::IsoV].Min)
-		.Add(TEXT("EndAngle"), Boundary[EIso::IsoV].Max);
+		.Add(TEXT("StartAngle"), Boundary[EIso::IsoU].Min)
+		.Add(TEXT("EndAngle"), Boundary[EIso::IsoU].Max)
+		.Add(TEXT("StartRuleLength"), Boundary[EIso::IsoV].Min)
+		.Add(TEXT("EndRuleLength"), Boundary[EIso::IsoV].Max);
 }
 #endif
diff --git a/Engine/Source/Runtime/Datasmith/CADKernel/Private/CADKernel/Geo/Surfaces/CoonsSurface.cpp b/Engine/Source/Runtime/Datasmith/CADKernel/Private/CADKernel/Geo/Surfaces/CoonsSurface.cpp
--- a/Engine/Source/Runtime/Datasmith/CADKernel/Private/CADKernel/Geo/Surfaces/CoonsSurface.cpp
+++ b/Engine/Source/Runtime/Datasmith/CADKernel/Private/CADKernel/Geo/Surfaces/CoonsSurface.cpp
@@ -12,6 +12,7 @@ using namespace CADKernel;
 FCoonsSurface::FCoonsSurface(const double InToleranceGeometric, TSharedPtr<FCurve> InCurves[4])
 	: FCoonsSurface(InToleranceGeometric, InCurves[0], InCurves[1], InCurves[2], InCurves[3])
 {
+	SetMinToleranceIso();
 }
 
 FCoonsSurface::FCoonsSurface(const double InToleranceGeometric, TSharedPtr<FCurve> InCurve1, TSharedPtr<FCurve> InCurve2, TSharedPtr<FCurve> InCurve3, TSharedPtr<FCurve> InCurve4)
@@ -27,6 +28,8 @@ FCoonsSurface::FCoonsSurface(const double InToleranceGeometric, TSharedPtr<FCurv
 	Corners[1] = Curves[0]->EvaluatePoint(Curves[0]->GetUMax()) * 0.5 + Curves[3]->EvaluatePoint(Curves[3]->GetUMin()) * 0.5;
 	Corners[2] = Curves[1]->EvaluatePoint(Curves[1]->GetUMin()) * 0.5 + Curves[2]->EvaluatePoint(Curves[2]->GetUMax()) * 0.5;
 	Corners[3] = Curves[1]->EvaluatePoint(Curves[1]->GetUMax()) * 0.5 + Curves[3]->EvaluatePoint(Curves[3]->GetUMax()) * 0.5;
+
+	SetMinToleranceIso();
 }
 
 void FCoonsSurface::LinesNotDerivables(const FSurfacicBoundary& Bounds, int32 InDerivativeOrder, FCoordinateGrid& OutNotDerivableCoordinates) const
@@ -52,7 +55,6 @@ void FCoonsSurface::LinesNotDerivables(const FSurfacicBoundary& Bounds, int32 In
 		Algo::Sort(OutNotDerivableCoordinates[(EIso)Iso]);
 		ArrayUtils::RemoveDuplicates(OutNotDerivableCoordinates[(EIso)Iso], GetIsoTolerances()[Iso]);
 	}
-
 }
 
 void FCoonsSurface::EvaluatePoint(const FPoint2D& InPoint2D, FSurfacicPoint& OutPoint3D, int32 InDerivativeOrder) const
diff --git a/Engine/Source/Runtime/Datasmith/CADKernel/Private/CADKernel/Geo/Surfaces/CylinderSurface.cpp b/Engine/Source/Runtime/Datasmith/CADKernel/Private/CADKernel/Geo/Surfaces/CylinderSurface.cpp
--- a/Engine/Source/Runtime/Datasmith/CADKernel/Private/CADKernel/Geo/Surfaces/CylinderSurface.cpp
+++ b/Engine/Source/Runtime/Datasmith/CADKernel/Private/CADKernel/Geo/Surfaces/CylinderSurface.cpp
@@ -13,6 +13,7 @@ FCylinderSurface::FCylinderSurface(const double InToleranceGeometric, const FMat
 	, Matrix(INMatrix)
 	, Radius(InRadius)
 {
+	SetMinToleranceIso();
 }
 
 void FCylinderSurface::InitBoundary() const
diff --git a/Engine/Source/Runtime/Datasmith/CADKernel/Private/CADKernel/Geo/Surfaces/PlaneSurface.cpp b/Engine/Source/Runtime/Datasmith/CADKernel/Private/CADKernel/Geo/Surfaces/PlaneSurface.cpp
--- a/Engine/Source/Runtime/Datasmith/CADKernel/Private/CADKernel/Geo/Surfaces/PlaneSurface.cpp
+++ b/Engine/Source/Runtime/Datasmith/CADKernel/Private/CADKernel/Geo/Surfaces/PlaneSurface.cpp
@@ -10,8 +10,8 @@
 
 using namespace CADKernel;
 
-FPlaneSurface::FPlaneSurface(const double InToleranceGeometric, const FMatrixH& InMatrix)
-	: FSurface(InToleranceGeometric, -HUGE_VALUE, HUGE_VALUE, -HUGE_VALUE, HUGE_VALUE)
+FPlaneSurface::FPlaneSurface(const double InToleranceGeometric, const FMatrixH& InMatrix, const FSurfacicBoundary& InBoundary)
+	: FSurface(InToleranceGeometric, InBoundary)
 {
 	Matrix = InMatrix;
 
@@ -19,43 +19,18 @@ FPlaneSurface::FPlaneSurface(const double InToleranceGeometric, const FMatrixH&
 
 	InverseMatrix = Matrix;
 	InverseMatrix.Inverse();
+	SetMinToleranceIso();
 }
 
-FPlaneSurface::FPlaneSurface(const double InToleranceGeometric, const FPoint& Position, FPoint Normal)
-	: FSurface(InToleranceGeometric, -HUGE_VALUE, HUGE_VALUE, -HUGE_VALUE, HUGE_VALUE)
+FPlaneSurface::FPlaneSurface(const double InToleranceGeometric, const FPoint& InPosition, FPoint InNormal, const FSurfacicBoundary& InBoundary)
+	: FSurface(InToleranceGeometric, InBoundary)
 {
-	Normal.Normalize();
-	Matrix.FromAxisOrigin(Normal, Position);
+	InNormal.Normalize();
+	Matrix.FromAxisOrigin(InNormal, InPosition);
 
 	InverseMatrix = Matrix;
 	InverseMatrix.Inverse();
-}
-
-const void FPlaneSurface::ComputeIsoTolerances() const
-{
-	FPoint Origin = Matrix.Multiply(FPoint::ZeroPoint);
-
-	TFunction<double(FPoint)> ComputeTolerance = [&](FPoint Point2D) 
-	{
-		FPoint Point = Matrix.Multiply(Point2D);
-		double Length = Origin.Distance(Point);
-		if (Length > SMALL_NUMBER)
-		{
-			return Tolerance3D / Length;
-		}
-		else
-		{
-			return Tolerance3D;
-		}
-	};
-
-	FSurfacicTolerance& Tolerance = ToleranceIsos;
-
-	FPoint Point2DU { 1 , 0, 0 };
-	Tolerance[EIso::IsoU] = ComputeTolerance(Point2DU);
-
-	FPoint Point2DV{ 0, 1, 0 };
-	Tolerance[EIso::IsoV] = ComputeTolerance(Point2DV);
+	SetMinToleranceIso();
 }
 
 CADKernel::FPlane FPlaneSurface::GetPlane() const
@@ -176,7 +151,7 @@ void FPlaneSurface::ProjectPoints(const TArray<FPoint>& Points, TArray<FPoint>*
 TSharedPtr<FEntityGeom> FPlaneSurface::ApplyMatrix(const FMatrixH& InMatrix) const
 {
 	FMatrixH NewMatrix = InMatrix * Matrix;
-	return FEntity::MakeShared<FPlaneSurface>(Tolerance3D, NewMatrix);
+	return FEntity::MakeShared<FPlaneSurface>(Tolerance3D, NewMatrix, Boundary);
 }
 
 #ifdef CADKERNEL_DEV
diff --git a/Engine/Source/Runtime/Datasmith/CADKernel/Private/CADKernel/Geo/Surfaces/RevolutionSurface.cpp b/Engine/Source/Runtime/Datasmith/CADKernel/Private/CADKernel/Geo/Surfaces/RevolutionSurface.cpp
--- a/Engine/Source/Runtime/Datasmith/CADKernel/Private/CADKernel/Geo/Surfaces/RevolutionSurface.cpp
+++ b/Engine/Source/Runtime/Datasmith/CADKernel/Private/CADKernel/Geo/Surfaces/RevolutionSurface.cpp
@@ -20,6 +20,7 @@ FRevolutionSurface::FRevolutionSurface(const double InToleranceGeometric, TShare
 	// Compute rotation axis
 	RotationAxis = Axis->GetEndPoint() - Axis->GetStartPoint();
 	RotationAxis.Normalize();
+	SetMinToleranceIso();
 }
 
 void FRevolutionSurface::EvaluatePoint(const FPoint2D& InSurfacicCoordinate, FSurfacicPoint& OutPoint3D, int32 InDerivativeOrder) const
@@ -99,7 +100,7 @@ void FRevolutionSurface::EvaluatePointGrid(const FCoordinateGrid& Coordinates, F
 
 void FRevolutionSurface::Presample(const FSurfacicBoundary& InBoundaries, FCoordinateGrid& Coordinates)
 {
-	Generatrix->Presample(Coordinates[EIso::IsoU]);
+	Generatrix->Presample(Coordinates[EIso::IsoU], Tolerance3D);
 	PresampleIsoCircle(InBoundaries, Coordinates, EIso::IsoV);
 }
 
diff --git a/Engine/Source/Runtime/Datasmith/CADKernel/Private/CADKernel/Geo/Surfaces/RuledSurface.cpp b/Engine/Source/Runtime/Datasmith/CADKernel/Private/CADKernel/Geo/Surfaces/RuledSurface.cpp
--- a/Engine/Source/Runtime/Datasmith/CADKernel/Private/CADKernel/Geo/Surfaces/RuledSurface.cpp
+++ b/Engine/Source/Runtime/Datasmith/CADKernel/Private/CADKernel/Geo/Surfaces/RuledSurface.cpp
@@ -121,7 +121,7 @@ void FRuledSurface::Presample(const FSurfacicBoundary& InBoundaries, FCoordinate
 {
 	TFunction<void(int32, TArray<double>&)> PresampleCurve = [&](int32 CurveIndex, TArray<double>& Sample)
 	{
-		Curves[CurveIndex]->Presample(Sample);
+		Curves[CurveIndex]->Presample(Sample, Tolerance3D);
 
 		double ULength = Curves[CurveIndex]->GetUMax() - Curves[CurveIndex]->GetUMin();
 		double UMin = Curves[CurveIndex]->GetUMin();
@@ -137,7 +137,7 @@ void FRuledSurface::Presample(const FSurfacicBoundary& InBoundaries, FCoordinate
 	PresampleCurve(1, Curve1Sample);
 
 	// remove duplicated
-	ArrayUtils::InsertInside(Coordinates[EIso::IsoU], Curve1Sample, GetIsoTolerances()[IsoU]);
+	ArrayUtils::InsertInside(Coordinates[EIso::IsoU], Curve1Sample, GetIsoTolerance(IsoU));
 
 	Coordinates[EIso::IsoV].Empty(3);
 	Coordinates[EIso::IsoV].Add(InBoundaries.UVBoundaries[EIso::IsoV].Min);
diff --git a/Engine/Source/Runtime/Datasmith/CADKernel/Private/CADKernel/Geo/Surfaces/Surface.cpp b/Engine/Source/Runtime/Datasmith/CADKernel/Private/CADKernel/Geo/Surfaces/Surface.cpp
--- a/Engine/Source/Runtime/Datasmith/CADKernel/Private/CADKernel/Geo/Surfaces/Surface.cpp
+++ b/Engine/Source/Runtime/Datasmith/CADKernel/Private/CADKernel/Geo/Surfaces/Surface.cpp
@@ -53,80 +53,6 @@ namespace CADKernel
 		Sampler.Sample();
 	}
 
-	const void FSurface::ComputeIsoTolerances() const
-	{
-		FSurfacicSampling Grid;
-		int32 NumberOfSubdivisions[2] = { 10, 10 };
-
-		Sample(Boundary, NumberOfSubdivisions, Grid);
-
-		const FSurfacicBoundary& SurfaceBoundaries = GetBoundary();
-
-		FSurfacicTolerance& Tolerances = ToleranceIsos;
-
-		TFunction<double(double, double)> ComputeTolerance = [&](double CurveLength3D, double CurveLength2D)
-		{
-			if (CurveLength3D > SMALL_NUMBER)
-			{
-				double Tolerance2D = Tolerance3D * CurveLength2D / CurveLength3D;
-				return FMath::Max(FMath::Min(Tolerance2D, CurveLength2D / 2.), SMALL_NUMBER_SQUARE);
-			}
-			return CurveLength2D / 2.;
-		};
-
-		// Compute 2D tolerance along U
-		{
-			double MaxLengthIsoV = 0.0;
-
-			// For each isoV curve
-			for (int32 Index1 = 0, IndexPoint = 0; Index1 < NumberOfSubdivisions[EIso::IsoV]; ++Index1, ++IndexPoint)
-			{
-				double Length = 0.0;
-
-				// Compute the length of the curve
-				for (int32 Index2 = 0; Index2 < NumberOfSubdivisions[EIso::IsoU] - 1; ++Index2, IndexPoint ++)
-				{
-					Length += Grid.Points3D[IndexPoint].Distance(Grid.Points3D[IndexPoint + 1]);
-				}
-
-				if (Length > MaxLengthIsoV)
-				{
-					MaxLengthIsoV = Length;
-				}
-			}
-
-			Tolerances[EIso::IsoU] = ComputeTolerance(MaxLengthIsoV, SurfaceBoundaries[EIso::IsoU].Length());
-		};
-
-		// Compute 2D tolerance along V
-		{
-			double MaxLengthIsoU = 0.0;
-
-			int32 IndexPoint = 0;
-			int32 Delta = NumberOfSubdivisions[EIso::IsoU];
-
-			// For each isoU curve
-			for (int32 Index1 = 0; Index1 < NumberOfSubdivisions[EIso::IsoU]; ++Index1, ++IndexPoint)
-			{
-				double Length = 0.0;
-					IndexPoint = Index1;
-
-				// Compute the length of the curve
-				for (int32 Index2 = 0; Index2 < NumberOfSubdivisions[EIso::IsoV] - 1; ++Index2, IndexPoint += Delta)
-				{
-					Length += Grid.Points3D[IndexPoint].Distance(Grid.Points3D[IndexPoint + Delta]);
-				}
-
-				if (Length > MaxLengthIsoU)
-				{
-					MaxLengthIsoU = Length;
-				}
-			}
-
-			Tolerances[EIso::IsoV] = ComputeTolerance(MaxLengthIsoU, SurfaceBoundaries[EIso::IsoV].Length());
-		};
-	}
-
 	void FSurface::EvaluatePoints(const TArray<FPoint2D>& InSurfacicCoordinates, TArray<FSurfacicPoint>& OutPoint3D, int32 InDerivativeOrder) const
 	{
 		OutPoint3D.SetNum(InSurfacicCoordinates.Num());
diff --git a/Engine/Source/Runtime/Datasmith/CADKernel/Private/CADKernel/Geo/Surfaces/TabulatedCylinderSurface.cpp b/Engine/Source/Runtime/Datasmith/CADKernel/Private/CADKernel/Geo/Surfaces/TabulatedCylinderSurface.cpp
--- a/Engine/Source/Runtime/Datasmith/CADKernel/Private/CADKernel/Geo/Surfaces/TabulatedCylinderSurface.cpp
+++ b/Engine/Source/Runtime/Datasmith/CADKernel/Private/CADKernel/Geo/Surfaces/TabulatedCylinderSurface.cpp
@@ -93,7 +93,7 @@ void FTabulatedCylinderSurface::EvaluatePointGrid(const FCoordinateGrid& Coordin
 
 void FTabulatedCylinderSurface::Presample(const FSurfacicBoundary& InBoundaries, FCoordinateGrid& Coordinates)
 {
-	GuideCurve->Presample(Coordinates[EIso::IsoU]);
+	GuideCurve->Presample(Coordinates[EIso::IsoU], Tolerance3D);
 
 	Coordinates[EIso::IsoV].Empty(3);
 	Coordinates[EIso::IsoV].Add(InBoundaries.UVBoundaries[EIso::IsoV].Min);
diff --git a/Engine/Source/Runtime/Datasmith/CADKernel/Private/CADKernel/Math/BSpline.cpp b/Engine/Source/Runtime/Datasmith/CADKernel/Private/CADKernel/Math/BSpline.cpp
--- a/Engine/Source/Runtime/Datasmith/CADKernel/Private/CADKernel/Math/BSpline.cpp
+++ b/Engine/Source/Runtime/Datasmith/CADKernel/Private/CADKernel/Math/BSpline.cpp
@@ -24,9 +24,9 @@ namespace CADKernel
 
 		// DeBoor Valeurs B-Spline
 		void DeBoorValeursBSpline( const int32 degre, const double* const knot, const int32 segment, const int32 ndim, const int32 derivee, const double u, double* const pole_aux, double* const grad_aux, double* const lap_aux, double* const valeur, double* const Gradient, double* const Laplacian);
-		void insertKnot(TArray<double>& vn, TArray<FPoint>& poles, TArray<double>& weights, double u, double* newU = nullptr);
-		void blossom(int32 degre, TArray<FPoint>& poles, TArray<double>& vecteurNodal, TArray<double>& poids, int32 seg, TArray<double>& params, FPoint& pnt, double& weight);
-
+		void InsertKnot(TArray<double>& vn, TArray<FPoint>& poles, TArray<double>& weights, double u, double* newU = nullptr);
+		void Blossom(int32 degre, const TArray<FPoint>& poles, const TArray<double>& vecteurNodal, const TArray<double>& poids, int32 seg, TArray<double>& params, FPoint& pnt, double& weight);
+		void DuplicateNurbsCurveWithHigherDegree(int32 degre, const TArray<FPoint>& poles, const TArray<double>& nodalVector, const TArray<double>& weights, TArray<FPoint>& newPoles, TArray<double>& newNodalVector, TArray<double>& newWeights);
 
 #define DEGRE_MAX_POLY 9
 		namespace
@@ -768,7 +768,7 @@ namespace CADKernel
 
 		void Evaluate2DPoint(const FNURBSCurve& Nurbs, double Coordinate, FCurvePoint2D& OutPoint, int32 DerivativeOrder)
 		{
-			ensure(Nurbs.GetDimension() == 2);
+			ensureCADKernel(Nurbs.GetDimension() == 2);
 
 			OutPoint.DerivativeOrder = DerivativeOrder;
 
@@ -815,6 +815,26 @@ namespace CADKernel
 			}
 		}
 
+		TSharedRef<FNURBSCurve> DuplicateNurbsCurveWithHigherDegree(int32 Degre, const FNURBSCurve& InCurve)
+		{
+			TArray<FPoint> NewPoles;
+			TArray<double> NewNodalVector;
+			TArray<double> NewWeights;
+
+			DuplicateNurbsCurveWithHigherDegree(Degre, InCurve.GetPoles(), InCurve.GetNodalVector(), InCurve.GetWeights(), NewPoles, NewNodalVector, NewWeights);
+			
+			TSharedPtr<FCurve> Curve;
+			if (InCurve.IsRational())
+			{
+				return FEntity::MakeShared<FNURBSCurve>(Degre, NewNodalVector, NewPoles, NewWeights, InCurve.GetDimension());
+			}
+			else
+			{
+				return FEntity::MakeShared<FNURBSCurve>(Degre, NewNodalVector, NewPoles, InCurve.GetDimension());
+			}
+		}
+
+
 		void EvaluatePoint(const FNURBSSurface& Nurbs, const FPoint2D& InPoint2D, FSurfacicPoint& OutPoint3D, int32 InDerivativeOrder)
 		{
 			double Point4[4];
@@ -1491,7 +1511,7 @@ namespace CADKernel
 			TabU->Add(UMax);
 		}
 
-		void IncreaseDegree(int32 degre, TArray<FPoint>& poles, TArray<double>& nodalVector, TArray<double>& weights,
+		void DuplicateNurbsCurveWithHigherDegree(int32 degre, const TArray<FPoint>& poles, const TArray<double>& nodalVector, const TArray<double>& weights,
 			TArray<FPoint>& newPoles, TArray<double>& newNodalVector, TArray<double>& newWeights)
 		{
 			//Count number of polynomial segments and
@@ -1556,7 +1576,7 @@ namespace CADKernel
 					weightsTmp.SetNum(newN);
 					for (int32 k = 0; k <newN; k++)
 					{
-						blossom(degre, poles, nodalVector, weights, seg, params[k], poles_[k], weightsTmp[k]);
+						Blossom(degre, poles, nodalVector, weights, seg, params[k], poles_[k], weightsTmp[k]);
 					}
 
 					FPoint pole(0, 0, 0);
@@ -1613,7 +1633,7 @@ namespace CADKernel
 			Swap(vn, newKnots);
 		}
 
-		void insertKnot(TArray<double>& vn, TArray<FPoint>& poles, TArray<double>& weights, double u, double* newU)
+		void InsertKnot(TArray<double>& vn, TArray<FPoint>& poles, TArray<double>& weights, double u, double* newU)
 		{
 			int32 i = 0;
 
@@ -1665,7 +1685,7 @@ namespace CADKernel
 				knotsTmp = knotsV;
 				for (int32 j = 0; j <nbTime; j++)
 				{
-					insertKnot(knotsTmp, poles[i], weights[i], v);
+					InsertKnot(knotsTmp, poles[i], weights[i], v);
 				}
 			}
 			knotsV = knotsTmp;
@@ -1701,7 +1721,7 @@ namespace CADKernel
 
 				for (j = 0; j <nbTime; j++)
 				{
-					insertKnot(knotsTmp, polesCurveU, weightsCurveU, u);
+					InsertKnot(knotsTmp, polesCurveU, weightsCurveU, u);
 				}
 
 				for (j = 0; j <(int32)polesCurveU.Num(); j++)
@@ -1729,11 +1749,11 @@ namespace CADKernel
 			//Rebound curve by knot insertion of u1 then u2...
 			while (getKnotMultiplicity(knots, u1) <deg)
 			{
-				insertKnot(knots, poles, weights, u1);
+				InsertKnot(knots, poles, weights, u1);
 			}
 			while (getKnotMultiplicity(knots, u2) <deg)
 			{
-				insertKnot(knots, poles, weights, u2);
+				InsertKnot(knots, poles, weights, u2);
 			}
 
 			//Now find the corresponding greville abscissas
@@ -1803,7 +1823,7 @@ namespace CADKernel
 			return knots;
 		}
 
-		void blossom(int32 degre, TArray<FPoint>& poles, TArray<double>& nodalVector, TArray<double>& weights, int32 seg, TArray<double>& params, FPoint& pnt, double& weight)
+		void Blossom(int32 degre, const TArray<FPoint>& poles, const TArray<double>& nodalVector, const TArray<double>& weights, int32 seg, TArray<double>& params, FPoint& pnt, double& weight)
 		{
 			int32 n = degre;
 			int32 i = 0;
diff --git a/Engine/Source/Runtime/Datasmith/CADKernel/Private/CADKernel/Math/Point.cpp b/Engine/Source/Runtime/Datasmith/CADKernel/Private/CADKernel/Math/Point.cpp
--- a/Engine/Source/Runtime/Datasmith/CADKernel/Private/CADKernel/Math/Point.cpp
+++ b/Engine/Source/Runtime/Datasmith/CADKernel/Private/CADKernel/Math/Point.cpp
@@ -8,9 +8,20 @@
 using namespace CADKernel;
 
 const FPoint FPoint::ZeroPoint(0.,0.,0.);
+const FPoint FPoint::FarawayPoint(HUGE_VALUE, HUGE_VALUE, HUGE_VALUE);
+const int32 FPoint::Dimension = 3;
+
 const FFPoint FFPoint::ZeroPoint(0.f, 0.f, 0.f);
+const FFPoint FFPoint::FarawayPoint(HUGE_VALUE, HUGE_VALUE, HUGE_VALUE);
+const int32 FFPoint::Dimension = 3;
+
 const FPoint2D FPoint2D::ZeroPoint(0., 0.);
+const FPoint2D FPoint2D::FarawayPoint(HUGE_VALUE, HUGE_VALUE);
+const int32 FPoint2D::Dimension = 2;
+
 const FPointH FPointH::ZeroPoint(0., 0., 0., 1.);
+const FPointH FPointH::FarawayPoint(HUGE_VALUE, HUGE_VALUE, HUGE_VALUE, 1.);
+const int32 FPointH::Dimension = 4;
 
 double FPoint::SignedAngle(const FPoint & Other, const FPoint & Normal) const
 {
@@ -37,34 +48,4 @@ double FPoint::SignedAngle(const FPoint & Other, const FPoint & Normal) const
 	return MixedTripleProduct(Vector1, Vector2, Vector3) > 0 ? acos(ScalarProduct) : -acos(ScalarProduct);
 }
 
-double FPoint::ComputeCosinus(const FPoint& OtherVector) const
-{
-	FPoint ThisNormalized = *this;
-	FPoint OtherNormalized = OtherVector;
-
-	ThisNormalized.Normalize();
-	OtherNormalized.Normalize();
-
-	double Cosinus = ThisNormalized * OtherNormalized;
-
-	return FMath::Max(-1.0, FMath::Min(Cosinus, 1.0));
-}
 
-double FPoint::ComputeSinus(const FPoint& OtherVector) const
-{
-	FPoint ThisNormalized = *this;
-	FPoint OtherNormalized = OtherVector;
-
-	ThisNormalized.Normalize();
-	OtherNormalized.Normalize();
-
-	FPoint SinusPoint = ThisNormalized ^ OtherNormalized;
-	double Sinus = SinusPoint.Length();
-	return FMath::Min(Sinus, 1.0);
-}
-
-double FPoint::ComputeAngle(const FPoint& OtherVector) const
-{
-	double CosAngle = ComputeCosinus(OtherVector);
-	return acos(CosAngle);
-}
diff --git a/Engine/Source/Runtime/Datasmith/CADKernel/Private/CADKernel/Math/SlopeUtils.h b/Engine/Source/Runtime/Datasmith/CADKernel/Private/CADKernel/Math/SlopeUtils.h
--- a/Engine/Source/Runtime/Datasmith/CADKernel/Private/CADKernel/Math/SlopeUtils.h
+++ b/Engine/Source/Runtime/Datasmith/CADKernel/Private/CADKernel/Math/SlopeUtils.h
@@ -172,7 +172,7 @@ namespace CADKernel
 	 * Return true if the segment BP is inside the sector defined the half-lines [BA) and [BC) in the counterclockwise.
 	 * Return false if ABP angle or PBC angle is too flat (smaller than FlatAngle)
 	 */
-	inline bool IsPointPIsInsideSectorABC(const FPoint2D& PointA, const FPoint2D& PointB, const FPoint2D& PointC, const FPoint2D& PointP, const double FlatAngle)
+	inline bool IsPointPBeInsideSectorABC(const FPoint2D& PointA, const FPoint2D& PointB, const FPoint2D& PointC, const FPoint2D& PointP, const double FlatAngle)
 	 {
 		 double SlopWithNextBoundary = ComputeSlope(PointB, PointC);
 		 double BoundaryDeltaSlope = ComputePositiveSlope(PointB, PointA, SlopWithNextBoundary);
