Pluggin of CADKernel : Add Retesselation

diff --git a/Engine/Plugins/Enterprise/DatasmithCADImporter/Source/CADInterfaces/Private/CoreTechFileReader.cpp b/Engine/Plugins/Enterprise/DatasmithCADImporter/Source/CADInterfaces/Private/CoreTechFileReader.cpp
--- a/Engine/Plugins/Enterprise/DatasmithCADImporter/Source/CADInterfaces/Private/CoreTechFileReader.cpp
+++ b/Engine/Plugins/Enterprise/DatasmithCADImporter/Source/CADInterfaces/Private/CoreTechFileReader.cpp
@@ -20,25 +20,19 @@
 #pragma warning(pop)
 
 #include "CADData.h"
+#include "CADKernelTools.h"
 #include "CADOptions.h"
 #include "DatasmithUtils.h"
 
 #include "CADKernel/Core/Entity.h"
-#include "CADKernel/Core/MetadataDictionary.h"
 #include "CADKernel/Core/Session.h"
 #include "CADKernel/Core/Types.h"
 
-#include "CADKernel/Mesh/Criteria/Criterion.h"
 #include "CADKernel/Mesh/Meshers/ParametricMesher.h"
-#include "CADKernel/Mesh/Structure/FaceMesh.h"
 #include "CADKernel/Mesh/Structure/ModelMesh.h"
 
 #include "CADKernel/Topo/Model.h"
 #include "CADKernel/Topo/Body.h"
-#include "CADKernel/Topo/Shell.h"
-#include "CADKernel/Topo/TopologicalEdge.h"
-#include "CADKernel/Topo/TopologicalFace.h"
-#include "CADKernel/Topo/TopologicalVertex.h"
 
 #include "HAL/FileManager.h"
 #include "Internationalization/Text.h"
@@ -787,7 +781,7 @@ namespace CADLibrary
 			// Tesselate the body
 			TSharedRef<CADKernel::FModelMesh> CADKernelModelMesh = CADKernel::FEntity::MakeShared<CADKernel::FModelMesh>();
 
-			DefineMeshCriteria(CADKernelModelMesh);
+			FCADKernelTools::DefineMeshCriteria(CADKernelModelMesh, Context.ImportParameters);
 
 			CADKernel::FParametricMesher Mesher(CADKernelModelMesh);
 			Mesher.MeshEntity(CADKernelModel);
@@ -798,7 +792,7 @@ namespace CADLibrary
 				SetFaceMainMaterial(FaceMaterial, BodyMaterial, BodyMesh, Index);
 			};
 
-			CADKernelUtils::GetBodyTessellation(CADKernelModelMesh, CADKernelBody, BodyMesh, DefaultMaterialHash, ProcessFace);
+			FCADKernelTools::GetBodyTessellation(CADKernelModelMesh, CADKernelBody, BodyMesh, DefaultMaterialHash, ProcessFace);
 
 			CADKernelSession->Clear();
 		}
@@ -1385,32 +1379,6 @@ namespace CADLibrary
 		}
 	}
 
-	void FCoreTechFileReader::DefineMeshCriteria(TSharedRef<CADKernel::FModelMesh>& MeshModel)
-	{
-		{
-			TSharedPtr<CADKernel::FCriterion> CurvatureCriterion = CADKernel::FCriterion::CreateCriterion(CADKernel::ECriterion::CADCurvature);
-			MeshModel->AddCriterion(CurvatureCriterion);
-		}
-
-		if (Context.ImportParameters.MaxEdgeLength > SMALL_NUMBER)
-		{
-			TSharedPtr<CADKernel::FCriterion> MaxSizeCriterion = CADKernel::FCriterion::CreateCriterion(CADKernel::ECriterion::MaxSize, Context.ImportParameters.MaxEdgeLength);
-			MeshModel->AddCriterion(MaxSizeCriterion);
-		}
-
-		if (Context.ImportParameters.ChordTolerance > SMALL_NUMBER)
-		{
-			TSharedPtr<CADKernel::FCriterion> ChordCriterion = CADKernel::FCriterion::CreateCriterion(CADKernel::ECriterion::Sag, Context.ImportParameters.ChordTolerance);
-			MeshModel->AddCriterion(ChordCriterion);
-		}
-
-		if (Context.ImportParameters.MaxNormalAngle > SMALL_NUMBER)
-		{
-			TSharedPtr<CADKernel::FCriterion> MaxNormalAngleCriterion = CADKernel::FCriterion::CreateCriterion(CADKernel::ECriterion::Angle, Context.ImportParameters.MaxNormalAngle);
-			MeshModel->AddCriterion(MaxNormalAngleCriterion);
-		}
-	}
-
 	namespace CoreTechFileReaderUtils
 	{
 		template<typename ValueType>
@@ -2009,108 +1977,6 @@ namespace CADLibrary
 
 	}
 
-	namespace CADKernelUtils
-	{
-		uint32 GetFaceTessellation(const TSharedRef<CADKernel::FFaceMesh>& FaceMesh, FBodyMesh& OutBodyMesh)
-		{
-			// Something wrong happened, either an error or no data to collect
-			if (FaceMesh->TrianglesVerticesIndex.Num() == 0)
-			{
-				return 0;
-			}
-
-			FTessellationData& Tessellation = OutBodyMesh.Faces.Emplace_GetRef();
-			
-			const TSharedRef<CADKernel::FMetadataDictionary>& HaveMetadata = StaticCastSharedRef<CADKernel::FMetadataDictionary>(StaticCastSharedRef<CADKernel::FTopologicalFace>(FaceMesh->GetGeometricEntity()));
-			Tessellation.PatchId = HaveMetadata->GetPatchId();
-
-			Tessellation.PositionIndices = FaceMesh->VerticesGlobalIndex;
-			Tessellation.VertexIndices = FaceMesh->TrianglesVerticesIndex;
-
-			Tessellation.NormalArray.Reserve(FaceMesh->Normals.Num());
-			for (const CADKernel::FPoint& Normal : FaceMesh->Normals)
-			{
-				Tessellation.NormalArray.Emplace((float)Normal.X, (float)Normal.Y, (float)Normal.Z);
-			}
-
-			Tessellation.TexCoordArray.Reserve(FaceMesh->UVMap.Num());
-			for (const CADKernel::FPoint2D& TexCoord : FaceMesh->UVMap)
-			{
-				Tessellation.TexCoordArray.Emplace((float)TexCoord.U, (float)TexCoord.V);
-			}
-
-			return Tessellation.VertexIndices.Num() / 3;
-		}
-
-		template<class ClassType>
-		void GetDisplayDataIds(const TSharedRef<ClassType>& Entity, FObjectDisplayDataId& DisplayDataId)
-		{
-			const TSharedRef<CADKernel::FMetadataDictionary>& HaveMetadata = StaticCastSharedRef<CADKernel::FMetadataDictionary>(Entity);
-			DisplayDataId.Color = HaveMetadata->GetColorId();
-			DisplayDataId.Material = HaveMetadata->GetMaterialId();
-		}
-
-		void GetBodyTessellation(const TSharedRef<CADKernel::FModelMesh>& ModelMesh, const TSharedRef<CADKernel::FBody>& Body, FBodyMesh& OutBodyMesh, uint32 DefaultMaterialHash, TFunction<void(FObjectDisplayDataId, FObjectDisplayDataId, int32)> SetFaceMainMaterial)
-		{
-			ModelMesh->GetNodeCoordinates(OutBodyMesh.VertexArray);
-
-			uint32 FaceSize = Body->FaceCount();
-
-			// Allocate memory space for tessellation data
-			OutBodyMesh.Faces.Reserve(FaceSize);
-			OutBodyMesh.ColorSet.Reserve(FaceSize);
-			OutBodyMesh.MaterialSet.Reserve(FaceSize);
-
-			FObjectDisplayDataId BodyMaterial;
-			BodyMaterial.DefaultMaterialName = DefaultMaterialHash;
-
-			GetDisplayDataIds(Body, BodyMaterial);
-
-			// Loop through the face of bodies and collect all tessellation data
-			int32 FaceIndex = 0;
-			for (const TSharedPtr<CADKernel::FShell>& Shell : Body->GetShells())
-			{
-				if (!Shell.IsValid())
-				{
-					continue;
-				}
-
-				FObjectDisplayDataId ShellMaterial = BodyMaterial;
-				GetDisplayDataIds(Shell.ToSharedRef(), ShellMaterial);
-
-				for (const CADKernel::FOrientedFace& Face : Shell->GetFaces())
-				{
-					if (!Face.Entity.IsValid())
-					{
-						continue;
-					}
-
-					if (!Face.Entity->HasTesselation())
-					{
-						continue;
-					}
-
-					FObjectDisplayDataId FaceMaterial;
-					GetDisplayDataIds(Face.Entity.ToSharedRef(), FaceMaterial);
-
-					uint32 TriangleNum = GetFaceTessellation(Face.Entity->GetMesh(), OutBodyMesh);
-
-					if (TriangleNum == 0)
-					{
-						continue;
-					}
-
-					OutBodyMesh.TriangleCount += TriangleNum;
-
-					if (SetFaceMainMaterial)
-					{
-						SetFaceMainMaterial(FaceMaterial, BodyMaterial, FaceIndex);
-					}
-					FaceIndex++;
-				}
-			}
-		}
-	}
 }
 
 #endif // USE_KERNEL_IO_SDK
diff --git a/Engine/Plugins/Enterprise/DatasmithCADImporter/Source/CADKernelSurface/CADKernelSurface.Build.cs b/Engine/Plugins/Enterprise/DatasmithCADImporter/Source/CADKernelSurface/CADKernelSurface.Build.cs
--- a/Engine/Plugins/Enterprise/DatasmithCADImporter/Source/CADKernelSurface/CADKernelSurface.Build.cs
+++ b/Engine/Plugins/Enterprise/DatasmithCADImporter/Source/CADKernelSurface/CADKernelSurface.Build.cs
@@ -18,6 +18,7 @@ namespace UnrealBuildTool.Rules
 			PrivateDependencyModuleNames.AddRange(
 				new string[]
 				{
+					"CADKernel",
 					"CADLibrary",
 					"CADTools",
 					"CoreUObject",
diff --git a/Engine/Plugins/Enterprise/DatasmithCADImporter/Source/CADKernelSurface/Private/CADKernelSurfaceExtension.cpp b/Engine/Plugins/Enterprise/DatasmithCADImporter/Source/CADKernelSurface/Private/CADKernelSurfaceExtension.cpp
--- a/Engine/Plugins/Enterprise/DatasmithCADImporter/Source/CADKernelSurface/Private/CADKernelSurfaceExtension.cpp
+++ b/Engine/Plugins/Enterprise/DatasmithCADImporter/Source/CADKernelSurface/Private/CADKernelSurfaceExtension.cpp
@@ -2,14 +2,26 @@
 
 #include "CADKernelSurfaceExtension.h"
 
+#include "CADKernelTools.h"
+#include "CADData.h"
 #include "CADOptions.h"
 #include "DatasmithPayload.h"
 #include "Engine/StaticMesh.h"
+#include "IDatasmithSceneElements.h"
 #include "MeshDescription.h"
 #include "MeshDescriptionHelper.h"
 #include "StaticMeshAttributes.h"
 #include "UObject/EnterpriseObjectVersion.h"
 
+#include "CADKernel/Core/CADKernelArchive.h"
+#include "CADKernel/Core/Entity.h"
+#include "CADKernel/Core/Session.h"
+#include "CADKernel/Core/Types.h"
+#include "CADKernel/Mesh/Meshers/ParametricMesher.h"
+#include "CADKernel/Mesh/Structure/ModelMesh.h"
+#include "CADKernel/Topo/Model.h"
+#include "CADKernel/Topo/Body.h"
+
 void UCADKernelParametricSurfaceData::Serialize(FArchive& Ar)
 {
 	Ar.UsingCustomVersion(FEnterpriseObjectVersion::GUID);
@@ -49,24 +61,46 @@ bool UCADKernelParametricSurfaceData::Tessellate(UStaticMesh& StaticMesh, const
 			CADLibrary::CopyPatchGroups(*DestinationMeshDescription, MeshDescription);
 		}
 
-		if (false /*CADLibrary::LoadFile(ResourceFile, MeshDescription, ImportParameters, CadMeshParameters)*/)
+		TSharedRef<CADKernel::FSession> CADKernelSession = MakeShared<CADKernel::FSession>(0.00001 / ImportParameters.MetricUnit);
+		CADKernelSession->AddDatabase(RawData);
+
+		TSharedRef<CADKernel::FModel> CADKernelModel = CADKernelSession->GetModel();
+
+		// Tesselate the model
+		TSharedRef<CADKernel::FModelMesh> CADKernelModelMesh = CADKernel::FEntity::MakeShared<CADKernel::FModelMesh>();
+
+		FCADKernelTools::DefineMeshCriteria(CADKernelModelMesh, ImportParameters);
+
+		CADKernel::FParametricMesher Mesher(CADKernelModelMesh);
+		Mesher.MeshEntity(CADKernelModel);
+
+		TArray<TSharedPtr<CADKernel::FBody>> CADKernelBodies = CADKernelModel->GetBodies();
+		if (CADKernelBodies.Num() != 1)
+		{
+			return bSuccessfulTessellation;
+		}
+		const TSharedRef<CADKernel::FBody> Body = CADKernelBodies[0].ToSharedRef();
+
+		CADLibrary::FBodyMesh BodyMesh;
+		uint32 DefaultMaterialHash = 0;
+
+		FCADKernelTools::GetBodyTessellation(CADKernelModelMesh, Body, BodyMesh, DefaultMaterialHash, [](CADLibrary::FObjectDisplayDataId, CADLibrary::FObjectDisplayDataId, int32) {;});
+		if (CADLibrary::ConvertBodyMeshToMeshDescription(ImportParameters, CadMeshParameters, BodyMesh, MeshDescription))
 		{
 			// To update the SectionInfoMap 
-			{
-				TPolygonGroupAttributesConstRef<FName> MaterialSlotNames = MeshDescriptionAttributes.GetPolygonGroupMaterialSlotNames();
-				FMeshSectionInfoMap& SectionInfoMap = StaticMesh.GetSectionInfoMap();
+			TPolygonGroupAttributesConstRef<FName> MaterialSlotNames = MeshDescriptionAttributes.GetPolygonGroupMaterialSlotNames();
+			FMeshSectionInfoMap& SectionInfoMap = StaticMesh.GetSectionInfoMap();
 
-				for (FPolygonGroupID PolygonGroupID : MeshDescription.PolygonGroups().GetElementIDs())
+			for (FPolygonGroupID PolygonGroupID : MeshDescription.PolygonGroups().GetElementIDs())
+			{
+				FMeshSectionInfo Section = SectionInfoMap.Get(0, PolygonGroupID.GetValue());
+				int32 MaterialIndex = StaticMesh.GetMaterialIndex(MaterialSlotNames[PolygonGroupID]);
+				if (MaterialIndex < 0)
 				{
-					FMeshSectionInfo Section = SectionInfoMap.Get(0, PolygonGroupID.GetValue());
-					int32 MaterialIndex = StaticMesh.GetMaterialIndex(MaterialSlotNames[PolygonGroupID]);
-					if (MaterialIndex < 0)
-					{
-						MaterialIndex = 0;
-					}
-					Section.MaterialIndex = MaterialIndex;
-					SectionInfoMap.Set(0, PolygonGroupID.GetValue(), Section);
+					MaterialIndex = 0;
 				}
+				Section.MaterialIndex = MaterialIndex;
+				SectionInfoMap.Set(0, PolygonGroupID.GetValue(), Section);
 			}
 			*DestinationMeshDescription = MoveTemp(MeshDescription);
 			bSuccessfulTessellation = true;
@@ -81,24 +115,27 @@ namespace CADKernelSurface
 {
 	void AddSurfaceDataForMesh(const TSharedRef<IDatasmithMeshElement>& InMeshElement, const CADLibrary::FImportParameters& InSceneParameters, const CADLibrary::FMeshParameters& InMeshParameters, const FDatasmithTessellationOptions& InTessellationOptions, FDatasmithMeshElementPayload& OutMeshPayload)
 	{
-		// Store CADKernel additional data if provided
-		// TODO
-		TArray<uint8> ByteArray;
-		if (ByteArray.Num() > 0)
+		// Store CoreTech additional data if provided
+		FString CADKernelDatabase = InMeshElement->GetFile();
+		if (FPaths::FileExists(*CADKernelDatabase))
 		{
-			UCADKernelParametricSurfaceData* CADKernelData = Datasmith::MakeAdditionalData<UCADKernelParametricSurfaceData>();
-			CADKernelData->RawData = MoveTemp(ByteArray);
-			CADKernelData->SceneParameters.ModelCoordSys = uint8(InSceneParameters.ModelCoordSys);
-			CADKernelData->SceneParameters.MetricUnit = InSceneParameters.MetricUnit;
-			CADKernelData->SceneParameters.ScaleFactor = InSceneParameters.ScaleFactor;
-
-			CADKernelData->MeshParameters.bNeedSwapOrientation = InMeshParameters.bNeedSwapOrientation;
-			CADKernelData->MeshParameters.bIsSymmetric = InMeshParameters.bIsSymmetric;
-			CADKernelData->MeshParameters.SymmetricNormal = InMeshParameters.SymmetricNormal;
-			CADKernelData->MeshParameters.SymmetricOrigin = InMeshParameters.SymmetricOrigin;
-
-			CADKernelData->LastTessellationOptions = InTessellationOptions;
-			OutMeshPayload.AdditionalData.Add(CADKernelData);
+			TArray<uint8> ByteArray;
+			if (FFileHelper::LoadFileToArray(ByteArray, *CADKernelDatabase))
+			{
+				UCADKernelParametricSurfaceData* CADKernelData = Datasmith::MakeAdditionalData<UCADKernelParametricSurfaceData>();
+				CADKernelData->RawData = MoveTemp(ByteArray);
+				CADKernelData->SceneParameters.ModelCoordSys = uint8(InSceneParameters.ModelCoordSys);
+				CADKernelData->SceneParameters.MetricUnit = InSceneParameters.MetricUnit;
+				CADKernelData->SceneParameters.ScaleFactor = InSceneParameters.ScaleFactor;
+
+				CADKernelData->MeshParameters.bNeedSwapOrientation = InMeshParameters.bNeedSwapOrientation;
+				CADKernelData->MeshParameters.bIsSymmetric = InMeshParameters.bIsSymmetric;
+				CADKernelData->MeshParameters.SymmetricNormal = InMeshParameters.SymmetricNormal;
+				CADKernelData->MeshParameters.SymmetricOrigin = InMeshParameters.SymmetricOrigin;
+
+				CADKernelData->LastTessellationOptions = InTessellationOptions;
+				OutMeshPayload.AdditionalData.Add(CADKernelData);
+			}
 		}
 	}
 
diff --git a/Engine/Plugins/Enterprise/DatasmithCADImporter/Source/CADKernelSurface/Public/CADKernelSurfaceExtension.h b/Engine/Plugins/Enterprise/DatasmithCADImporter/Source/CADKernelSurface/Public/CADKernelSurfaceExtension.h
--- a/Engine/Plugins/Enterprise/DatasmithCADImporter/Source/CADKernelSurface/Public/CADKernelSurfaceExtension.h
+++ b/Engine/Plugins/Enterprise/DatasmithCADImporter/Source/CADKernelSurface/Public/CADKernelSurfaceExtension.h
@@ -18,9 +18,6 @@ class CADKERNELSURFACE_API UCADKernelParametricSurfaceData : public UParametricS
 	GENERATED_BODY()
 
 public:
-	UPROPERTY()
-	FString SourceFile;
-
 	// Too costly to serialize as a UPROPERTY, will use custom serialization.
 	TArray<uint8> RawData;
 
diff --git a/Engine/Plugins/Enterprise/DatasmithCADImporter/Source/CADTools/CADTools.Build.cs b/Engine/Plugins/Enterprise/DatasmithCADImporter/Source/CADTools/CADTools.Build.cs
--- a/Engine/Plugins/Enterprise/DatasmithCADImporter/Source/CADTools/CADTools.Build.cs
+++ b/Engine/Plugins/Enterprise/DatasmithCADImporter/Source/CADTools/CADTools.Build.cs
@@ -8,13 +8,14 @@ namespace UnrealBuildTool.Rules
 	{
 		public CADTools(ReadOnlyTargetRules Target) : base(Target)
 		{
-            PublicDependencyModuleNames.AddRange(
+			PublicDependencyModuleNames.AddRange(
 				new string[]
 				{
                     "Core",
+					"CADKernel",
 					"DatasmithCore"
-                }
-            );
+				}
+			);
 
 			PrivateDependencyModuleNames.AddRange(
 				new string[]
diff --git a/Engine/Plugins/Enterprise/DatasmithCADImporter/Source/CADTools/Private/CADKernelTools.cpp b/Engine/Plugins/Enterprise/DatasmithCADImporter/Source/CADTools/Private/CADKernelTools.cpp
new file mode 100644
--- /dev/null
+++ b/Engine/Plugins/Enterprise/DatasmithCADImporter/Source/CADTools/Private/CADKernelTools.cpp
@@ -0,0 +1,149 @@
+// Copyright Epic Games, Inc. All Rights Reserved.
+
+#include "CADKernelTools.h"
+
+#include "CADData.h"
+#include "CADOptions.h"
+
+#include "CADKernel/Core/Entity.h"
+#include "CADKernel/Core/MetadataDictionary.h"
+#include "CADKernel/Core/Session.h"
+#include "CADKernel/Core/Types.h"
+
+#include "CADKernel/Mesh/Criteria/Criterion.h"
+#include "CADKernel/Mesh/Meshers/ParametricMesher.h"
+#include "CADKernel/Mesh/Structure/FaceMesh.h"
+#include "CADKernel/Mesh/Structure/ModelMesh.h"
+
+#include "CADKernel/Topo/Model.h"
+#include "CADKernel/Topo/Body.h"
+#include "CADKernel/Topo/Shell.h"
+#include "CADKernel/Topo/TopologicalEdge.h"
+#include "CADKernel/Topo/TopologicalFace.h"
+#include "CADKernel/Topo/TopologicalVertex.h"
+
+uint32 FCADKernelTools::GetFaceTessellation(const TSharedRef<CADKernel::FFaceMesh>& FaceMesh, CADLibrary::FBodyMesh& OutBodyMesh)
+{
+	// Something wrong happened, either an error or no data to collect
+	if (FaceMesh->TrianglesVerticesIndex.Num() == 0)
+	{
+		return 0;
+	}
+
+	CADLibrary::FTessellationData& Tessellation = OutBodyMesh.Faces.Emplace_GetRef();
+
+	const TSharedRef<CADKernel::FMetadataDictionary>& HaveMetadata = StaticCastSharedRef<CADKernel::FMetadataDictionary>(StaticCastSharedRef<CADKernel::FTopologicalFace>(FaceMesh->GetGeometricEntity()));
+	Tessellation.PatchId = HaveMetadata->GetPatchId();
+
+	Tessellation.PositionIndices = FaceMesh->VerticesGlobalIndex;
+	Tessellation.VertexIndices = FaceMesh->TrianglesVerticesIndex;
+
+	Tessellation.NormalArray.Reserve(FaceMesh->Normals.Num());
+	for (const CADKernel::FPoint& Normal : FaceMesh->Normals)
+	{
+		Tessellation.NormalArray.Emplace((float)Normal.X, (float)Normal.Y, (float)Normal.Z);
+	}
+
+	Tessellation.TexCoordArray.Reserve(FaceMesh->UVMap.Num());
+	for (const CADKernel::FPoint2D& TexCoord : FaceMesh->UVMap)
+	{
+		Tessellation.TexCoordArray.Emplace((float)TexCoord.U, (float)TexCoord.V);
+	}
+
+	return Tessellation.VertexIndices.Num() / 3;
+}
+
+template<class ClassType>
+void GetDisplayDataIds(const TSharedRef<ClassType>& Entity, CADLibrary::FObjectDisplayDataId& DisplayDataId)
+{
+	const TSharedRef<CADKernel::FMetadataDictionary>& HaveMetadata = StaticCastSharedRef<CADKernel::FMetadataDictionary>(Entity);
+	DisplayDataId.Color = HaveMetadata->GetColorId();
+	DisplayDataId.Material = HaveMetadata->GetMaterialId();
+}
+
+void FCADKernelTools::GetBodyTessellation(const TSharedRef<CADKernel::FModelMesh>& ModelMesh, const TSharedRef<CADKernel::FBody>& Body, CADLibrary::FBodyMesh& OutBodyMesh, uint32 DefaultMaterialHash, TFunction<void(CADLibrary::FObjectDisplayDataId, CADLibrary::FObjectDisplayDataId, int32)> SetFaceMainMaterial)
+{
+	ModelMesh->GetNodeCoordinates(OutBodyMesh.VertexArray);
+
+	uint32 FaceSize = Body->FaceCount();
+
+	// Allocate memory space for tessellation data
+	OutBodyMesh.Faces.Reserve(FaceSize);
+	OutBodyMesh.ColorSet.Reserve(FaceSize);
+	OutBodyMesh.MaterialSet.Reserve(FaceSize);
+
+	CADLibrary::FObjectDisplayDataId BodyMaterial;
+	BodyMaterial.DefaultMaterialName = DefaultMaterialHash;
+
+	GetDisplayDataIds(Body, BodyMaterial);
+
+	// Loop through the face of bodies and collect all tessellation data
+	int32 FaceIndex = 0;
+	for (const TSharedPtr<CADKernel::FShell>& Shell : Body->GetShells())
+	{
+		if (!Shell.IsValid())
+		{
+			continue;
+		}
+
+		CADLibrary::FObjectDisplayDataId ShellMaterial = BodyMaterial;
+		GetDisplayDataIds(Shell.ToSharedRef(), ShellMaterial);
+
+		for (const CADKernel::FOrientedFace& Face : Shell->GetFaces())
+		{
+			if (!Face.Entity.IsValid())
+			{
+				continue;
+			}
+
+			if (!Face.Entity->HasTesselation())
+			{
+				continue;
+			}
+
+			CADLibrary::FObjectDisplayDataId FaceMaterial;
+			GetDisplayDataIds(Face.Entity.ToSharedRef(), FaceMaterial);
+
+			uint32 TriangleNum = GetFaceTessellation(Face.Entity->GetMesh(), OutBodyMesh);
+
+			if (TriangleNum == 0)
+			{
+				continue;
+			}
+
+			OutBodyMesh.TriangleCount += TriangleNum;
+
+			if (SetFaceMainMaterial)
+			{
+				SetFaceMainMaterial(FaceMaterial, BodyMaterial, FaceIndex);
+			}
+			FaceIndex++;
+		}
+	}
+}
+
+void FCADKernelTools::DefineMeshCriteria(TSharedRef<CADKernel::FModelMesh>& MeshModel, const CADLibrary::FImportParameters& ImportParameters)
+{
+	{
+		TSharedPtr<CADKernel::FCriterion> CurvatureCriterion = CADKernel::FCriterion::CreateCriterion(CADKernel::ECriterion::CADCurvature);
+		MeshModel->AddCriterion(CurvatureCriterion);
+	}
+
+	if (ImportParameters.MaxEdgeLength > SMALL_NUMBER)
+	{
+		TSharedPtr<CADKernel::FCriterion> MaxSizeCriterion = CADKernel::FCriterion::CreateCriterion(CADKernel::ECriterion::MaxSize, ImportParameters.MaxEdgeLength);
+		MeshModel->AddCriterion(MaxSizeCriterion);
+	}
+
+	if (ImportParameters.ChordTolerance > SMALL_NUMBER)
+	{
+		TSharedPtr<CADKernel::FCriterion> ChordCriterion = CADKernel::FCriterion::CreateCriterion(CADKernel::ECriterion::Sag, ImportParameters.ChordTolerance);
+		MeshModel->AddCriterion(ChordCriterion);
+	}
+
+	if (ImportParameters.MaxNormalAngle > SMALL_NUMBER)
+	{
+		TSharedPtr<CADKernel::FCriterion> MaxNormalAngleCriterion = CADKernel::FCriterion::CreateCriterion(CADKernel::ECriterion::Angle, ImportParameters.MaxNormalAngle);
+		MeshModel->AddCriterion(MaxNormalAngleCriterion);
+	}
+}
diff --git a/Engine/Plugins/Enterprise/DatasmithCADImporter/Source/CADTools/Public/CADKernelTools.h b/Engine/Plugins/Enterprise/DatasmithCADImporter/Source/CADTools/Public/CADKernelTools.h
new file mode 100644
--- /dev/null
+++ b/Engine/Plugins/Enterprise/DatasmithCADImporter/Source/CADTools/Public/CADKernelTools.h
@@ -0,0 +1,31 @@
+// Copyright Epic Games, Inc. All Rights Reserved.
+#pragma once
+
+#include "CoreMinimal.h"
+
+#include "CADData.h"
+#include "CADOptions.h"
+
+namespace CADKernel
+{
+	class FBody;
+	class FBodyMesh;
+	class FFaceMesh;
+	class FModel;
+	class FModelMesh;
+}
+
+namespace CADLibrary
+{
+	struct FImportParameters;
+}
+
+class CADTOOLS_API FCADKernelTools
+{
+public:
+	static void DefineMeshCriteria(TSharedRef<CADKernel::FModelMesh>& MeshModel, const CADLibrary::FImportParameters& ImportParameters);
+
+	static void GetBodyTessellation(const TSharedRef<CADKernel::FModelMesh>& ModelMesh, const TSharedRef<CADKernel::FBody>& Body, CADLibrary::FBodyMesh& OutBodyMesh, uint32 DefaultMaterialHash, TFunction<void(CADLibrary::FObjectDisplayDataId, CADLibrary::FObjectDisplayDataId, int32)> SetFaceMainMaterial);
+
+	static uint32 GetFaceTessellation(const TSharedRef<CADKernel::FFaceMesh>& FaceMesh, CADLibrary::FBodyMesh& OutBodyMesh);
+};
diff --git a/Engine/Plugins/Enterprise/DatasmithCADImporter/Source/CADTools/Public/CADOptions.h b/Engine/Plugins/Enterprise/DatasmithCADImporter/Source/CADTools/Public/CADOptions.h
--- a/Engine/Plugins/Enterprise/DatasmithCADImporter/Source/CADTools/Public/CADOptions.h
+++ b/Engine/Plugins/Enterprise/DatasmithCADImporter/Source/CADTools/Public/CADOptions.h
@@ -4,6 +4,7 @@
 
 #include "DatasmithUtils.h"
 #include "Math/Vector.h"
+#include "Misc/Paths.h"
 #include "Templates/TypeHash.h"
 
 namespace CADLibrary
@@ -78,6 +79,19 @@ namespace CADLibrary
 			return Ar;
 		}
 
+		FString DefineCADFilePath(const TCHAR* FolderPath, const TCHAR* InFileName) const
+		{
+			FString OutFileName = FPaths::Combine(FolderPath, InFileName);
+			if (bEnableKernelIOTessellation)
+			{
+				OutFileName += TEXT(".ct");
+			}
+			else
+			{
+				OutFileName += TEXT(".ugeom");
+			}
+			return OutFileName;
+		}
 	};
 
 	struct FMeshParameters
diff --git a/Engine/Plugins/Enterprise/DatasmithCADImporter/Source/CoreTechSurface/Private/CoreTechSurfaceHelper.cpp b/Engine/Plugins/Enterprise/DatasmithCADImporter/Source/CoreTechSurface/Private/CoreTechSurfaceHelper.cpp
--- a/Engine/Plugins/Enterprise/DatasmithCADImporter/Source/CoreTechSurface/Private/CoreTechSurfaceHelper.cpp
+++ b/Engine/Plugins/Enterprise/DatasmithCADImporter/Source/CoreTechSurface/Private/CoreTechSurfaceHelper.cpp
@@ -69,11 +69,11 @@ namespace CoreTechSurface
 		if (ICADInterfacesModule::GetAvailability() == ECADInterfaceAvailability::Available)
 		{
 			// Store CoreTech additional data if provided
-			const TCHAR* CoretechFile = InMeshElement->GetFile();
-			if (FPaths::FileExists(CoretechFile))
+			FString CoretechFile = InMeshElement->GetFile(); 
+			if (FPaths::FileExists(*CoretechFile))
 			{
 				TArray<uint8> ByteArray;
-				if (FFileHelper::LoadFileToArray(ByteArray, CoretechFile))
+				if (FFileHelper::LoadFileToArray(ByteArray, *CoretechFile))
 				{
 					UCoreTechParametricSurfaceData* CoreTechData = Datasmith::MakeAdditionalData<UCoreTechParametricSurfaceData>();
 					CoreTechData->SourceFile = CoretechFile;
diff --git a/Engine/Plugins/Enterprise/DatasmithCADImporter/Source/DatasmithCADTranslator/DatasmithCADTranslator.Build.cs b/Engine/Plugins/Enterprise/DatasmithCADImporter/Source/DatasmithCADTranslator/DatasmithCADTranslator.Build.cs
--- a/Engine/Plugins/Enterprise/DatasmithCADImporter/Source/DatasmithCADTranslator/DatasmithCADTranslator.Build.cs
+++ b/Engine/Plugins/Enterprise/DatasmithCADImporter/Source/DatasmithCADTranslator/DatasmithCADTranslator.Build.cs
@@ -19,6 +19,7 @@ namespace UnrealBuildTool.Rules
 			PrivateDependencyModuleNames.AddRange(
 				new string[]
 				{
+					"CADKernelSurface",
 					"CADInterfaces",
 					"CADLibrary",
 					"CADTools",
diff --git a/Engine/Plugins/Enterprise/DatasmithCADImporter/Source/DatasmithCADTranslator/Private/DatasmithCADTranslator.cpp b/Engine/Plugins/Enterprise/DatasmithCADImporter/Source/DatasmithCADTranslator/Private/DatasmithCADTranslator.cpp
--- a/Engine/Plugins/Enterprise/DatasmithCADImporter/Source/DatasmithCADTranslator/Private/DatasmithCADTranslator.cpp
+++ b/Engine/Plugins/Enterprise/DatasmithCADImporter/Source/DatasmithCADTranslator/Private/DatasmithCADTranslator.cpp
@@ -3,6 +3,7 @@
 #include "DatasmithCADTranslator.h"
 
 #include "CADInterfacesModule.h"
+#include "CADKernelSurfaceExtension.h"
 #include "CoreTechSurfaceHelper.h"
 #include "DatasmithCADTranslatorModule.h"
 #include "DatasmithDispatcher.h"
@@ -154,16 +155,16 @@ bool FDatasmithCADTranslator::LoadScene(TSharedRef<IDatasmithScene> DatasmithSce
 			DatasmithDispatcher::FDatasmithDispatcher Dispatcher(ImportParameters, CachePath, NumCores, CADFileToUEFileMap, CADFileToUEGeomMap);
 			Dispatcher.AddTask(FileDescription);
 
-		Dispatcher.Process(bWithProcessor);
-	}
+			Dispatcher.Process(bWithProcessor);
+		}
 
-	FDatasmithSceneGraphBuilder SceneGraphBuilder(CADFileToUEFileMap, CachePath, DatasmithScene, GetSource(), ImportParameters);
-	SceneGraphBuilder.Build();
+		FDatasmithSceneGraphBuilder SceneGraphBuilder(CADFileToUEFileMap, CachePath, DatasmithScene, GetSource(), ImportParameters);
+		SceneGraphBuilder.Build();
 
-	MeshBuilderPtr = MakeUnique<FDatasmithMeshBuilder>(CADFileToUEGeomMap, CachePath, ImportParameters);
+		MeshBuilderPtr = MakeUnique<FDatasmithMeshBuilder>(CADFileToUEGeomMap, CachePath, ImportParameters);
 
-	return true;
-}
+		return true;
+	}
 
 	ImportParameters.bEnableCacheUsage = false;
 
@@ -206,7 +207,7 @@ bool FDatasmithCADTranslator::LoadStaticMesh(const TSharedRef<IDatasmithMeshElem
 		}
 		else
 		{
-			//CADKernelSurface::AddSurfaceDataForMesh(MeshElement, ImportParameters, MeshParameters, GetCommonTessellationOptions(), OutMeshPayload);
+			CADKernelSurface::AddSurfaceDataForMesh(MeshElement, ImportParameters, MeshParameters, GetCommonTessellationOptions(), OutMeshPayload);
 		}
 	}
 	return OutMeshPayload.LodMeshes.Num() > 0;
diff --git a/Engine/Plugins/Enterprise/DatasmithCADImporter/Source/DatasmithCADTranslator/Private/DatasmithSceneGraphBuilder.cpp b/Engine/Plugins/Enterprise/DatasmithCADImporter/Source/DatasmithCADTranslator/Private/DatasmithSceneGraphBuilder.cpp
--- a/Engine/Plugins/Enterprise/DatasmithCADImporter/Source/DatasmithCADTranslator/Private/DatasmithSceneGraphBuilder.cpp
+++ b/Engine/Plugins/Enterprise/DatasmithCADImporter/Source/DatasmithCADTranslator/Private/DatasmithSceneGraphBuilder.cpp
@@ -168,9 +168,8 @@ TSharedPtr<IDatasmithMeshElement> FDatasmithSceneGraphBuilder::FindOrAddMeshElem
 {
 	if (TSharedPtr<IDatasmithMeshElement> MeshElement = FDatasmithSceneBaseGraphBuilder::FindOrAddMeshElement(Body, InLabel))
 	{
-		FString BodyFile = FString::Printf(TEXT("UEx%08x"), Body.MeshActorName);
-		MeshElement->SetFile(*FPaths::Combine(CachePath, TEXT("body"), BodyFile + TEXT(".ct")));
-
+		FString BodyFile = ImportParameters.DefineCADFilePath(*FPaths::Combine(CachePath, TEXT("body")), *FString::Printf(TEXT("UEx%08x"), Body.MeshActorName));
+		MeshElement->SetFile(*BodyFile);
 		return MeshElement;
 	}
 
@@ -566,17 +565,6 @@ TSharedPtr< IDatasmithMeshElement > FDatasmithSceneBaseGraphBuilder::FindOrAddMe
 	Hash.Set(MD5);
 	MeshElement->SetFileHash(Hash);
 
-	MeshElement->SetFile(*(ShellUuidName + TEXT(".ct")));
-
-
-	// TODO: Set bounding box 
-	//float BoundingBox[6];
-	//FString Buffer = GetStringAttribute(GeomID, TEXT("UE_MESH_BBOX"));
-	//if (FString::ToHexBlob(Buffer, (uint8*)BoundingBox, sizeof(BoundingBox)))
-	//{
-	//	MeshElement->SetDimensions(BoundingBox[3] - BoundingBox[0], BoundingBox[4] - BoundingBox[1], BoundingBox[5] - BoundingBox[2], 0.0f);
-	//}
-
 	// Currently we assume that face has only colors
 	TSet<uint32>& MaterialSet = Body.ColorFaceSet;
 
diff --git a/Engine/Plugins/Enterprise/DatasmithCADImporter/Source/DatasmithOpenNurbsTranslator/Private/DatasmithOpenNurbsTranslator.cpp b/Engine/Plugins/Enterprise/DatasmithCADImporter/Source/DatasmithOpenNurbsTranslator/Private/DatasmithOpenNurbsTranslator.cpp
--- a/Engine/Plugins/Enterprise/DatasmithCADImporter/Source/DatasmithOpenNurbsTranslator/Private/DatasmithOpenNurbsTranslator.cpp
+++ b/Engine/Plugins/Enterprise/DatasmithCADImporter/Source/DatasmithOpenNurbsTranslator/Private/DatasmithOpenNurbsTranslator.cpp
@@ -3189,8 +3189,7 @@ bool FOpenNurbsTranslatorImpl::TranslateBRep(ON_Brep* Brep, const ON_3dmObjectAt
 
 		LocalSession->AddBRep(*Brep, Offset);
 
-		FString Filename = FString::Printf(TEXT("%s.ct"), *Name);
-		FString FilePath = FPaths::Combine(OutputPath, Filename);
+		FString FilePath = LocalSession->GetImportParameters().DefineCADFilePath(*OutputPath, *Name);
 		if (LocalSession->SaveBrep(FilePath))
 		{
 			MeshElement->SetFile(*FilePath);
diff --git a/Engine/Source/Runtime/Datasmith/CADKernel/Private/CADKernel/Core/Session.cpp b/Engine/Source/Runtime/Datasmith/CADKernel/Private/CADKernel/Core/Session.cpp
--- a/Engine/Source/Runtime/Datasmith/CADKernel/Private/CADKernel/Core/Session.cpp
+++ b/Engine/Source/Runtime/Datasmith/CADKernel/Private/CADKernel/Core/Session.cpp
@@ -66,3 +66,9 @@ void FSession::LoadDatabase(const TCHAR* FilePath)
 	}
 	Archive->Close();
 }
+
+void FSession::AddDatabase(const TArray<uint8>& InRawData)
+{
+	FCADKernelArchive Archive = FCADKernelArchive(*this, InRawData);
+	Database.Deserialize(Archive);
+}
diff --git a/Engine/Source/Runtime/Datasmith/CADKernel/Public/CADKernel/Core/CADKernelArchive.h b/Engine/Source/Runtime/Datasmith/CADKernel/Public/CADKernel/Core/CADKernelArchive.h
--- a/Engine/Source/Runtime/Datasmith/CADKernel/Public/CADKernel/Core/CADKernelArchive.h
+++ b/Engine/Source/Runtime/Datasmith/CADKernel/Public/CADKernel/Core/CADKernelArchive.h
@@ -6,6 +6,7 @@
 
 #include "HAL/FileManager.h"
 #include "Serialization/Archive.h"
+#include "Serialization/MemoryReader.h"
 
 namespace CADKernel
 {
@@ -16,6 +17,7 @@ namespace CADKernel
 	class FCADKernelArchive
 	{
 	public:
+		TUniquePtr<FMemoryReader> MemoryReader;
 		FArchive& Archive;
 		FSession& Session;
 		FModel* ArchiveModel = nullptr;
@@ -26,6 +28,16 @@ namespace CADKernel
 		{
 		}
 
+		/**
+		 * Archive reader constructor
+		 */
+		FCADKernelArchive(FSession& InSession, const TArray<uint8>& InRawData)
+			: MemoryReader(MakeUnique<FMemoryReader>(InRawData, true))
+			, Archive(*MemoryReader)
+			, Session(InSession)
+		{
+		}
+
 		static TSharedPtr<FCADKernelArchive> CreateArchiveWriter(FSession& InSession, const TCHAR* FileName)
 		{
 			FArchive* Ar = IFileManager::Get().CreateFileWriter(FileName);
diff --git a/Engine/Source/Runtime/Datasmith/CADKernel/Public/CADKernel/Core/Session.h b/Engine/Source/Runtime/Datasmith/CADKernel/Public/CADKernel/Core/Session.h
--- a/Engine/Source/Runtime/Datasmith/CADKernel/Public/CADKernel/Core/Session.h
+++ b/Engine/Source/Runtime/Datasmith/CADKernel/Public/CADKernel/Core/Session.h
@@ -74,11 +74,17 @@ namespace CADKernel
 		}
 
 		/**
-		 * load and add a database in the current session database
+		 * Load and add a database in the current session database
 		 * Entity ID is set for all loaded entities
 		 */
 		void LoadDatabase(const TCHAR* FilePath);
 
+		/**
+		 * Add a database defined by a RawData in the current session database
+		 * Entity ID is set for all loaded entities
+		 */
+		void AddDatabase(const TArray<uint8>& InRawData);
+
 		void Clear()
 		{
 			Model.Reset();
