Improve Wire translator
 - Simpily the definition of an Uuid for AlDagNode,
 - Use this Uuid for Actor name,
 - Use StaticMeshActor name for ct file name,

diff --git a/Engine/Plugins/Enterprise/DatasmithCADImporter/Source/DatasmithWireTranslator/Private/DatasmithWireTranslator.cpp b/Engine/Plugins/Enterprise/DatasmithCADImporter/Source/DatasmithWireTranslator/Private/DatasmithWireTranslator.cpp
--- a/Engine/Plugins/Enterprise/DatasmithCADImporter/Source/DatasmithWireTranslator/Private/DatasmithWireTranslator.cpp
+++ b/Engine/Plugins/Enterprise/DatasmithCADImporter/Source/DatasmithWireTranslator/Private/DatasmithWireTranslator.cpp
@@ -2,6 +2,7 @@
 
 #include "DatasmithWireTranslator.h"
 
+#include "CADOptions.h"
 #include "Containers/List.h"
 #include "DatasmithImportOptions.h"
 #include "DatasmithSceneFactory.h"
@@ -62,7 +63,6 @@ DEFINE_LOG_CATEGORY_STATIC(LogDatasmithWireTranslator, Log, All);
 
 #ifdef USE_OPENMODEL
 
-using namespace OpenModelUtils;
 using namespace CADLibrary;
 
 class BodyData
@@ -82,49 +82,25 @@ public:
 	bool bCadData = true;
 
 	// Generates BodyData's unique id from AlDagNode objects
-	FString GetUUID(const FString& ParentUUID)
+	uint32 GetUuid(const uint32& ParentUuid)
 	{
 		if(ShellSet.Num() == 0)
 		{
-			return ParentUUID;
+			return ParentUuid;
 		}
 
-		auto GetLongPersistentID = []( AlDagNode& DagNode ) -> int64
+		if (ShellSet.Num() > 1)
 		{
-			union {
-				int a[2];
-				int64 b;
-			} Value;
-
-			Value.b = -1;
-
-			AlPersistentID* PersistentID = nullptr;
-			DagNode.persistentID( PersistentID );
-
-			if( PersistentID != nullptr )
-			{
-				int Dummy;
-				PersistentID->id( Value.a[0], Value.a[1], Dummy, Dummy );
-			}
-
-			return Value.b;
-		};
-
-		if(ShellSet.Num() > 1)
-		{
-			ShellSet.Sort([&](AlDagNode& A, AlDagNode& B)
-			{
-				return GetLongPersistentID( A ) < GetLongPersistentID( B );
-			});
+			ShellSet.Sort([&](AlDagNode& NodeA, AlDagNode& NodeB) { return OpenModelUtils::GetAlDagNodeUuid(NodeA) < OpenModelUtils::GetAlDagNodeUuid(NodeB); });
 		}
 
-		FString Buffer;
-		for(AlDagNode* DagNode : ShellSet)
+		uint32 BodyUUID = 0;
+		for (AlDagNode* DagNode : ShellSet)
 		{
-			Buffer += FString::Printf(TEXT("%016lx"), GetLongPersistentID( *DagNode ) );
+			BodyUUID = HashCombine(BodyUUID, OpenModelUtils::GetAlDagNodeUuid(*DagNode));
 		}
 
-		return GetUEUUIDFromAIPersistentID( ParentUUID, Buffer );
+		return HashCombine(ParentUuid, BodyUUID);
 	}
 };
 
@@ -191,7 +167,7 @@ private:
 
 	struct FDagNodeInfo
 	{
-		FString UEuuid;  // Use for actor name
+		uint32 Uuid;  // Use for actor name
 		FString Label;
 		TSharedPtr< IDatasmithActorElement > ActorElement;
 	};
@@ -214,7 +190,7 @@ private:
 	TOptional<FMeshDescription> GetMeshOfShellBody(TSharedRef<BodyData> DagNode, TSharedRef<IDatasmithMeshElement> MeshElement, CADLibrary::FMeshParameters& MeshParameters);
 	TOptional<FMeshDescription> GetMeshOfMeshBody(TSharedRef<BodyData> DagNode, TSharedRef<IDatasmithMeshElement> MeshElement, CADLibrary::FMeshParameters& MeshParameters);
 
-	void AddNodeInBodySet(AlDagNode& DagNode, const char* ShaderName, TMap<uint32, TSharedPtr<BodyData>>& ShellToProcess, bool bIsAPatch, uint32 MaxSize);
+	void AddNodeInBodyGroup(AlDagNode& DagNode, const char* ShaderName, TMap<uint32, TSharedPtr<BodyData>>& ShellToProcess, bool bIsAPatch, uint32 MaxSize);
 
 	TOptional<FMeshDescription> MeshDagNodeWithExternalMesher(AlDagNode& DagNode, TSharedRef<IDatasmithMeshElement> MeshElement, CADLibrary::FMeshParameters& MeshParameters);
 	TOptional<FMeshDescription> MeshDagNodeWithExternalMesher(TSharedRef<BodyData> DagNode, TSharedRef<IDatasmithMeshElement> MeshElement, CADLibrary::FMeshParameters& MeshParameters);
@@ -285,15 +261,15 @@ private:
 	AlDagNode* AlRootNode;
 
 	/** Table of correspondence between mesh identifier and associated Datasmith mesh element */
-	TMap< uint32, TSharedPtr< IDatasmithMeshElement > > ShellUUIDToMeshElementMap;
-	TMap< FString, TSharedPtr< IDatasmithMeshElement > > BodyToMeshElementMap;
+	TMap<uint32, TSharedPtr<IDatasmithMeshElement>> ShellUuidToMeshElementMap;
+	TMap<uint32, TSharedPtr<IDatasmithMeshElement>> BodyUuidToMeshElementMap;
 
 	/** Datasmith mesh elements to OpenModel objects */
-	TMap< IDatasmithMeshElement*, AlDagNode* > MeshElementToAlDagNodeMap;
+	TMap<IDatasmithMeshElement*, AlDagNode*> MeshElementToAlDagNodeMap;
 
-	TMap< IDatasmithMeshElement*, TSharedPtr<BodyData>> MeshElementToBodyMap;
+	TMap<IDatasmithMeshElement*, TSharedPtr<BodyData>> MeshElementToBodyMap;
 
-	TMap <FString, TSharedPtr< IDatasmithMaterialIDElement >> ShaderNameToUEMaterialId;
+	TMap<FString, TSharedPtr<IDatasmithMaterialIDElement>> ShaderNameToUEMaterialId;
 
 	// start section information
 	int32 FileVersion;
@@ -1236,37 +1212,27 @@ void FWireTranslatorImpl::GetDagNodeMeta(AlDagNode& CurrentNode, TSharedPtr< IDa
 
 }
 
-
 void FWireTranslatorImpl::GetDagNodeInfo(AlDagNode& CurrentNode, const FDagNodeInfo& ParentInfo, FDagNodeInfo& CurrentNodeInfo)
 {
-	CurrentNodeInfo.Label = CurrentNode.name();
-
-	AlPersistentID* GroupNodeId = nullptr;
-	CurrentNode.persistentID(GroupNodeId);
-
-	FString ThisGroupNodeID( GetPersistentIDString(GroupNodeId) );
-
-    // Limit length of UUID by combining hash of parent UUID and container's UUID if ParentUuid is not empty
-	CurrentNodeInfo.UEuuid = GetUEUUIDFromAIPersistentID(ParentInfo.UEuuid, ThisGroupNodeID);
+    CurrentNodeInfo.Label = CurrentNode.name();
+	uint32 ThisGroupNodeUuid = OpenModelUtils::GetAlDagNodeUuid(CurrentNode);
+	CurrentNodeInfo.Uuid = HashCombine(ParentInfo.Uuid, ThisGroupNodeUuid);
 }
 
 void FWireTranslatorImpl::GetDagNodeInfo(TSharedRef<BodyData> CurrentNode, const FDagNodeInfo& ParentInfo, FDagNodeInfo& CurrentNodeInfo)
 {
-	CurrentNodeInfo.Label = ParentInfo.Label; // +TEXT(" ") + CurrentNode->LayerName + TEXT(" ") + CurrentNode->ShaderName;
+	CurrentNodeInfo.Label = ParentInfo.Label;
 	CurrentNode->Label = CurrentNodeInfo.Label;
-
-	// Limit length of UUID by combining hash of parent UUID and container's UUID if ParentUuid is not empty
-	CurrentNodeInfo.UEuuid = GetUEUUIDFromAIPersistentID( ParentInfo.UEuuid, CurrentNode->GetUUID( ParentInfo.UEuuid ) );
+	CurrentNodeInfo.Uuid = CurrentNode->GetUuid(ParentInfo.Uuid);
 }
 
-
 bool FWireTranslatorImpl::ProcessAlGroupNode(AlGroupNode& GroupNode, const FDagNodeInfo& ParentInfo)
 {
 	FDagNodeInfo ThisGroupNodeInfo;
 
 	GetDagNodeInfo(GroupNode, ParentInfo, ThisGroupNodeInfo);
 
-	ThisGroupNodeInfo.ActorElement = FDatasmithSceneFactory::CreateActor(*ThisGroupNodeInfo.UEuuid);
+	ThisGroupNodeInfo.ActorElement = FDatasmithSceneFactory::CreateActor(*OpenModelUtils::UuidToString(ThisGroupNodeInfo.Uuid));
 	ThisGroupNodeInfo.ActorElement->SetLabel(*ThisGroupNodeInfo.Label);
 	GetDagNodeMeta(GroupNode, ThisGroupNodeInfo.ActorElement);
 
@@ -1277,10 +1243,10 @@ bool FWireTranslatorImpl::ProcessAlGroupNode(AlGroupNode& GroupNode, const FDagN
 	}
 
 	// add the resulting actor to the scene
-	if (IsValidActor(ThisGroupNodeInfo.ActorElement))
+	if (OpenModelUtils::IsValidActor(ThisGroupNodeInfo.ActorElement))
 	{
 		// Apply local transform to actor element
-		SetActorTransform(ThisGroupNodeInfo.ActorElement, GroupNode);
+		OpenModelUtils::SetActorTransform(ThisGroupNodeInfo.ActorElement, GroupNode);
 
 		if (ParentInfo.ActorElement.IsValid())
 		{
@@ -1297,13 +1263,13 @@ bool FWireTranslatorImpl::ProcessAlGroupNode(AlGroupNode& GroupNode, const FDagN
 
 TSharedPtr< IDatasmithMeshElement > FWireTranslatorImpl::FindOrAddMeshElement(TSharedRef<BodyData> Body, const FDagNodeInfo& NodeInfo)
 {
-	TSharedPtr< IDatasmithMeshElement >* MeshElementPtr = BodyToMeshElementMap.Find( NodeInfo.UEuuid );
+	TSharedPtr< IDatasmithMeshElement >* MeshElementPtr = BodyUuidToMeshElementMap.Find( NodeInfo.Uuid );
 	if (MeshElementPtr != nullptr)
 	{
 		return *MeshElementPtr;
 	}
 	
-	TSharedPtr< IDatasmithMeshElement > MeshElement = FDatasmithSceneFactory::CreateMesh( *NodeInfo.UEuuid );
+	TSharedPtr< IDatasmithMeshElement > MeshElement = FDatasmithSceneFactory::CreateMesh(*OpenModelUtils::UuidToString(NodeInfo.Uuid));
 	MeshElement->SetLabel(*NodeInfo.Label);
 	MeshElement->SetLightmapSourceUV(-1);
 
@@ -1315,20 +1281,20 @@ TSharedPtr< IDatasmithMeshElement > FWireTranslatorImpl::FindOrAddMeshElement(TS
 
 	DatasmithScene->AddMesh(MeshElement);
 
-	ShellUUIDToMeshElementMap.Add(FCString::Atoi(*NodeInfo.UEuuid), MeshElement);
+	ShellUuidToMeshElementMap.Add(NodeInfo.Uuid, MeshElement);
 	MeshElementToBodyMap.Add(MeshElement.Get(), Body);
 
-	BodyToMeshElementMap.Add( NodeInfo.UEuuid, MeshElement );
+	BodyUuidToMeshElementMap.Add( NodeInfo.Uuid, MeshElement );
 
 	return MeshElement;
 }
 
 TSharedPtr< IDatasmithMeshElement > FWireTranslatorImpl::FindOrAddMeshElement(AlDagNode& ShellNode, const FDagNodeInfo& ParentInfo, const char* ShaderName)
 {
-	uint32 ShellUUID = GetUUIDFromAIPersistentID(ShellNode);
+	uint32 ShellUuid = OpenModelUtils::GetAlDagNodeUuid(ShellNode);
 
 	// Look if geometry has not been already processed, return it if found
-	TSharedPtr< IDatasmithMeshElement >* MeshElementPtr = ShellUUIDToMeshElementMap.Find(ShellUUID);
+	TSharedPtr< IDatasmithMeshElement >* MeshElementPtr = ShellUuidToMeshElementMap.Find(ShellUuid);
 	if (MeshElementPtr != nullptr)
 	{
 		return *MeshElementPtr;
@@ -1337,26 +1303,15 @@ TSharedPtr< IDatasmithMeshElement > FWireTranslatorImpl::FindOrAddMeshElement(Al
 	FDagNodeInfo MeshInfo;
 	GetDagNodeInfo(ShellNode, ParentInfo, MeshInfo);
 
-	TSharedPtr< IDatasmithMeshElement > MeshElement = FDatasmithSceneFactory::CreateMesh(*MeshInfo.UEuuid);
+	TSharedPtr< IDatasmithMeshElement > MeshElement = FDatasmithSceneFactory::CreateMesh(*OpenModelUtils::UuidToString(MeshInfo.Uuid));
 	MeshElement->SetLabel(*MeshInfo.Label);
 	MeshElement->SetLightmapSourceUV(-1);
 
-	// TODO
-	// Get bounding box saved by GPure
-	double BoundingBox[8][4];
-	ShellNode.boundingBox(BoundingBox);
-	//float BoundingBox[6];
-	//FString Buffer = GetStringAttribute(GeomID, TEXT("UE_MESH_BBOX"));
-	//if (FString::ToHexBlob(Buffer, (uint8*)BoundingBox, sizeof(BoundingBox)))
-	//{
-	//	MeshElement->SetDimensions(BoundingBox[3] - BoundingBox[0], BoundingBox[4] - BoundingBox[1], BoundingBox[5] - BoundingBox[2], 0.0f);
-	//}
-
 	// Set MeshElement FileHash used for re-import task 
 	FMD5 MD5; // unique Value that define the mesh
 	MD5.Update(reinterpret_cast<const uint8*>(&SceneFileHash), sizeof SceneFileHash);
 	// MeshActor Name
-	MD5.Update(reinterpret_cast<const uint8*>(&ShellUUID), sizeof ShellUUID);
+	MD5.Update(reinterpret_cast<const uint8*>(&ShellUuid), sizeof ShellUuid);
 	FMD5Hash Hash;
 	Hash.Set(MD5);
 	MeshElement->SetFileHash(Hash);
@@ -1369,7 +1324,7 @@ TSharedPtr< IDatasmithMeshElement > FWireTranslatorImpl::FindOrAddMeshElement(Al
 
 	DatasmithScene->AddMesh(MeshElement);
 
-	ShellUUIDToMeshElementMap.Add(ShellUUID, MeshElement);
+	ShellUuidToMeshElementMap.Add(ShellUuid, MeshElement);
 	MeshElementToAlDagNodeMap.Add(MeshElement.Get(), &ShellNode);
 
 	return MeshElement;
@@ -1386,7 +1341,7 @@ bool FWireTranslatorImpl::ProcessAlShellNode(AlDagNode& ShellNode, const FDagNod
 		return false;
 	}
 
-	TSharedPtr< IDatasmithMeshActorElement > ActorElement = FDatasmithSceneFactory::CreateMeshActor(*ShellInfo.UEuuid);
+	TSharedPtr< IDatasmithMeshActorElement > ActorElement = FDatasmithSceneFactory::CreateMeshActor(*OpenModelUtils::UuidToString(ShellInfo.Uuid));
 	if (!ActorElement.IsValid())
 	{
 		return false;
@@ -1398,7 +1353,7 @@ bool FWireTranslatorImpl::ProcessAlShellNode(AlDagNode& ShellNode, const FDagNod
 
 	GetDagNodeMeta(ShellNode, ActorElement);
 
-	SetActorTransform(ShellInfo.ActorElement, ShellNode);
+	OpenModelUtils::SetActorTransform(ShellInfo.ActorElement, ShellNode);
 
 	// Apply materials on the current part
 	if (ShaderName)
@@ -1413,7 +1368,7 @@ bool FWireTranslatorImpl::ProcessAlShellNode(AlDagNode& ShellNode, const FDagNod
 		}
 	}
 
-	if (ActorElement.IsValid() && IsValidActor(ActorElement))
+	if (ActorElement.IsValid() && OpenModelUtils::IsValidActor(ActorElement))
 	{
 		if (ParentInfo.ActorElement.IsValid())
 		{
@@ -1443,7 +1398,7 @@ bool FWireTranslatorImpl::ProcessBodyNode(TSharedRef<BodyData> Body, const FDagN
 		return false;
 	}
 
-	TSharedPtr< IDatasmithMeshActorElement > ActorElement = FDatasmithSceneFactory::CreateMeshActor(*ShellInfo.UEuuid);
+	TSharedPtr< IDatasmithMeshActorElement > ActorElement = FDatasmithSceneFactory::CreateMeshActor(*OpenModelUtils::UuidToString(ShellInfo.Uuid));
 	if (!ActorElement.IsValid())
 	{
 		return false;
@@ -1470,7 +1425,7 @@ bool FWireTranslatorImpl::ProcessBodyNode(TSharedRef<BodyData> Body, const FDagN
 		}
 	}
 
-	if (ActorElement.IsValid() && IsValidActor(ActorElement))
+	if (ActorElement.IsValid() && OpenModelUtils::IsValidActor(ActorElement))
 	{
 		if (ParentInfo.ActorElement.IsValid())
 		{
@@ -1529,26 +1484,26 @@ bool IsHidden(AlDagNode * DagNode)
 }
 
 
-uint32 getBodySetId(const char* ShaderName, const char* LayerName, bool bCadData)
+uint32 GetBodyGroupUuid(const char* ShaderName, const char* LayerName, bool bCadData)
 {
-	uint32 UUID = HashCombine(GetTypeHash(ShaderName), GetTypeHash(bCadData));
-	UUID = HashCombine(GetTypeHash(LayerName), UUID);
-	return UUID;
+	uint32 Uuid = HashCombine(GetTypeHash(ShaderName), GetTypeHash(bCadData));
+	Uuid = HashCombine(GetTypeHash(LayerName), Uuid);
+	return Uuid;
 }
 
-uint32 getNumOfPatch(AlShell& Shell)
+uint32 GetPatchCount(AlShell& Shell)
 {
-	uint32 NbFace = 0;
+	uint32 PatchCount = 0;
 	AlTrimRegion *TrimRegion = Shell.firstTrimRegion();
 	while (TrimRegion)
 	{
-		NbFace++;
+		PatchCount++;
 		TrimRegion = TrimRegion->nextRegion();
 	}
-	return NbFace;
+	return PatchCount;
 }
 
-void FWireTranslatorImpl::AddNodeInBodySet(AlDagNode& DagNode, const char* ShaderName, TMap<uint32, TSharedPtr<BodyData>>& ShellToProcess, bool bIsAPatch, uint32 MaxSize)
+void FWireTranslatorImpl::AddNodeInBodyGroup(AlDagNode& DagNode, const char* ShaderName, TMap<uint32, TSharedPtr<BodyData>>& ShellToProcess, bool bIsAPatch, uint32 MaxSize)
 {
 	const char* LayerName = nullptr;
 	if (AlLayer* Layer = DagNode.layer())
@@ -1556,7 +1511,7 @@ void FWireTranslatorImpl::AddNodeInBodySet(AlDagNode& DagNode, const char* Shade
 		LayerName = Layer->name();
 	}
 
-	uint32 SetId = getBodySetId(ShaderName, LayerName, bIsAPatch);
+	uint32 SetId = GetBodyGroupUuid(ShaderName, LayerName, bIsAPatch);
 
 	TSharedPtr<BodyData> Body;
 	if (TSharedPtr<BodyData>* PBody = ShellToProcess.Find(SetId))
@@ -1617,7 +1572,7 @@ bool FWireTranslatorImpl::RecurseDagForLeaves(AlDagNode* FirstDagNode, const FDa
 		{
 			AlShellNode *ShellNode = DagNode->asShellNodePtr();
 			AlShell *Shell = ShellNode->shell();
-			uint32 NbPatch = getNumOfPatch(*Shell);
+			uint32 NbPatch = GetPatchCount(*Shell);
 
 			if (AlShader* Shader = Shell->firstShader())
 			{
@@ -1626,7 +1581,7 @@ bool FWireTranslatorImpl::RecurseDagForLeaves(AlDagNode* FirstDagNode, const FDa
 
 			if (NbPatch == 1)
 			{
-				AddNodeInBodySet(*DagNode, ShaderName, ShellToProcess, true, MaxSize);
+				AddNodeInBodyGroup(*DagNode, ShaderName, ShellToProcess, true, MaxSize);
 			}
 			else
 			{
@@ -1642,7 +1597,7 @@ bool FWireTranslatorImpl::RecurseDagForLeaves(AlDagNode* FirstDagNode, const FDa
 			{
 				ShaderName = Shader->name();
 			}
-			AddNodeInBodySet(*DagNode, ShaderName, ShellToProcess, true, MaxSize);
+			AddNodeInBodyGroup(*DagNode, ShaderName, ShellToProcess, true, MaxSize);
 			break;
 		}
 
@@ -1654,7 +1609,7 @@ bool FWireTranslatorImpl::RecurseDagForLeaves(AlDagNode* FirstDagNode, const FDa
 			{
 				ShaderName = Shader->name();
 			}
-			AddNodeInBodySet(*DagNode, ShaderName, ShellToProcess, false, MaxSize);
+			AddNodeInBodyGroup(*DagNode, ShaderName, ShellToProcess, false, MaxSize);
 			break;
 		}
 
@@ -1779,8 +1734,6 @@ TOptional<FMeshDescription> FWireTranslatorImpl::MeshDagNodeWithExternalMesher(A
 	// Wire unit is cm
 	LocalSession->SetSceneUnit(0.01);
 	
-	FString Filename = DagNode.name();
-
 	EAliasObjectReference ObjectReference = EAliasObjectReference::LocalReference;
 
 	if (MeshParameters.bIsSymmetric)
@@ -1793,9 +1746,7 @@ TOptional<FMeshDescription> FWireTranslatorImpl::MeshDagNodeWithExternalMesher(A
 	DagNodeSet.Add(&DagNode);
 	LocalSession->AddBRep(DagNodeSet, ObjectReference);
 
-	Filename += TEXT(".ct");
-
-	FString FilePath = FPaths::Combine(OutputPath, Filename);
+	FString FilePath = LocalSession->GetImportParameters().DefineCADFilePath(*OutputPath, MeshElement->GetName());
 	if (LocalSession->SaveBrep(FilePath))
 	{
 		MeshElement->SetFile(*FilePath);
@@ -1830,9 +1781,7 @@ TOptional<FMeshDescription> FWireTranslatorImpl::MeshDagNodeWithExternalMesher(T
 
 	LocalSession->AddBRep(Body->ShellSet, ObjectReference);
 
-	FString Filename = Body->Label + TEXT(".ct");
-
-	FString FilePath = FPaths::Combine(OutputPath, Filename);
+	FString FilePath = LocalSession->GetImportParameters().DefineCADFilePath(*OutputPath, MeshElement->GetName());
 	if (LocalSession->SaveBrep(FilePath))
 	{
 		MeshElement->SetFile(*FilePath);
@@ -1858,7 +1807,7 @@ TOptional<FMeshDescription> FWireTranslatorImpl::GetMeshOfShellNode(AlDagNode& D
 		AlMatrix4x4 AlMatrix;
 		DagNode.inverseGlobalTransformationMatrix(AlMatrix);
 		// TODO: the best way, should be to don't have to apply inverse global transform to the generated mesh
-		AlDagNode *TesselatedNode = TesselateDagLeaf(&DagNode, ETesselatorType::Fast, TessellationOptions.ChordTolerance);
+		AlDagNode *TesselatedNode = OpenModelUtils::TesselateDagLeaf(&DagNode, ETesselatorType::Fast, TessellationOptions.ChordTolerance);
 		if (TesselatedNode != nullptr)
 		{
 			// Get the meshes from the dag nodes. Note that removing the mesh's DAG.
@@ -1891,7 +1840,7 @@ TOptional<FMeshDescription> FWireTranslatorImpl::GetMeshOfMeshBody(TSharedRef<Bo
 		AlMesh *Mesh = MeshNode->mesh();
 		if (Mesh)
 		{
-			TransferAlMeshToMeshDescription(*Mesh, MeshDescription, MeshParameters, True, true);
+			OpenModelUtils::TransferAlMeshToMeshDescription(*Mesh, MeshDescription, MeshParameters, True, true);
 		}
 	}
 
@@ -2001,6 +1950,9 @@ TOptional<FMeshDescription> FWireTranslatorImpl::GetMeshDescription(TSharedRef<I
 			return GetMeshOfNodeMesh(*(DagNode.asMeshNodePtr()), MeshElement, MeshParameters);
 			break;
 		}
+
+		default:
+			break;
 	}
 
 	return TOptional< FMeshDescription >();
@@ -2014,7 +1966,7 @@ TOptional< FMeshDescription > FWireTranslatorImpl::ImportMesh(AlMesh& CurrentMes
 	FMeshDescription MeshDescription;
 	DatasmithMeshHelper::PrepareAttributeForStaticMesh(MeshDescription);
 	bool True = true;
-	TransferAlMeshToMeshDescription(CurrentMesh, MeshDescription, MeshParameters, True);
+	OpenModelUtils::TransferAlMeshToMeshDescription(CurrentMesh, MeshDescription, MeshParameters, True);
 
 	return MoveTemp(MeshDescription);
 }
diff --git a/Engine/Plugins/Enterprise/DatasmithCADImporter/Source/DatasmithWireTranslator/Private/OpenModelUtils.cpp b/Engine/Plugins/Enterprise/DatasmithCADImporter/Source/DatasmithWireTranslator/Private/OpenModelUtils.cpp
--- a/Engine/Plugins/Enterprise/DatasmithCADImporter/Source/DatasmithWireTranslator/Private/OpenModelUtils.cpp
+++ b/Engine/Plugins/Enterprise/DatasmithCADImporter/Source/DatasmithWireTranslator/Private/OpenModelUtils.cpp
@@ -11,7 +11,6 @@
 #include "AlDagNode.h"
 #include "AlMesh.h"
 #include "AlLayer.h"
-#include "AlPersistentID.h"
 #include "AlTesselate.h"
 #include "MeshAttributes.h"
 #include "MeshDescription.h"
@@ -1600,45 +1599,6 @@ bool OpenModelUtils::IsValidActor(const TSharedPtr< IDatasmithActorElement >& Ac
 	return false;
 }
 
-
-FString OpenModelUtils::GetPersistentIDString(AlPersistentID* GroupNodeId)
-{
-	if(GroupNodeId == nullptr)
-	{
-		return FString();
-	}
-
-	int IdA, IdB, IdC, IdD;
-	GroupNodeId->id(IdA, IdB, IdC, IdD);
-	FString ThisGroupNodeID = FString::FromInt(IdA) + FString::FromInt(IdB) + FString::FromInt(IdC) + FString::FromInt(IdD);
-	return ThisGroupNodeID;
-}
-
-uint32 OpenModelUtils::GetUUIDFromAIPersistentID(AlPersistentID* GroupNodeId)
-{
-	FString ThisGroupNodeID = GetPersistentIDString(GroupNodeId);
-	return GetTypeHash(*ThisGroupNodeID);
-}
-
-uint32 OpenModelUtils::GetUUIDFromAIPersistentID(AlDagNode& GroupNode)
-{
-	AlPersistentID* ShellNodeId = new AlPersistentID();
-	if (GroupNode.hasPersistentID() == sSuccess)
-	{
-		GroupNode.persistentID(ShellNodeId);
-		return OpenModelUtils::GetUUIDFromAIPersistentID(ShellNodeId);
-	}
-	return 0;
-}
-
-FString OpenModelUtils::GetUEUUIDFromAIPersistentID(const FString& ParentUEuuid, const FString& CurrentNodePersistentID)
-{
-	// Limit length of UUID by combining hash of parent UUID and container's UUID if ParentUuid is not empty
-	return FString::Printf(TEXT("0x%08x"), HashCombine(GetTypeHash(ParentUEuuid), GetTypeHash(CurrentNodePersistentID)));
-}
-
-
-
 bool OpenModelUtils::TransferAlMeshToMeshDescription(const AlMesh& AliasMesh, FMeshDescription& MeshDescription, CADLibrary::FMeshParameters& MeshParameters, bool& bHasNormal, bool bMerge)
 {
 	// Ref. GP3DMVisitorImpl::visitMesh
diff --git a/Engine/Plugins/Enterprise/DatasmithCADImporter/Source/DatasmithWireTranslator/Private/OpenModelUtils.h b/Engine/Plugins/Enterprise/DatasmithCADImporter/Source/DatasmithWireTranslator/Private/OpenModelUtils.h
--- a/Engine/Plugins/Enterprise/DatasmithCADImporter/Source/DatasmithWireTranslator/Private/OpenModelUtils.h
+++ b/Engine/Plugins/Enterprise/DatasmithCADImporter/Source/DatasmithWireTranslator/Private/OpenModelUtils.h
@@ -10,6 +10,9 @@
 
 #include "AlAccessTypes.h"
 #include "AlShadingFields.h"
+#include "AlDagNode.h"
+#include "AlPersistentID.h"
+
 
 class IDatasmithActorElement;
 class AlDagNode;
@@ -43,10 +46,29 @@ namespace OpenModelUtils
 
 	bool IsValidActor(const TSharedPtr< IDatasmithActorElement >& ActorElement);
 
-	uint32 GetUUIDFromAIPersistentID(AlPersistentID* GroupNodeId);
-	uint32 GetUUIDFromAIPersistentID(AlDagNode& GroupNode);
-	FString GetPersistentIDString(AlPersistentID* GroupNodeId);
-	FString GetUEUUIDFromAIPersistentID(const FString& ParentUEuuid, const FString& CurrentNodePersistentID);
+	inline FString UuidToString(const uint32& Uuid)
+	{
+		return FString::Printf(TEXT("0x%08x"), Uuid);
+	}
+
+	inline uint32 GetTypeHash(AlPersistentID& GroupNodeId)
+	{
+		int IdA, IdB, IdC, IdD;
+		GroupNodeId.id(IdA, IdB, IdC, IdD);
+		return HashCombine(IdA, HashCombine(IdB, HashCombine(IdC, IdD)));
+	}
+
+	inline uint32 GetAlDagNodeUuid(AlDagNode& GroupNode)
+	{
+		if (GroupNode.hasPersistentID() == sSuccess)
+		{
+			AlPersistentID* PersistentID;
+			GroupNode.persistentID(PersistentID);
+			return GetTypeHash(*PersistentID);
+		}
+		FString Label = GroupNode.name();
+		return GetTypeHash(Label);
+	}
 
 	// Note that Alias file unit is cm like UE
 	bool TransferAlMeshToMeshDescription(const AlMesh& Mesh, FMeshDescription& MeshDescription, CADLibrary::FMeshParameters& SymmetricParameters, bool& bHasNormal, bool bMerge = false);
