Update naming of datasmith CVar

diff --git a/Engine/Config/ConsoleVariables.ini b/Engine/Config/ConsoleVariables.ini
--- a/Engine/Config/ConsoleVariables.ini
+++ b/Engine/Config/ConsoleVariables.ini
@@ -101,3 +101,20 @@ fx.SkipVectorVMBackendOptimizations=1
  
 ; Comment this out to turn off audio submix effects.
 ;au.BypassAllSubmixEffects=0
+
+
+; Datasmith CVar
+; Uncomment to set our own values
+
+; Activate to save temporary CAD processing file. These file will be use in a next import to avoid CAD file processing
+;ds.CADTranslator.EnableCADCache=1
+
+; The timer that kill the worker if the import time is unusually long. With this time control, the load of the corrupted file is canceled but the rest of the scene is imported
+;ds.CADTranslator.EnableTimeControl=1
+
+; Activate to parallelise CAD file processing
+;ds.CADTranslator.EnableThreadedImport=1
+
+; Activate to use KernelIO Tessellation
+;ds.CADTranslator.EnableKernelIOTessellation=1
+
diff --git a/Engine/Plugins/Enterprise/DatasmithCADImporter/Source/DatasmithCADTranslator/Private/DatasmithCADTranslator.cpp b/Engine/Plugins/Enterprise/DatasmithCADImporter/Source/DatasmithCADTranslator/Private/DatasmithCADTranslator.cpp
--- a/Engine/Plugins/Enterprise/DatasmithCADImporter/Source/DatasmithCADTranslator/Private/DatasmithCADTranslator.cpp
+++ b/Engine/Plugins/Enterprise/DatasmithCADImporter/Source/DatasmithCADTranslator/Private/DatasmithCADTranslator.cpp
@@ -12,13 +12,13 @@
 #include "DatasmithUtils.h"
 
 static TAutoConsoleVariable<int32> CVarStaticCADTranslatorEnableThreadedImport(
-	TEXT("DS.CADTranslator.EnableThreadedImport"),
+	TEXT("ds.CADTranslator.EnableThreadedImport"),
 	1,
 	TEXT("Activate to parallelise CAD file processing.\n"),
 	ECVF_Default);
 
 static TAutoConsoleVariable<int32> CVarStaticCADTranslatorEnableKernelIOTessellation(
-	TEXT("DS.CADTranslator.EnableKernelIOTessellation"),
+	TEXT("ds.CADTranslator.EnableKernelIOTessellation"),
 	1,
 	TEXT("Activate to use KernelIO Tessellation.\n"),
 	ECVF_Default);
@@ -85,9 +85,6 @@ void FDatasmithCADTranslator::Initialize(FDatasmithTranslatorCapabilities& OutCa
 bool FDatasmithCADTranslator::LoadScene(TSharedRef<IDatasmithScene> DatasmithScene)
 {
 	ImportParameters.bEnableKernelIOTessellation = (CVarStaticCADTranslatorEnableKernelIOTessellation.GetValueOnAnyThread() != 0);;
-#ifdef CADKERNEL_DEBUG
-	ImportParameters.bEnableKernelIOTessellation = false;
-#endif
 
 	ImportParameters.MetricUnit = 0.001;
 	ImportParameters.ScaleFactor = 0.1;
@@ -145,10 +142,6 @@ bool FDatasmithCADTranslator::LoadScene(TSharedRef<IDatasmithScene> DatasmithSce
 	{
 		bool bWithProcessor = (CVarStaticCADTranslatorEnableThreadedImport.GetValueOnAnyThread() != 0);
 
-#ifdef CAD_TRANSLATOR_DEBUG
-		bWithProcessor = false;
-#endif //CAD_TRANSLATOR_DEBUG
-
 		TMap<uint32, FString> CADFileToUEFileMap;
 		int32 NumCores = FPlatformMisc::NumberOfCores();
 		{
diff --git a/Engine/Plugins/Enterprise/DatasmithCADImporter/Source/DatasmithCADTranslator/Private/DatasmithCADTranslatorModule.cpp b/Engine/Plugins/Enterprise/DatasmithCADImporter/Source/DatasmithCADTranslator/Private/DatasmithCADTranslatorModule.cpp
--- a/Engine/Plugins/Enterprise/DatasmithCADImporter/Source/DatasmithCADTranslator/Private/DatasmithCADTranslatorModule.cpp
+++ b/Engine/Plugins/Enterprise/DatasmithCADImporter/Source/DatasmithCADTranslator/Private/DatasmithCADTranslatorModule.cpp
@@ -26,7 +26,7 @@ void FDatasmithCADTranslatorModule::StartupModule()
 	}
 
 	bool bCreateCacheFolder = false;
-	if (IConsoleVariable* CVar = IConsoleManager::Get().FindConsoleVariable(TEXT("DS.CADTranslator.EnableCADCache")))
+	if (IConsoleVariable* CVar = IConsoleManager::Get().FindConsoleVariable(TEXT("ds.CADTranslator.EnableCADCache")))
 	{
 		bCreateCacheFolder = CVar->GetInt() != 0;
 	}
diff --git a/Engine/Plugins/Enterprise/DatasmithCADImporter/Source/DatasmithDispatcher/Private/DatasmithDispatcher.cpp b/Engine/Plugins/Enterprise/DatasmithCADImporter/Source/DatasmithDispatcher/Private/DatasmithDispatcher.cpp
--- a/Engine/Plugins/Enterprise/DatasmithCADImporter/Source/DatasmithDispatcher/Private/DatasmithDispatcher.cpp
+++ b/Engine/Plugins/Enterprise/DatasmithCADImporter/Source/DatasmithDispatcher/Private/DatasmithDispatcher.cpp
@@ -14,14 +14,14 @@
 
 #ifndef DATASMITH_CAD_IGNORE_CACHE
 static TAutoConsoleVariable<int32> CVarStaticCADTranslatorEnableCADCache(
-	TEXT("DS.CADTranslator.EnableCADCache"),
+	TEXT("ds.CADTranslator.EnableCADCache"),
 	1,
 	TEXT("Activate to save temporary CAD processing file. These file will be use in a next import to avoid CAD file processing\n"),
 	ECVF_Default);
 #endif
 
 static TAutoConsoleVariable<int32> CVarStaticCADTranslatorEnableTimeControl(
-	TEXT("DS.CADTranslator.EnableTimeControl"),
+	TEXT("ds.CADTranslator.EnableTimeControl"),
 	1,
 	TEXT("Enable the timer that kill the worker if the import time is unusually long. With this time control, the load of the corrupted file is canceled but the rest of the scene is imported.\n"),
 	ECVF_Default);
@@ -40,7 +40,6 @@ FDatasmithDispatcher::FDatasmithDispatcher(const CADLibrary::FImportParameters&
 	, NextWorkerId(0)
 {
 	ImportParameters.bEnableTimeControl = (CVarStaticCADTranslatorEnableTimeControl.GetValueOnAnyThread() != 0);
-#ifndef DATASMITH_CAD_IGNORE_CACHE
 	ImportParameters.bEnableCacheUsage = (CVarStaticCADTranslatorEnableCADCache.GetValueOnAnyThread() != 0);
 
 	if (ImportParameters.bEnableCacheUsage)
@@ -51,9 +50,6 @@ FDatasmithDispatcher::FDatasmithDispatcher(const CADLibrary::FImportParameters&
 		IFileManager::Get().MakeDirectory(*FPaths::Combine(ProcessCacheFolder, TEXT("mesh")), true);
 		IFileManager::Get().MakeDirectory(*FPaths::Combine(ProcessCacheFolder, TEXT("body")), true);
 	}
-#else
-	ImportParameters.bEnableCacheUsage = false;
-#endif
 }
 
 void FDatasmithDispatcher::AddTask(const CADLibrary::FFileDescription & InFileDescription)
diff --git a/Engine/Plugins/Experimental/Enterprise/DatasmithRuntime/Source/Private/DatasmithRuntime.cpp b/Engine/Plugins/Experimental/Enterprise/DatasmithRuntime/Source/Private/DatasmithRuntime.cpp
--- a/Engine/Plugins/Experimental/Enterprise/DatasmithRuntime/Source/Private/DatasmithRuntime.cpp
+++ b/Engine/Plugins/Experimental/Enterprise/DatasmithRuntime/Source/Private/DatasmithRuntime.cpp
@@ -96,13 +96,13 @@ void ADatasmithRuntimeActor::Tick(float DeltaTime)
 #if WITH_EDITOR
 				if (EnableThreadedImport != MAX_int32)
 				{
-					IConsoleVariable* CVar = IConsoleManager::Get().FindConsoleVariable(TEXT("DS.CADTranslator.EnableThreadedImport"));
+					IConsoleVariable* CVar = IConsoleManager::Get().FindConsoleVariable(TEXT("ds.CADTranslator.EnableThreadedImport"));
 					CVar->Set(EnableThreadedImport);
 				}
 
 				if (EnableCADCache != MAX_int32)
 				{
-					IConsoleVariable* CVar = IConsoleManager::Get().FindConsoleVariable(TEXT("DS.CADTranslator.EnableCADCache"));
+					IConsoleVariable* CVar = IConsoleManager::Get().FindConsoleVariable(TEXT("ds.CADTranslator.EnableCADCache"));
 					CVar->Set(EnableCADCache);
 				}
 #endif
@@ -363,14 +363,14 @@ bool ADatasmithRuntimeActor::LoadFile(const FString& FilePath)
 
 #if WITH_EDITOR
 	EnableThreadedImport = MAX_int32;
-	if (IConsoleVariable* CVar = IConsoleManager::Get().FindConsoleVariable(TEXT("DS.CADTranslator.EnableThreadedImport")))
+	if (IConsoleVariable* CVar = IConsoleManager::Get().FindConsoleVariable(TEXT("ds.CADTranslator.EnableThreadedImport")))
 	{
 		EnableThreadedImport = CVar->GetInt();
 		CVar->Set(0);
 	}
 
 	EnableCADCache = MAX_int32;
-	if (IConsoleVariable* CVar = IConsoleManager::Get().FindConsoleVariable(TEXT("DS.CADTranslator.EnableCADCache")))
+	if (IConsoleVariable* CVar = IConsoleManager::Get().FindConsoleVariable(TEXT("ds.CADTranslator.EnableCADCache")))
 	{
 		EnableCADCache = CVar->GetInt();
 		CVar->Set(0);
diff --git a/Engine/Plugins/Experimental/Enterprise/DatasmithRuntime/Source/Private/DatasmithRuntimeModule.cpp b/Engine/Plugins/Experimental/Enterprise/DatasmithRuntime/Source/Private/DatasmithRuntimeModule.cpp
--- a/Engine/Plugins/Experimental/Enterprise/DatasmithRuntime/Source/Private/DatasmithRuntimeModule.cpp
+++ b/Engine/Plugins/Experimental/Enterprise/DatasmithRuntime/Source/Private/DatasmithRuntimeModule.cpp
@@ -92,12 +92,12 @@ public:
 					//FModuleManager::Get().LoadModuleChecked(DATASMITHOPENNURBSTRANSLATOR_MODULE_NAME);
 					FModuleManager::Get().LoadModuleChecked(DATASMITHCADTRANSLATOR_MODULE_NAME);
 
-					if (IConsoleVariable* CVar = IConsoleManager::Get().FindConsoleVariable(TEXT("DS.CADTranslator.EnableThreadedImport")))
+					if (IConsoleVariable* CVar = IConsoleManager::Get().FindConsoleVariable(TEXT("ds.CADTranslator.EnableThreadedImport")))
 					{
 						CVar->Set(0);
 					}
 
-					if (IConsoleVariable* CVar = IConsoleManager::Get().FindConsoleVariable(TEXT("DS.CADTranslator.EnableCADCache")))
+					if (IConsoleVariable* CVar = IConsoleManager::Get().FindConsoleVariable(TEXT("ds.CADTranslator.EnableCADCache")))
 					{
 						CVar->Set(0);
 					}
