Pre-pluggin of CADKernel in Wire
Add AliasCADWrapper interface to allow to use the good wrapper according to the modeler used.
By setting DS.CADTranslator.EnableKernelIOTessellation CVar to 0 (disabling KernelIO Tessellation) CADKernel is activated in wire interface

diff --git a/Engine/Plugins/Enterprise/DatasmithCADImporter/Source/CADInterfaces/CADInterfaces.Build.cs b/Engine/Plugins/Enterprise/DatasmithCADImporter/Source/CADInterfaces/CADInterfaces.Build.cs
--- a/Engine/Plugins/Enterprise/DatasmithCADImporter/Source/CADInterfaces/CADInterfaces.Build.cs
+++ b/Engine/Plugins/Enterprise/DatasmithCADImporter/Source/CADInterfaces/CADInterfaces.Build.cs
@@ -15,6 +15,7 @@ namespace UnrealBuildTool.Rules
 				{
 					"Core",
 					"CADKernel",
+					"CADLibrary",
 					"CADTools",
 					"DatasmithCore"
 				}
diff --git a/Engine/Plugins/Enterprise/DatasmithCADImporter/Source/CADInterfaces/Public/CTSession.h b/Engine/Plugins/Enterprise/DatasmithCADImporter/Source/CADInterfaces/Public/CTSession.h
--- a/Engine/Plugins/Enterprise/DatasmithCADImporter/Source/CADInterfaces/Public/CTSession.h
+++ b/Engine/Plugins/Enterprise/DatasmithCADImporter/Source/CADInterfaces/Public/CTSession.h
@@ -63,7 +63,12 @@ public:
 	 */
 	void SetImportParameters(float ChordTolerance, float MaxEdgeLength, float NormalTolerance, CADLibrary::EStitchingTechnique StitchingTechnique, bool bScaleUVMap);
 	
-	CADLibrary::FImportParameters& GetImportParameters()
+	void SetModelCoordinateSystem(FDatasmithUtils::EModelCoordSystem NewCoordinateSystem)
+	{
+		ImportParams.ModelCoordSys = NewCoordinateSystem;
+	}
+
+	const CADLibrary::FImportParameters& GetImportParameters() const
 	{
 		return ImportParams;
 	}
diff --git a/Engine/Plugins/Enterprise/DatasmithCADImporter/Source/CADKernelSurface/Private/CADKernelSurfaceExtension.cpp b/Engine/Plugins/Enterprise/DatasmithCADImporter/Source/CADKernelSurface/Private/CADKernelSurfaceExtension.cpp
--- a/Engine/Plugins/Enterprise/DatasmithCADImporter/Source/CADKernelSurface/Private/CADKernelSurfaceExtension.cpp
+++ b/Engine/Plugins/Enterprise/DatasmithCADImporter/Source/CADKernelSurface/Private/CADKernelSurfaceExtension.cpp
@@ -7,7 +7,6 @@
 #include "CADOptions.h"
 #include "DatasmithPayload.h"
 #include "Engine/StaticMesh.h"
-#include "IDatasmithSceneElements.h"
 #include "MeshDescription.h"
 #include "MeshDescriptionHelper.h"
 #include "StaticMeshAttributes.h"
@@ -69,7 +68,7 @@ bool UCADKernelParametricSurfaceData::Tessellate(UStaticMesh& StaticMesh, const
 		// Tesselate the model
 		TSharedRef<CADKernel::FModelMesh> CADKernelModelMesh = CADKernel::FEntity::MakeShared<CADKernel::FModelMesh>();
 
-		FCADKernelTools::DefineMeshCriteria(CADKernelModelMesh, ImportParameters);
+		CADLibrary::FCADKernelTools::DefineMeshCriteria(CADKernelModelMesh, ImportParameters);
 
 		CADKernel::FParametricMesher Mesher(CADKernelModelMesh);
 		Mesher.MeshEntity(CADKernelModel);
@@ -79,13 +78,12 @@ bool UCADKernelParametricSurfaceData::Tessellate(UStaticMesh& StaticMesh, const
 		{
 			return bSuccessfulTessellation;
 		}
-		const TSharedRef<CADKernel::FBody> Body = CADKernelBodies[0].ToSharedRef();
 
 		CADLibrary::FBodyMesh BodyMesh;
 		uint32 DefaultMaterialHash = 0;
 
-		FCADKernelTools::GetBodyTessellation(CADKernelModelMesh, Body, BodyMesh, DefaultMaterialHash, [](CADLibrary::FObjectDisplayDataId, CADLibrary::FObjectDisplayDataId, int32) {;});
-		if (CADLibrary::ConvertBodyMeshToMeshDescription(ImportParameters, CadMeshParameters, BodyMesh, MeshDescription))
+		TSharedRef<CADKernel::FTopologicalEntity> Body = StaticCastSharedRef<CADKernel::FTopologicalEntity>(CADKernelBodies[0].ToSharedRef());
+		if(CADLibrary::FCADKernelTools::Tessellate(Body, ImportParameters, CadMeshParameters, MeshDescription))
 		{
 			// To update the SectionInfoMap 
 			TPolygonGroupAttributesConstRef<FName> MaterialSlotNames = MeshDescriptionAttributes.GetPolygonGroupMaterialSlotNames();
@@ -113,14 +111,13 @@ bool UCADKernelParametricSurfaceData::Tessellate(UStaticMesh& StaticMesh, const
 
 namespace CADKernelSurface
 {
-	void AddSurfaceDataForMesh(const TSharedRef<IDatasmithMeshElement>& InMeshElement, const CADLibrary::FImportParameters& InSceneParameters, const CADLibrary::FMeshParameters& InMeshParameters, const FDatasmithTessellationOptions& InTessellationOptions, FDatasmithMeshElementPayload& OutMeshPayload)
+	void AddSurfaceDataForMesh(const TCHAR* CADKernelArchive, const CADLibrary::FImportParameters& InSceneParameters, const CADLibrary::FMeshParameters& InMeshParameters, const FDatasmithTessellationOptions& InTessellationOptions, FDatasmithMeshElementPayload& OutMeshPayload)
 	{
-		// Store CoreTech additional data if provided
-		FString CADKernelDatabase = InMeshElement->GetFile();
-		if (FPaths::FileExists(*CADKernelDatabase))
+		// Store CADKernel archive if provided
+		if (FPaths::FileExists(CADKernelArchive))
 		{
 			TArray<uint8> ByteArray;
-			if (FFileHelper::LoadFileToArray(ByteArray, *CADKernelDatabase))
+			if (FFileHelper::LoadFileToArray(ByteArray, CADKernelArchive))
 			{
 				UCADKernelParametricSurfaceData* CADKernelData = Datasmith::MakeAdditionalData<UCADKernelParametricSurfaceData>();
 				CADKernelData->RawData = MoveTemp(ByteArray);
diff --git a/Engine/Plugins/Enterprise/DatasmithCADImporter/Source/CADKernelSurface/Public/CADKernelSurfaceExtension.h b/Engine/Plugins/Enterprise/DatasmithCADImporter/Source/CADKernelSurface/Public/CADKernelSurfaceExtension.h
--- a/Engine/Plugins/Enterprise/DatasmithCADImporter/Source/CADKernelSurface/Public/CADKernelSurfaceExtension.h
+++ b/Engine/Plugins/Enterprise/DatasmithCADImporter/Source/CADKernelSurface/Public/CADKernelSurfaceExtension.h
@@ -33,7 +33,6 @@ private:
 	virtual void Serialize(FArchive& Ar) override;
 };
 
-class IDatasmithMeshElement;
 struct FDatasmithMeshElementPayload;
 
 namespace CADLibrary
@@ -44,5 +43,5 @@ namespace CADLibrary
 
 namespace CADKernelSurface
 {
-	void CADKERNELSURFACE_API AddSurfaceDataForMesh(const TSharedRef<IDatasmithMeshElement>& InMeshElement, const CADLibrary::FImportParameters& InSceneParameters, const CADLibrary::FMeshParameters&, const FDatasmithTessellationOptions& InTessellationOptions, FDatasmithMeshElementPayload& OutMeshPayload);
+	void CADKERNELSURFACE_API AddSurfaceDataForMesh(const TCHAR* CADKernelArchive, const CADLibrary::FImportParameters& InSceneParameters, const CADLibrary::FMeshParameters&, const FDatasmithTessellationOptions& InTessellationOptions, FDatasmithMeshElementPayload& OutMeshPayload);
 }
diff --git a/Engine/Plugins/Enterprise/DatasmithCADImporter/Source/CADLibrary/CADLibrary.Build.cs b/Engine/Plugins/Enterprise/DatasmithCADImporter/Source/CADLibrary/CADLibrary.Build.cs
--- a/Engine/Plugins/Enterprise/DatasmithCADImporter/Source/CADLibrary/CADLibrary.Build.cs
+++ b/Engine/Plugins/Enterprise/DatasmithCADImporter/Source/CADLibrary/CADLibrary.Build.cs
@@ -15,6 +15,7 @@ public class CADLibrary : ModuleRules
 			{
 				"Core",
 				"CADTools",
+				"CADKernel",
 				"DatasmithCore",
 				"MeshDescription",
 				"StaticMeshDescription"
diff --git a/Engine/Plugins/Enterprise/DatasmithCADImporter/Source/CADLibrary/Private/CADKernelTools.cpp b/Engine/Plugins/Enterprise/DatasmithCADImporter/Source/CADLibrary/Private/CADKernelTools.cpp
new file mode 100644
--- /dev/null
+++ b/Engine/Plugins/Enterprise/DatasmithCADImporter/Source/CADLibrary/Private/CADKernelTools.cpp
@@ -0,0 +1,410 @@
+// Copyright Epic Games, Inc. All Rights Reserved.
+
+#include "CADKernelTools.h"
+
+#include "CADData.h"
+#include "CADOptions.h"
+#include "MeshDescription.h"
+#include "MeshDescriptionHelper.h"
+#include "StaticMeshAttributes.h"
+#include "StaticMeshOperations.h"
+
+#include "CADKernel/Core/Entity.h"
+#include "CADKernel/Core/MetadataDictionary.h"
+#include "CADKernel/Core/Session.h"
+#include "CADKernel/Core/Types.h"
+
+#include "CADKernel/Mesh/Criteria/Criterion.h"
+#include "CADKernel/Mesh/Meshers/ParametricMesher.h"
+#include "CADKernel/Mesh/Structure/FaceMesh.h"
+#include "CADKernel/Mesh/Structure/ModelMesh.h"
+
+#include "CADKernel/Topo/Model.h"
+#include "CADKernel/Topo/Body.h"
+#include "CADKernel/Topo/Shell.h"
+#include "CADKernel/Topo/TopologicalEdge.h"
+#include "CADKernel/Topo/TopologicalEntity.h"
+#include "CADKernel/Topo/TopologicalFace.h"
+#include "CADKernel/Topo/TopologicalVertex.h"
+
+using namespace CADKernel;
+typedef uint32 TriangleIndex[3];
+
+namespace CADLibrary
+{
+	class FMeshConversionContext
+	{
+	public:
+		const FImportParameters& ImportParams;
+		const FMeshParameters& MeshParameters;
+		TArray<int32> VertexIds;
+		TArray<int32> SymmetricVertexIds;
+
+		FMeshConversionContext(const FImportParameters& InImportParams, const FMeshParameters& InMeshParameters)
+			: ImportParams(InImportParams)
+			, MeshParameters(InMeshParameters)
+		{
+		}
+	};
+
+	void FillVertexPosition(FMeshConversionContext& Context, const TSharedRef<FModelMesh>& ModelMesh, FMeshDescription& MeshDescription)
+	{
+		TArray<FVector> VertexArray;
+		ModelMesh->GetNodeCoordinates(VertexArray);
+
+		for (FVector& Vertex : VertexArray)
+		{
+			Vertex *= Context.ImportParams.ScaleFactor;
+		}
+
+		int32 VertexCount = VertexArray.Num();
+
+		TVertexAttributesRef<FVector3f> VertexPositions = MeshDescription.GetVertexPositions();
+		MeshDescription.ReserveNewVertices(Context.MeshParameters.bIsSymmetric ? VertexCount * 2 : VertexCount);
+
+		Context.VertexIds.SetNumZeroed(VertexCount);
+
+		// Make MeshDescription.VertexPositions and VertexID
+		int32 VertexIndex = -1;
+		for (const FVector& Vertex : VertexArray)
+		{
+			VertexIndex++;
+
+			FVertexID VertexID = MeshDescription.CreateVertex();
+			VertexPositions[VertexID] = FDatasmithUtils::ConvertVector((FDatasmithUtils::EModelCoordSystem) Context.ImportParams.ModelCoordSys, Vertex);
+			Context.VertexIds[VertexIndex] = VertexID;
+		}
+
+		// if Symmetric mesh, the symmetric side of the mesh have to be generated
+		if (Context.MeshParameters.bIsSymmetric)
+		{
+			FMatrix SymmetricMatrix = FDatasmithUtils::GetSymmetricMatrix(Context.MeshParameters.SymmetricOrigin, Context.MeshParameters.SymmetricNormal);
+
+			Context.SymmetricVertexIds.SetNum(VertexArray.Num());
+
+			VertexIndex = 0;
+			for (const FVector& Vertex : VertexArray)
+			{
+				FVertexID VertexID = MeshDescription.CreateVertex();
+				VertexPositions[VertexID] = FDatasmithUtils::ConvertVector((FDatasmithUtils::EModelCoordSystem) Context.ImportParams.ModelCoordSys, Vertex);
+				VertexPositions[VertexID] = SymmetricMatrix.TransformPosition(VertexPositions[VertexID]);
+				Context.SymmetricVertexIds[VertexIndex++] = VertexID;
+			}
+		}
+	}
+
+	bool FillMesh(FMeshConversionContext& Context, const TSharedRef<CADKernel::FModelMesh>& ModelMesh, FMeshDescription& MeshDescription)
+	{
+		const int32 UVChannel = 0;
+		const int32 TriangleCount = 3;
+		const TriangleIndex Clockwise = { 0, 1, 2 };
+		const TriangleIndex CounterClockwise = { 0, 2, 1 };
+
+		TArray<FVertexInstanceID> TriangleVertexInstanceIDs;
+		TriangleVertexInstanceIDs.SetNum(TriangleCount);
+
+		TArray<FVertexInstanceID> MeshVertexInstanceIDs;
+
+		// Gather all array data
+		FStaticMeshAttributes Attributes(MeshDescription);
+		TVertexInstanceAttributesRef<FVector3f> VertexInstanceNormals = Attributes.GetVertexInstanceNormals();
+		TVertexInstanceAttributesRef<FVector3f> VertexInstanceTangents = Attributes.GetVertexInstanceTangents();
+		TVertexInstanceAttributesRef<float> VertexInstanceBinormalSigns = Attributes.GetVertexInstanceBinormalSigns();
+		TVertexInstanceAttributesRef<FVector4> VertexInstanceColors = Attributes.GetVertexInstanceColors();
+		TVertexInstanceAttributesRef<FVector2D> VertexInstanceUVs = Attributes.GetVertexInstanceUVs();
+		TPolygonGroupAttributesRef<FName> PolygonGroupImportedMaterialSlotNames = Attributes.GetPolygonGroupMaterialSlotNames();
+
+		if (!VertexInstanceNormals.IsValid() || !VertexInstanceTangents.IsValid() || !VertexInstanceBinormalSigns.IsValid() || !VertexInstanceColors.IsValid() || !VertexInstanceUVs.IsValid() || !PolygonGroupImportedMaterialSlotNames.IsValid())
+		{
+			return false;
+		}
+
+		// Find all the materials used
+		TMap<uint32, FPolygonGroupID> MaterialToPolygonGroupMapping;
+		for (const TSharedPtr<FFaceMesh>& FaceMesh : ModelMesh->GetFaceMeshes())
+		{
+			const FTopologicalFace& Face = (const FTopologicalFace&)FaceMesh->GetGeometricEntity();
+			// we assume that face has only color
+			MaterialToPolygonGroupMapping.Add(Face.GetColorId(), INDEX_NONE);
+		}
+
+		// Add to the mesh, a polygon groups per material
+		for (auto& Material : MaterialToPolygonGroupMapping)
+		{
+			uint32 MaterialHash = Material.Key;
+			FName ImportedSlotName = *LexToString<uint32>(MaterialHash);
+
+			FPolygonGroupID PolyGroupID = MeshDescription.CreatePolygonGroup();
+			PolygonGroupImportedMaterialSlotNames[PolyGroupID] = ImportedSlotName;
+			Material.Value = PolyGroupID;
+		}
+
+		VertexInstanceUVs.SetNumChannels(1);
+
+		int32 NbStep = 1;
+		if (Context.MeshParameters.bIsSymmetric)
+		{
+			NbStep = 2;
+		}
+
+		int32 FaceIndex = 0;
+
+		TSet<int32> PatchIdSet;
+		GetExistingPatches(MeshDescription, PatchIdSet);
+		bool bImportOnlyAlreadyPresent = (bool)PatchIdSet.Num();
+
+		TPolygonAttributesRef<int32> PatchGroups = EnableCADPatchGroups(MeshDescription);
+
+		int32 PatchIndex = 0;
+		for (int32 Step = 0; Step < NbStep; ++Step)
+		{
+			// Swap mesh if needed
+			const TriangleIndex& Orientation = (!Context.MeshParameters.bNeedSwapOrientation == (bool)Step) ? CounterClockwise : Clockwise;
+			TArray<int32>& VertexIdSet = (Step == 0) ? Context.VertexIds : Context.SymmetricVertexIds;
+
+			// Loop through the FaceMeshes and collect all tessellation data
+			for (const TSharedPtr<FFaceMesh>& FaceMesh : ModelMesh->GetFaceMeshes())
+			{
+				const FTopologicalFace& Face = (const FTopologicalFace&)FaceMesh->GetGeometricEntity();
+				if (bImportOnlyAlreadyPresent && !PatchIdSet.Contains(Face.GetPatchId()))
+				{
+					continue;
+				}
+
+				// Get the polygonGroup
+				const FPolygonGroupID* PolygonGroupID = MaterialToPolygonGroupMapping.Find(Face.GetColorId());
+				if (PolygonGroupID == nullptr)
+				{
+					continue;
+				}
+
+				int32 VertexIDs[3];
+				//FVector Temp3D = { 0, 0, 0 };
+				//FVector2D TexCoord2D = { 0, 0 };
+
+				TArray<int32>& TriangleVertexIndices = FaceMesh->TrianglesVerticesIndex;
+				TArray<int32>& VerticesGlobalIndex = FaceMesh->VerticesGlobalIndex;
+				MeshVertexInstanceIDs.Empty(TriangleVertexIndices.Num());
+
+				PatchIndex++;
+
+				// build each valid face i.e. 3 different indexes
+				for (int32 Index = 0; Index < TriangleVertexIndices.Num(); Index += 3)
+				{
+					VertexIDs[0] = VertexIdSet[VerticesGlobalIndex[TriangleVertexIndices[Index + Orientation[0]]]];
+					VertexIDs[1] = VertexIdSet[VerticesGlobalIndex[TriangleVertexIndices[Index + Orientation[1]]]];
+					VertexIDs[2] = VertexIdSet[VerticesGlobalIndex[TriangleVertexIndices[Index + Orientation[2]]]];
+
+					MeshVertexInstanceIDs.Add(TriangleVertexInstanceIDs[0] = MeshDescription.CreateVertexInstance((FVertexID)VertexIDs[0]));
+					MeshVertexInstanceIDs.Add(TriangleVertexInstanceIDs[1] = MeshDescription.CreateVertexInstance((FVertexID)VertexIDs[1]));
+					MeshVertexInstanceIDs.Add(TriangleVertexInstanceIDs[2] = MeshDescription.CreateVertexInstance((FVertexID)VertexIDs[2]));
+
+					// Add the triangle as a polygon to the mesh description
+					const FPolygonID PolygonID = MeshDescription.CreatePolygon(*PolygonGroupID, TriangleVertexInstanceIDs);
+
+					// Set patch id attribute
+					PatchGroups[PolygonID] = Face.GetPatchId();
+				}
+
+				for (int32 IndexFace = 0; IndexFace < MeshVertexInstanceIDs.Num(); IndexFace += 3)
+				{
+					for (int32 Index = 0; Index < TriangleCount; Index++)
+					{
+						const FVertexInstanceID VertexInstanceID = MeshVertexInstanceIDs[IndexFace + Index];
+						VertexInstanceUVs.Set(VertexInstanceID, UVChannel, FaceMesh->UVMap[TriangleVertexIndices[IndexFace + Orientation[Index]]]);
+
+						VertexInstanceColors[VertexInstanceID] = FLinearColor::White;
+						VertexInstanceTangents[VertexInstanceID] = FVector(ForceInitToZero);
+						VertexInstanceBinormalSigns[VertexInstanceID] = 0.0f;
+					}
+				}
+
+				if (!Step)
+				{
+					FDatasmithUtils::ConvertVectorArray(Context.ImportParams.ModelCoordSys, FaceMesh->Normals);
+					for (FVector& Normal : FaceMesh->Normals)
+					{
+						Normal = Normal.GetSafeNormal();
+					}
+				}
+
+				for (int32 IndexFace = 0; IndexFace < MeshVertexInstanceIDs.Num(); IndexFace += 3)
+				{
+					for (int32 Index = 0; Index < 3; Index++)
+					{
+						const FVertexInstanceID VertexInstanceID = MeshVertexInstanceIDs[IndexFace + Index];
+						VertexInstanceNormals[VertexInstanceID] = FaceMesh->Normals[TriangleVertexIndices[IndexFace + Orientation[Index]]];
+					}
+				}
+
+				// compute normals
+				if (Step)
+				{
+					// compute normals of Symmetric vertex
+					FMatrix SymmetricMatrix;
+					SymmetricMatrix = FDatasmithUtils::GetSymmetricMatrix(Context.MeshParameters.SymmetricOrigin, Context.MeshParameters.SymmetricNormal);
+					for (const FVertexInstanceID& VertexInstanceID : MeshVertexInstanceIDs)
+					{
+						VertexInstanceNormals[VertexInstanceID] = SymmetricMatrix.TransformVector(VertexInstanceNormals[VertexInstanceID]);;
+					}
+				}
+			}
+		}
+		return true;
+	}
+
+	bool ConvertModelMeshToMeshDescription(FMeshConversionContext& Context, const TSharedRef<FModelMesh>& InModelMesh, FMeshDescription& MeshDescription)
+	{
+		int32 VertexCount = InModelMesh->GetVertexCount();
+		int32 TriangleCount = InModelMesh->GetTriangleCount();
+
+		MeshDescription.ReserveNewVertexInstances(VertexCount);
+		MeshDescription.ReserveNewPolygons(TriangleCount);
+		MeshDescription.ReserveNewEdges(TriangleCount);
+
+		FillVertexPosition(Context, InModelMesh, MeshDescription);
+		if (!FillMesh(Context, InModelMesh, MeshDescription))
+		{
+			return false;
+		}
+
+		// Build edge meta data
+		FStaticMeshOperations::DetermineEdgeHardnessesFromVertexInstanceNormals(MeshDescription);
+
+		return MeshDescription.Polygons().Num() > 0;
+	}
+
+	bool FCADKernelTools::Tessellate(TSharedRef<FTopologicalEntity>& CADTopologicalEntity, const FImportParameters& ImportParameters, const FMeshParameters& MeshParameters, FMeshDescription& OutMeshDescription)
+	{
+		// Tessellate the model
+		TSharedRef<FModelMesh> CADKernelModelMesh = FEntity::MakeShared<FModelMesh>();
+		FParametricMesher Mesher(CADKernelModelMesh);
+
+		DefineMeshCriteria(CADKernelModelMesh, ImportParameters);
+		Mesher.MeshEntity(CADTopologicalEntity);
+
+		FMeshConversionContext Context(ImportParameters, MeshParameters);
+
+		return ConvertModelMeshToMeshDescription(Context, CADKernelModelMesh, OutMeshDescription);
+	}
+
+
+	uint32 FCADKernelTools::GetFaceTessellation(const TSharedRef<CADKernel::FFaceMesh>& FaceMesh, FBodyMesh& OutBodyMesh)
+	{
+		// Something wrong happened, either an error or no data to collect
+		if (FaceMesh->TrianglesVerticesIndex.Num() == 0)
+		{
+			return 0;
+		}
+
+		FTessellationData& Tessellation = OutBodyMesh.Faces.Emplace_GetRef();
+
+		const CADKernel::FTopologicalFace& HaveMetadata = (const CADKernel::FTopologicalFace&) FaceMesh->GetGeometricEntity();
+		Tessellation.PatchId = HaveMetadata.GetPatchId();
+
+		Tessellation.PositionIndices = MoveTemp(FaceMesh->VerticesGlobalIndex);
+		Tessellation.VertexIndices = MoveTemp(FaceMesh->TrianglesVerticesIndex);
+
+		Tessellation.NormalArray = MoveTemp(FaceMesh->Normals);
+		Tessellation.TexCoordArray = MoveTemp(FaceMesh->UVMap);
+
+		return Tessellation.VertexIndices.Num() / 3;
+	}
+
+	template<class ClassType>
+	void GetDisplayDataIds(const TSharedRef<ClassType>& Entity, FObjectDisplayDataId& DisplayDataId)
+	{
+		const TSharedRef<CADKernel::FMetadataDictionary>& HaveMetadata = StaticCastSharedRef<CADKernel::FMetadataDictionary>(Entity);
+		DisplayDataId.Color = HaveMetadata->GetColorId();
+		DisplayDataId.Material = HaveMetadata->GetMaterialId();
+	}
+
+	void FCADKernelTools::GetBodyTessellation(const TSharedRef<CADKernel::FModelMesh>& ModelMesh, const TSharedRef<CADKernel::FBody>& Body, FBodyMesh& OutBodyMesh, uint32 DefaultMaterialHash, TFunction<void(FObjectDisplayDataId, FObjectDisplayDataId, int32)> SetFaceMainMaterial)
+	{
+		ModelMesh->GetNodeCoordinates(OutBodyMesh.VertexArray);
+
+		uint32 FaceSize = Body->FaceCount();
+
+		// Allocate memory space for tessellation data
+		OutBodyMesh.Faces.Reserve(FaceSize);
+		OutBodyMesh.ColorSet.Reserve(FaceSize);
+		OutBodyMesh.MaterialSet.Reserve(FaceSize);
+
+		FObjectDisplayDataId BodyMaterial;
+		BodyMaterial.DefaultMaterialName = DefaultMaterialHash;
+
+		GetDisplayDataIds(Body, BodyMaterial);
+
+		// Loop through the face of bodies and collect all tessellation data
+		int32 FaceIndex = 0;
+		for (const TSharedPtr<CADKernel::FShell>& Shell : Body->GetShells())
+		{
+			if (!Shell.IsValid())
+			{
+				continue;
+			}
+
+			FObjectDisplayDataId ShellMaterial = BodyMaterial;
+			GetDisplayDataIds(Shell.ToSharedRef(), ShellMaterial);
+
+			for (const CADKernel::FOrientedFace& Face : Shell->GetFaces())
+			{
+				if (!Face.Entity.IsValid())
+				{
+					continue;
+				}
+
+				if (!Face.Entity->HasTesselation())
+				{
+					continue;
+				}
+
+				FObjectDisplayDataId FaceMaterial;
+				GetDisplayDataIds(Face.Entity.ToSharedRef(), FaceMaterial);
+
+				uint32 TriangleNum = GetFaceTessellation(Face.Entity->GetMesh(), OutBodyMesh);
+
+				if (TriangleNum == 0)
+				{
+					continue;
+				}
+
+				OutBodyMesh.TriangleCount += TriangleNum;
+
+				if (SetFaceMainMaterial)
+				{
+					SetFaceMainMaterial(FaceMaterial, ShellMaterial, FaceIndex);
+				}
+				FaceIndex++;
+			}
+		}
+	}
+
+	void FCADKernelTools::DefineMeshCriteria(TSharedRef<CADKernel::FModelMesh>& MeshModel, const FImportParameters& ImportParameters)
+	{
+		{
+			TSharedPtr<CADKernel::FCriterion> CurvatureCriterion = CADKernel::FCriterion::CreateCriterion(CADKernel::ECriterion::CADCurvature);
+			MeshModel->AddCriterion(CurvatureCriterion);
+		}
+
+		if (ImportParameters.MaxEdgeLength > SMALL_NUMBER)
+		{
+			TSharedPtr<CADKernel::FCriterion> MaxSizeCriterion = CADKernel::FCriterion::CreateCriterion(CADKernel::ECriterion::MaxSize, ImportParameters.MaxEdgeLength);
+			MeshModel->AddCriterion(MaxSizeCriterion);
+		}
+
+		if (ImportParameters.ChordTolerance > SMALL_NUMBER)
+		{
+			TSharedPtr<CADKernel::FCriterion> ChordCriterion = CADKernel::FCriterion::CreateCriterion(CADKernel::ECriterion::Sag, ImportParameters.ChordTolerance);
+			MeshModel->AddCriterion(ChordCriterion);
+		}
+
+		if (ImportParameters.MaxNormalAngle > SMALL_NUMBER)
+		{
+			TSharedPtr<CADKernel::FCriterion> MaxNormalAngleCriterion = CADKernel::FCriterion::CreateCriterion(CADKernel::ECriterion::Angle, ImportParameters.MaxNormalAngle);
+			MeshModel->AddCriterion(MaxNormalAngleCriterion);
+		}
+	}
+
+}
+
diff --git a/Engine/Plugins/Enterprise/DatasmithCADImporter/Source/CADLibrary/Private/MeshDescriptionHelper.cpp b/Engine/Plugins/Enterprise/DatasmithCADImporter/Source/CADLibrary/Private/MeshDescriptionHelper.cpp
--- a/Engine/Plugins/Enterprise/DatasmithCADImporter/Source/CADLibrary/Private/MeshDescriptionHelper.cpp
+++ b/Engine/Plugins/Enterprise/DatasmithCADImporter/Source/CADLibrary/Private/MeshDescriptionHelper.cpp
@@ -291,13 +291,13 @@ namespace CADLibrary
 	const FName PolyTriGroups("PolyTriGroups");
 
 	// Copy of FMeshDescriptionBuilder::EnablePolyGroups()
-	TPolygonAttributesRef<int32> EnableCADPatchGroups(FMeshDescription& MeshDescription)
+	TPolygonAttributesRef<int32> EnableCADPatchGroups(FMeshDescription& OutMeshDescription)
 	{
-		TPolygonAttributesRef<int32> PatchGroups = MeshDescription.PolygonAttributes().GetAttributesRef<int32>(PolyTriGroups);
+		TPolygonAttributesRef<int32> PatchGroups = OutMeshDescription.PolygonAttributes().GetAttributesRef<int32>(PolyTriGroups);
 		if (PatchGroups.IsValid() == false)
 		{
-			MeshDescription.PolygonAttributes().RegisterAttribute<int32>(PolyTriGroups, 1, 0, EMeshAttributeFlags::AutoGenerated);
-			PatchGroups = MeshDescription.PolygonAttributes().GetAttributesRef<int32>(PolyTriGroups);
+			OutMeshDescription.PolygonAttributes().RegisterAttribute<int32>(PolyTriGroups, 1, 0, EMeshAttributeFlags::AutoGenerated);
+			PatchGroups = OutMeshDescription.PolygonAttributes().GetAttributesRef<int32>(PolyTriGroups);
 			check(PatchGroups.IsValid());
 		}
 		return PatchGroups;
@@ -492,18 +492,6 @@ namespace CADLibrary
 					}
 				}
 
-				if (Tessellation.TexCoordArray.Num())
-				{
-					for (int32 IndexFace = 0; IndexFace < NewFaceIndex.Num(); IndexFace += 3)
-					{
-						for (int32 Index = 0; Index < TriangleCount; Index++)
-						{
-							FVertexInstanceID VertexInstanceID = MeshVertexInstanceIDs[IndexFace + Orientation[Index]];
-							VertexInstanceUVs.Set(VertexInstanceID, UVChannel, Tessellation.TexCoordArray[NewFaceIndex[IndexFace + Index]]);
-						}
-					}
-				}
-
 				if (!Step)
 				{
 					FDatasmithUtils::ConvertVectorArray(ImportParams.ModelCoordSys, Tessellation.NormalArray);
@@ -524,11 +512,6 @@ namespace CADLibrary
 				}
 				else
 				{
-					for (FVector& Normal : Tessellation.NormalArray)
-					{
-						Normal = Normal.GetSafeNormal();
-					}
-
 					for (int32 IndexFace = 0; IndexFace < NewFaceIndex.Num(); IndexFace += 3)
 					{
 						for (int32 Index = 0; Index < 3; Index++)
diff --git a/Engine/Plugins/Enterprise/DatasmithCADImporter/Source/CADLibrary/Public/CADKernelTools.h b/Engine/Plugins/Enterprise/DatasmithCADImporter/Source/CADLibrary/Public/CADKernelTools.h
new file mode 100644
--- /dev/null
+++ b/Engine/Plugins/Enterprise/DatasmithCADImporter/Source/CADLibrary/Public/CADKernelTools.h
@@ -0,0 +1,39 @@
+// Copyright Epic Games, Inc. All Rights Reserved.
+#pragma once
+
+#include "CoreMinimal.h"
+
+#include "CADData.h"
+#include "CADOptions.h"
+
+namespace CADKernel
+{
+	class FBody;
+	class FBodyMesh;
+	class FFaceMesh;
+	class FTopologicalEntity;
+	class FModelMesh;
+}
+
+struct FMeshDescription;
+
+namespace CADLibrary
+{
+	struct FImportParameters;
+	struct FMeshParameters;
+
+	class CADLIBRARY_API FCADKernelTools
+	{
+	public:
+		static void DefineMeshCriteria(TSharedRef<CADKernel::FModelMesh>& MeshModel, const FImportParameters& ImportParameters);
+
+		static void GetBodyTessellation(const TSharedRef<CADKernel::FModelMesh>& ModelMesh, const TSharedRef<CADKernel::FBody>& Body, FBodyMesh& OutBodyMesh, uint32 DefaultMaterialHash, TFunction<void(FObjectDisplayDataId, FObjectDisplayDataId, int32)> SetFaceMainMaterial);
+
+		/**
+		 * Tessellate a CADKernel entity and update the MeshDescription
+		 */
+		static bool Tessellate(TSharedRef<CADKernel::FTopologicalEntity>& InCADKernelEntity, const FImportParameters& ImportParameters, const FMeshParameters& MeshParameters, FMeshDescription& OutMeshDescription);
+
+		static uint32 GetFaceTessellation(const TSharedRef<CADKernel::FFaceMesh>& FaceMesh, FBodyMesh& OutBodyMesh);
+	};
+}
diff --git a/Engine/Plugins/Enterprise/DatasmithCADImporter/Source/CADLibrary/Public/MeshDescriptionHelper.h b/Engine/Plugins/Enterprise/DatasmithCADImporter/Source/CADLibrary/Public/MeshDescriptionHelper.h
--- a/Engine/Plugins/Enterprise/DatasmithCADImporter/Source/CADLibrary/Public/MeshDescriptionHelper.h
+++ b/Engine/Plugins/Enterprise/DatasmithCADImporter/Source/CADLibrary/Public/MeshDescriptionHelper.h
@@ -5,9 +5,9 @@
 
 #include "CADData.h"
 #include "CADOptions.h"
+#include "MeshDescription.h"
 
 struct FCTMesh;
-struct FMeshDescription;
 class FBodyMesh;
 class FCADMaterial;
 class IDatasmithMaterialIDElement;
@@ -27,4 +27,12 @@ namespace CADLibrary
 	CADLIBRARY_API bool ConvertBodyMeshToMeshDescription(const FImportParameters& ImportParams, const FMeshParameters& MeshParameters, FBodyMesh& Body, FMeshDescription& MeshDescription);
 	CADLIBRARY_API void CopyPatchGroups(FMeshDescription& MeshSource, FMeshDescription& MeshDestination);
 
+	/** 
+	 * Enable per-triangle integer attribute named PolyTriGroups 
+	 * This integer defines the identifier of the PolyTriGroup containing the triangle.
+	 * In case of mesh coming from a CAD file, a PolyTriGroup is associated to a CAD topological face
+	 */
+	CADLIBRARY_API TPolygonAttributesRef<int32> EnableCADPatchGroups(FMeshDescription& OutMeshDescription);
+	CADLIBRARY_API void GetExistingPatches(FMeshDescription& MeshDestination, TSet<int32>& OutPatchIdSet);
+
 } // namespace CADLibrary
diff --git a/Engine/Plugins/Enterprise/DatasmithCADImporter/Source/CADTools/CADTools.Build.cs b/Engine/Plugins/Enterprise/DatasmithCADImporter/Source/CADTools/CADTools.Build.cs
--- a/Engine/Plugins/Enterprise/DatasmithCADImporter/Source/CADTools/CADTools.Build.cs
+++ b/Engine/Plugins/Enterprise/DatasmithCADImporter/Source/CADTools/CADTools.Build.cs
@@ -12,7 +12,6 @@ namespace UnrealBuildTool.Rules
 				new string[]
 				{
                     "Core",
-					"CADKernel",
 					"DatasmithCore"
 				}
 			);
diff --git a/Engine/Plugins/Enterprise/DatasmithCADImporter/Source/CADTools/Private/CADKernelTools.cpp b/Engine/Plugins/Enterprise/DatasmithCADImporter/Source/CADTools/Private/CADKernelTools.cpp
deleted file mode 100644
diff --git a/Engine/Plugins/Enterprise/DatasmithCADImporter/Source/CADTools/Public/CADKernelTools.h b/Engine/Plugins/Enterprise/DatasmithCADImporter/Source/CADTools/Public/CADKernelTools.h
deleted file mode 100644
diff --git a/Engine/Plugins/Enterprise/DatasmithCADImporter/Source/CoreTechSurface/Private/CoreTechSurfaceExtension.cpp b/Engine/Plugins/Enterprise/DatasmithCADImporter/Source/CoreTechSurface/Private/CoreTechSurfaceExtension.cpp
--- a/Engine/Plugins/Enterprise/DatasmithCADImporter/Source/CoreTechSurface/Private/CoreTechSurfaceExtension.cpp
+++ b/Engine/Plugins/Enterprise/DatasmithCADImporter/Source/CoreTechSurface/Private/CoreTechSurfaceExtension.cpp
@@ -73,7 +73,7 @@ bool UCoreTechParametricSurfaceData::Tessellate(UStaticMesh& StaticMesh, const F
 			CADLibrary::CopyPatchGroups(*DestinationMeshDescription, MeshDescription);
 		}
 
-		if (CoreTechSurface::LoadFile(ResourceFile, MeshDescription, ImportParameters, CadMeshParameters))
+		if (CoreTechSurface::LoadFile(ResourceFile, ImportParameters, CadMeshParameters, MeshDescription))
 		{
 			// To update the SectionInfoMap 
 			{
diff --git a/Engine/Plugins/Enterprise/DatasmithCADImporter/Source/CoreTechSurface/Private/CoreTechSurfaceHelper.cpp b/Engine/Plugins/Enterprise/DatasmithCADImporter/Source/CoreTechSurface/Private/CoreTechSurfaceHelper.cpp
--- a/Engine/Plugins/Enterprise/DatasmithCADImporter/Source/CoreTechSurface/Private/CoreTechSurfaceHelper.cpp
+++ b/Engine/Plugins/Enterprise/DatasmithCADImporter/Source/CoreTechSurface/Private/CoreTechSurfaceHelper.cpp
@@ -16,7 +16,7 @@ typedef uint32 TriangleIndex[3];
 
 namespace CoreTechSurface
 {
-	bool Tessellate(uint64 MainObjectId, const CADLibrary::FImportParameters& ImportParams, FMeshDescription& MeshDesc, CADLibrary::FMeshParameters& MeshParameters)
+	bool Tessellate(uint64 MainObjectId, const CADLibrary::FImportParameters& ImportParams, const CADLibrary::FMeshParameters& MeshParameters, FMeshDescription& MeshDesc)
 	{
 		CADLibrary::CTKIO_SetCoreTechTessellationState(ImportParams);
 
@@ -39,7 +39,7 @@ namespace CoreTechSurface
 		return true;
 	}
 
-	bool LoadFile(const FString& FileName, FMeshDescription& MeshDescription, const CADLibrary::FImportParameters& ImportParameters, CADLibrary::FMeshParameters& MeshParameters)
+	bool LoadFile(const FString& FileName, const CADLibrary::FImportParameters& ImportParameters, const CADLibrary::FMeshParameters& MeshParameters, FMeshDescription& MeshDescription)
 	{
 		CADLibrary::FCoreTechSessionBase Session(TEXT("CoreTechMeshLoader::LoadFile"));
 		if (!Session.IsSessionValid())
@@ -60,23 +60,20 @@ namespace CoreTechSurface
 			CADLibrary::CTKIO_Repair(MainObjectID, CADLibrary::EStitchingTechnique::StitchingSew);
 		}
 
-		return Tessellate(MainObjectID, ImportParameters, MeshDescription, MeshParameters);
+		return Tessellate(MainObjectID, ImportParameters, MeshParameters, MeshDescription);
 	}
 
-	// TODO: convert to FCoreTechSceneParameters/FCoreTechMeshParameters ?
-	void AddSurfaceDataForMesh(const TSharedRef<IDatasmithMeshElement>& InMeshElement, const CADLibrary::FImportParameters& InSceneParameters, const CADLibrary::FMeshParameters& InMeshParameters, const FDatasmithTessellationOptions& InTessellationOptions, FDatasmithMeshElementPayload& OutMeshPayload)
+	void AddSurfaceDataForMesh(const TCHAR* InFilePath, const CADLibrary::FImportParameters& InSceneParameters, const CADLibrary::FMeshParameters& InMeshParameters, const FDatasmithTessellationOptions& InTessellationOptions, FDatasmithMeshElementPayload& OutMeshPayload)
 	{
 		if (ICADInterfacesModule::GetAvailability() == ECADInterfaceAvailability::Available)
 		{
-			// Store CoreTech additional data if provided
-			FString CoretechFile = InMeshElement->GetFile(); 
-			if (FPaths::FileExists(*CoretechFile))
+			if (FPaths::FileExists(InFilePath))
 			{
 				TArray<uint8> ByteArray;
-				if (FFileHelper::LoadFileToArray(ByteArray, *CoretechFile))
+				if (FFileHelper::LoadFileToArray(ByteArray, InFilePath))
 				{
 					UCoreTechParametricSurfaceData* CoreTechData = Datasmith::MakeAdditionalData<UCoreTechParametricSurfaceData>();
-					CoreTechData->SourceFile = CoretechFile;
+					CoreTechData->SourceFile = InFilePath;
 					CoreTechData->RawData = MoveTemp(ByteArray);
 					CoreTechData->SceneParameters.ModelCoordSys = uint8(InSceneParameters.ModelCoordSys);
 					CoreTechData->SceneParameters.MetricUnit = InSceneParameters.MetricUnit;
diff --git a/Engine/Plugins/Enterprise/DatasmithCADImporter/Source/CoreTechSurface/Public/CADModelToCoretechConverterBase.h b/Engine/Plugins/Enterprise/DatasmithCADImporter/Source/CoreTechSurface/Public/CADModelToCoretechConverterBase.h
new file mode 100644
--- /dev/null
+++ b/Engine/Plugins/Enterprise/DatasmithCADImporter/Source/CoreTechSurface/Public/CADModelToCoretechConverterBase.h
@@ -0,0 +1,78 @@
+// Copyright Epic Games, Inc. All Rights Reserved.
+
+#pragma once
+
+#include "CoreMinimal.h"
+
+#include "CADModelConverter.h"
+#include "CTSession.h"
+#include "CoreTechSurfaceHelper.h"
+#include "IDatasmithSceneElements.h"
+#include "Utility/DatasmithMeshHelper.h"
+
+class FCADModelToCoretechConverterBase : public CADLibrary::FCTSession, public CADLibrary::ICADModelConverter
+{
+public:
+
+	FCADModelToCoretechConverterBase(const TCHAR* InOwner)
+		: CADLibrary::FCTSession(InOwner)
+	{
+		// Unit for CoreTech session is set to cm, 0.01, because Wire's unit is cm. Consequently, Scale factor is set to 1.
+		ImportParams.MetricUnit = 0.01;
+		ImportParams.ScaleFactor = 1;
+		ImportParams.bEnableKernelIOTessellation = true;
+	}
+
+	virtual bool Tessellate(const CADLibrary::FMeshParameters& InMeshParameters, FMeshDescription& OutMeshDescription) override
+	{
+		// Apply stitching if applicable
+		TopoFixes(1.);
+
+		// Perform tessellation
+		DatasmithMeshHelper::PrepareAttributeForStaticMesh(OutMeshDescription);
+
+		return CoreTechSurface::Tessellate(MainObjectId, ImportParams, InMeshParameters, OutMeshDescription);
+	}
+
+	virtual bool RepairTopology() override
+	{
+		// Apply stitching if applicable
+		return TopoFixes(1.);
+	}
+
+	virtual void InitializeProcess(double NewSceneunit) override
+	{
+		ClearData();
+		SetSceneUnit(NewSceneunit);
+	}
+
+	virtual bool SaveBRep(const TCHAR* InFolderPath, TSharedRef<IDatasmithMeshElement>& MeshElement) override
+	{
+		FString FilePath = FPaths::Combine(InFolderPath, MeshElement->GetName()) + TEXT(".ct");
+		if (SaveBrep(FilePath))
+		{
+			MeshElement->SetFile(*FilePath);
+			return true;
+		}
+		return false;
+	}
+
+	virtual void SetImportParameters(float ChordTolerance, float MaxEdgeLength, float NormalTolerance, CADLibrary::EStitchingTechnique StitchingTechnique, bool bScaleUVMap) override
+	{
+		FCTSession::SetImportParameters(ChordTolerance, MaxEdgeLength, NormalTolerance, StitchingTechnique, bScaleUVMap);
+		ImportParams.bEnableKernelIOTessellation = true;
+	}
+
+	virtual bool IsSessionValid() override
+	{
+		return FCTSession::IsSessionValid();
+	}
+
+	void AddSurfaceDataForMesh(const TCHAR* InFilePath, const CADLibrary::FMeshParameters& InMeshParameters, const FDatasmithTessellationOptions& InTessellationOptions, FDatasmithMeshElementPayload& OutMeshPayload) const
+	{
+		const CADLibrary::FImportParameters& InImportParameters = GetImportParameters();
+		CoreTechSurface::AddSurfaceDataForMesh(InFilePath, InImportParameters, InMeshParameters, InTessellationOptions, OutMeshPayload);
+	}
+
+};
+
diff --git a/Engine/Plugins/Enterprise/DatasmithCADImporter/Source/CoreTechSurface/Public/CoreTechSurfaceHelper.h b/Engine/Plugins/Enterprise/DatasmithCADImporter/Source/CoreTechSurface/Public/CoreTechSurfaceHelper.h
--- a/Engine/Plugins/Enterprise/DatasmithCADImporter/Source/CoreTechSurface/Public/CoreTechSurfaceHelper.h
+++ b/Engine/Plugins/Enterprise/DatasmithCADImporter/Source/CoreTechSurface/Public/CoreTechSurfaceHelper.h
@@ -9,14 +9,13 @@ namespace CADLibrary
 	struct FMeshParameters;
 }
 
-class IDatasmithMeshElement;
 struct FDatasmithMeshElementPayload;
 struct FDatasmithTessellationOptions;
 struct FMeshDescription;
 
 namespace CoreTechSurface
 {
-	CORETECHSURFACE_API void AddSurfaceDataForMesh(const TSharedRef<IDatasmithMeshElement>& InMeshElement, const CADLibrary::FImportParameters& InSceneParameters, const CADLibrary::FMeshParameters&, const FDatasmithTessellationOptions& InTessellationOptions, FDatasmithMeshElementPayload& OutMeshPayload);
-	CORETECHSURFACE_API bool LoadFile(const FString& FileName, FMeshDescription& MeshDescription, const CADLibrary::FImportParameters& ImportParameters, CADLibrary::FMeshParameters& MeshParameters);
-	CORETECHSURFACE_API bool Tessellate(uint64 MainObjectId, const CADLibrary::FImportParameters& ImportParams, FMeshDescription& Mesh, CADLibrary::FMeshParameters& MeshParameters);
+	CORETECHSURFACE_API void AddSurfaceDataForMesh(const TCHAR* CoretechFile, const CADLibrary::FImportParameters& InSceneParameters, const CADLibrary::FMeshParameters&, const FDatasmithTessellationOptions& InTessellationOptions, FDatasmithMeshElementPayload& OutMeshPayload);
+	CORETECHSURFACE_API bool LoadFile(const FString& FileName, const CADLibrary::FImportParameters& InImportParameters, const CADLibrary::FMeshParameters& InMeshParameters, FMeshDescription& OutMeshDescription);
+	CORETECHSURFACE_API bool Tessellate(uint64 MainObjectId, const CADLibrary::FImportParameters& InImportParameters, const CADLibrary::FMeshParameters& InMeshParameters, FMeshDescription& OutMesh);
 }
\ No newline at end of file
diff --git a/Engine/Plugins/Enterprise/DatasmithCADImporter/Source/DatasmithCADTranslator/Private/DatasmithCADTranslator.cpp b/Engine/Plugins/Enterprise/DatasmithCADImporter/Source/DatasmithCADTranslator/Private/DatasmithCADTranslator.cpp
--- a/Engine/Plugins/Enterprise/DatasmithCADImporter/Source/DatasmithCADTranslator/Private/DatasmithCADTranslator.cpp
+++ b/Engine/Plugins/Enterprise/DatasmithCADImporter/Source/DatasmithCADTranslator/Private/DatasmithCADTranslator.cpp
@@ -10,6 +10,7 @@
 #include "DatasmithMeshBuilder.h"
 #include "DatasmithSceneGraphBuilder.h"
 #include "DatasmithUtils.h"
+#include "IDatasmithSceneElements.h"
 
 static TAutoConsoleVariable<int32> CVarStaticCADTranslatorEnableThreadedImport(
 	TEXT("ds.CADTranslator.EnableThreadedImport"),
@@ -196,11 +197,11 @@ bool FDatasmithCADTranslator::LoadStaticMesh(const TSharedRef<IDatasmithMeshElem
 		OutMeshPayload.LodMeshes.Add(MoveTemp(Mesh.GetValue()));
 		if(ImportParameters.bEnableKernelIOTessellation)
 		{
-			CoreTechSurface::AddSurfaceDataForMesh(MeshElement, ImportParameters, MeshParameters, GetCommonTessellationOptions(), OutMeshPayload);
+			CoreTechSurface::AddSurfaceDataForMesh(MeshElement->GetFile(), ImportParameters, MeshParameters, GetCommonTessellationOptions(), OutMeshPayload);
 		}
 		else
 		{
-			CADKernelSurface::AddSurfaceDataForMesh(MeshElement, ImportParameters, MeshParameters, GetCommonTessellationOptions(), OutMeshPayload);
+			CADKernelSurface::AddSurfaceDataForMesh(MeshElement->GetFile(), ImportParameters, MeshParameters, GetCommonTessellationOptions(), OutMeshPayload);
 		}
 	}
 	return OutMeshPayload.LodMeshes.Num() > 0;
diff --git a/Engine/Plugins/Enterprise/DatasmithCADImporter/Source/DatasmithOpenNurbsTranslator/Private/DatasmithOpenNurbsTranslator.cpp b/Engine/Plugins/Enterprise/DatasmithCADImporter/Source/DatasmithOpenNurbsTranslator/Private/DatasmithOpenNurbsTranslator.cpp
--- a/Engine/Plugins/Enterprise/DatasmithCADImporter/Source/DatasmithOpenNurbsTranslator/Private/DatasmithOpenNurbsTranslator.cpp
+++ b/Engine/Plugins/Enterprise/DatasmithCADImporter/Source/DatasmithOpenNurbsTranslator/Private/DatasmithOpenNurbsTranslator.cpp
@@ -3182,14 +3182,14 @@ bool FOpenNurbsTranslatorImpl::TranslateBRep(ON_Brep* Brep, const ON_3dmObjectAt
 		// Ref. visitBRep
 		const FDatasmithOpenNurbsOptions& TessellationOptions = OpenNurbsOptions;
 		LocalSession->SetImportParameters(TessellationOptions.ChordTolerance, TessellationOptions.MaxEdgeLength, TessellationOptions.NormalTolerance, (CADLibrary::EStitchingTechnique) TessellationOptions.StitchingTechnique, false);
-		LocalSession->GetImportParameters().ModelCoordSys = FDatasmithUtils::EModelCoordSystem::ZUp_RightHanded_FBXLegacy;
+		LocalSession->SetModelCoordinateSystem(FDatasmithUtils::EModelCoordSystem::ZUp_RightHanded_FBXLegacy);
 
 		LocalSession->ClearData();
 		LocalSession->SetSceneUnit(GetMetricUnit());
 
 		LocalSession->AddBRep(*Brep, Offset);
 
-		FString FilePath = LocalSession->GetImportParameters().DefineCADFilePath(*OutputPath, *Name);
+		FString FilePath = FPaths::Combine(OutputPath, Name) += TEXT(".ct");
 		if (LocalSession->SaveBrep(FilePath))
 		{
 			MeshElement->SetFile(*FilePath);
@@ -3508,7 +3508,7 @@ bool FDatasmithOpenNurbsTranslator::LoadStaticMesh(const TSharedRef<IDatasmithMe
 
 		if (ImportParameters.bEnableKernelIOTessellation)
 		{
-			CoreTechSurface::AddSurfaceDataForMesh(MeshElement, ImportParameters, MeshParameters, OpenNurbsOptions, OutMeshPayload);
+			CoreTechSurface::AddSurfaceDataForMesh(MeshElement->GetFile(), ImportParameters, MeshParameters, OpenNurbsOptions, OutMeshPayload);
 		}
 		else
 		{
diff --git a/Engine/Plugins/Enterprise/DatasmithCADImporter/Source/DatasmithOpenNurbsTranslator/Private/RhinoCoretechWrapper.cpp b/Engine/Plugins/Enterprise/DatasmithCADImporter/Source/DatasmithOpenNurbsTranslator/Private/RhinoCoretechWrapper.cpp
--- a/Engine/Plugins/Enterprise/DatasmithCADImporter/Source/DatasmithOpenNurbsTranslator/Private/RhinoCoretechWrapper.cpp
+++ b/Engine/Plugins/Enterprise/DatasmithCADImporter/Source/DatasmithOpenNurbsTranslator/Private/RhinoCoretechWrapper.cpp
@@ -399,7 +399,7 @@ void BRepToKernelIOBodyTranslator::CreateCTFace(const ON_BrepFace& Face, TArray<
 
 bool FRhinoCoretechWrapper::Tessellate(FMeshDescription& Mesh, CADLibrary::FMeshParameters& MeshParameters)
 {
-	return CoreTechSurface::Tessellate(MainObjectId, ImportParams, Mesh, MeshParameters);
+	return CoreTechSurface::Tessellate(MainObjectId, ImportParams, MeshParameters, Mesh);
 }
 
 uint64 BRepToKernelIOBodyTranslator::CreateBody(const ON_3dVector& Offset)
diff --git a/Engine/Plugins/Enterprise/DatasmithCADImporter/Source/DatasmithPLMXMLTranslator/Private/DatasmithPlmXmlImporter.cpp b/Engine/Plugins/Enterprise/DatasmithCADImporter/Source/DatasmithPLMXMLTranslator/Private/DatasmithPlmXmlImporter.cpp
--- a/Engine/Plugins/Enterprise/DatasmithCADImporter/Source/DatasmithPLMXMLTranslator/Private/DatasmithPlmXmlImporter.cpp
+++ b/Engine/Plugins/Enterprise/DatasmithCADImporter/Source/DatasmithPLMXMLTranslator/Private/DatasmithPlmXmlImporter.cpp
@@ -306,7 +306,7 @@ namespace PlmXml
 
 				if (ImportParameters.bEnableKernelIOTessellation)
 				{
-					CoreTechSurface::AddSurfaceDataForMesh(MeshElement, ImportParameters, MeshParameters, TessellationOptions, OutMeshPayload);
+					CoreTechSurface::AddSurfaceDataForMesh(MeshElement->GetFile(), ImportParameters, MeshParameters, TessellationOptions, OutMeshPayload);
 				}
 				else
 				{
diff --git a/Engine/Plugins/Enterprise/DatasmithCADImporter/Source/DatasmithWireTranslator/Private/AliasBrepConverter.h b/Engine/Plugins/Enterprise/DatasmithCADImporter/Source/DatasmithWireTranslator/Private/AliasBrepConverter.h
new file mode 100644
--- /dev/null
+++ b/Engine/Plugins/Enterprise/DatasmithCADImporter/Source/DatasmithWireTranslator/Private/AliasBrepConverter.h
@@ -0,0 +1,20 @@
+// Copyright Epic Games, Inc. All Rights Reserved.
+
+#pragma once
+
+class AlDagNode;
+
+// Defined the reference in which the object has to be defined
+enum class EAliasObjectReference
+{
+	LocalReference,  
+	ParentReference, 
+	WorldReference, 
+};
+
+class IAliasBRepConverter
+{
+public:
+	virtual bool AddBRep(AlDagNode& DagNode, EAliasObjectReference ObjectReference) = 0;
+};
+
diff --git a/Engine/Plugins/Enterprise/DatasmithCADImporter/Source/DatasmithWireTranslator/Private/AliasCoretechWrapper.cpp b/Engine/Plugins/Enterprise/DatasmithCADImporter/Source/DatasmithWireTranslator/Private/AliasCoretechWrapper.cpp
deleted file mode 100644
diff --git a/Engine/Plugins/Enterprise/DatasmithCADImporter/Source/DatasmithWireTranslator/Private/AliasCoretechWrapper.h b/Engine/Plugins/Enterprise/DatasmithCADImporter/Source/DatasmithWireTranslator/Private/AliasCoretechWrapper.h
deleted file mode 100644
diff --git a/Engine/Plugins/Enterprise/DatasmithCADImporter/Source/DatasmithWireTranslator/Private/AliasModelToCoretechConverter.cpp b/Engine/Plugins/Enterprise/DatasmithCADImporter/Source/DatasmithWireTranslator/Private/AliasModelToCoretechConverter.cpp
new file mode 100644
--- /dev/null
+++ b/Engine/Plugins/Enterprise/DatasmithCADImporter/Source/DatasmithWireTranslator/Private/AliasModelToCoretechConverter.cpp
@@ -0,0 +1,264 @@
+// Copyright Epic Games, Inc. All Rights Reserved.
+
+#include "AliasModelToCoretechConverter.h"
+
+#include "Hal/PlatformMemory.h"
+
+#include "CoreTechSurfaceHelper.h"
+#include "CoreTechTypes.h"
+
+#ifdef USE_OPENMODEL
+
+// Alias API wrappes object in AlObjects. This is an abstract base class which holds a reference to an anonymous data structure.
+// The only way to compare two AlObjects is to compare their data structure. That is the reason why private fields are made public. 
+// see FAliasCoretechWrapper::AddTrimCurve
+#define private public
+
+#include "AlCurve.h"
+#include "AlDagNode.h"
+#include "AlShell.h"
+#include "AlShellNode.h"
+#include "AlSurface.h"
+#include "AlSurfaceNode.h"
+#include "AlTrimBoundary.h"
+#include "AlTrimCurve.h"
+#include "AlTrimRegion.h"
+#include "AlTM.h"
+
+using namespace CADLibrary;
+
+namespace AliasToCoreTechUtils
+{
+	template<typename Surface_T>
+	uint64 CreateCTNurbs(Surface_T& Surface, EAliasObjectReference InObjectReference, AlMatrix4x4& InAlMatrix)
+	{
+		FNurbsSurface CTSurface;
+
+		CTSurface.ControlPointDimension = 4;  // Control Hull Dimension (3 for non-rational or 4 for rational)
+
+		CTSurface.ControlPointSizeU = Surface.uNumberOfCVsInclMultiples();
+		CTSurface.ControlPointSizeV = Surface.vNumberOfCVsInclMultiples();
+
+		CTSurface.OrderU = Surface.uDegree() + 1;  // U order of the surface
+		CTSurface.OrderV = Surface.vDegree() + 1;  // V order of the surface
+
+		CTSurface.KnotSizeU = Surface.realuNumberOfKnots() + 2;
+		CTSurface.KnotSizeV = Surface.realvNumberOfKnots() + 2;
+
+		CTSurface.KnotValuesU.SetNumUninitialized(CTSurface.KnotSizeU);
+		CTSurface.KnotValuesV.SetNumUninitialized(CTSurface.KnotSizeV);
+
+		Surface.realuKnotVector(CTSurface.KnotValuesU.GetData() + 1);
+		Surface.realvKnotVector(CTSurface.KnotValuesV.GetData() + 1);
+
+		CTSurface.KnotMultiplicityU.Init(1, CTSurface.KnotSizeU);
+		CTSurface.KnotMultiplicityV.Init(1, CTSurface.KnotSizeV);
+
+		CTSurface.KnotValuesU[0] = CTSurface.KnotValuesU[1];
+		CTSurface.KnotValuesV[0] = CTSurface.KnotValuesV[1];
+		CTSurface.KnotValuesU[CTSurface.KnotSizeU - 1] = CTSurface.KnotValuesU[CTSurface.KnotSizeU - 2];
+		CTSurface.KnotValuesV[CTSurface.KnotSizeV - 1] = CTSurface.KnotValuesV[CTSurface.KnotSizeV - 2];
+
+		const int32 NbCoords = CTSurface.ControlPointSizeU * CTSurface.ControlPointSizeV * CTSurface.ControlPointDimension;
+		CTSurface.ControlPoints.SetNumUninitialized(NbCoords);
+
+		if (InObjectReference == EAliasObjectReference::WorldReference)
+		{
+			Surface.CVsWorldPositionInclMultiples(CTSurface.ControlPoints.GetData());
+		}
+		else if (InObjectReference == EAliasObjectReference::ParentReference)
+		{
+			AlTM TranformMatrix(InAlMatrix);
+			Surface.CVsAffectedPositionInclMultiples(TranformMatrix, CTSurface.ControlPoints.GetData());
+		}
+		else  // EAliasObjectReference::LocalReference
+		{
+			Surface.CVsUnaffectedPositionInclMultiples(CTSurface.ControlPoints.GetData());
+		}
+
+		uint64 CTSurfaceID = 0;
+		return CTKIO_CreateNurbsSurface(CTSurface, CTSurfaceID) ? CTSurfaceID : 0;
+	}
+}
+
+uint64 FAliasModelToCoretechConverter::AddTrimCurve(AlTrimCurve& TrimCurve)
+{
+	FNurbsCurve CTCurve;
+
+	curveFormType Form = TrimCurve.form();
+	CTCurve.Order = TrimCurve.degree() + 1;
+	CTCurve.ControlPointSize = TrimCurve.numberOfCVs();
+	CTCurve.ControlPointDimension = 3;
+	CTCurve.KnotSize = TrimCurve.realNumberOfKnots() + 2;
+
+	CTCurve.ControlPoints.SetNumUninitialized(CTCurve.ControlPointSize * CTCurve.ControlPointDimension);
+	CTCurve.KnotValues.SetNumUninitialized(CTCurve.KnotSize);
+
+	using UVPoint2 = double[3];
+	TrimCurve.CVsUVPosition(CTCurve.KnotValues.GetData(), (UVPoint2*)CTCurve.ControlPoints.GetData());
+
+	CTCurve.KnotMultiplicity.Init(1, CTCurve.KnotSize);
+
+	TrimCurve.realKnotVector(CTCurve.KnotValues.GetData() + 1);
+
+	CTCurve.KnotValues[0] = CTCurve.KnotValues[1];
+	CTCurve.KnotValues[CTCurve.KnotSize - 1] = CTCurve.KnotValues[CTCurve.KnotSize - 2];
+
+	uint64 CoedgeID = 0;
+	if (CTKIO_CreateCoedge(CTCurve, (bool) TrimCurve.isReversed(), CoedgeID))
+	{
+		// Build topo
+		if (AlTrimCurve *TwinCurve = TrimCurve.getTwinCurve())
+		{
+			if (uint64* TwinCoedgeID = AlEdge2CTEdge.Find(TwinCurve->fSpline))
+			{
+				CTKIO_MatchCoedges(*TwinCoedgeID, CoedgeID);
+			}
+			// only TrimCurve with twin need to be in the map
+			AlEdge2CTEdge.Add(TrimCurve.fSpline, CoedgeID);
+		}
+		return CoedgeID;
+	}
+
+	return 0;
+}
+
+uint64 FAliasModelToCoretechConverter::AddTrimBoundary(AlTrimBoundary& TrimBoundary)
+{
+	TArray<uint64> Edges;
+	for(AlTrimCurve *TrimCurve = TrimBoundary.firstCurve(); TrimCurve; TrimCurve = TrimCurve->nextCurve())
+	{
+		if (uint64 CoedgeID = AddTrimCurve(*TrimCurve))
+		{
+			Edges.Add(CoedgeID);
+		}
+	}
+
+	uint64 LoopID;
+	return CTKIO_CreateLoop(Edges, LoopID) ? LoopID : 0;
+}
+
+uint64 FAliasModelToCoretechConverter::AddTrimRegion(AlTrimRegion& TrimRegion, EAliasObjectReference InObjectReference, AlMatrix4x4& InAlMatrix, bool bInOrientation)
+{
+	uint64 SurfaceID = AliasToCoreTechUtils::CreateCTNurbs(TrimRegion, InObjectReference, InAlMatrix);
+	if (SurfaceID == 0)
+	{
+		return 0;
+	}
+
+	TArray<uint64> Boundaries;
+	for(AlTrimBoundary *TrimBoundary = TrimRegion.firstBoundary(); TrimBoundary; TrimBoundary = TrimBoundary->nextBoundary())
+	{
+		if (uint64 LoopId = AddTrimBoundary(*TrimBoundary))
+		{
+			Boundaries.Add(LoopId);
+		}
+	}
+
+	uint64 FaceID;
+	return CTKIO_CreateFace(SurfaceID, bInOrientation, Boundaries, FaceID) ? FaceID : 0;
+}
+
+void FAliasModelToCoretechConverter::AddFace(AlSurface& Surface, EAliasObjectReference InObjectReference, AlMatrix4x4& InAlMatrix, bool bInOrientation, TArray<uint64>& OutFaceList)
+{
+	if (AlTrimRegion *TrimRegion = Surface.firstTrimRegion())
+	{
+		while (TrimRegion)
+		{
+			if(uint64 FaceID = AddTrimRegion(*TrimRegion, InObjectReference, InAlMatrix, bInOrientation))
+			{
+				OutFaceList.Add(FaceID);
+			}
+
+			TrimRegion = TrimRegion->nextRegion();
+		}
+
+		return;
+	}
+
+	uint64 SurfaceID = AliasToCoreTechUtils::CreateCTNurbs(Surface, InObjectReference, InAlMatrix);
+
+	uint64 FaceID;
+	if (CTKIO_CreateFace(SurfaceID, bInOrientation, TArray<uint64>(), FaceID) && FaceID != 0)
+	{
+		OutFaceList.Add(FaceID);
+	}
+}
+
+void FAliasModelToCoretechConverter::AddShell(AlShell& Shell, EAliasObjectReference InObjectReference, AlMatrix4x4& InAlMatrix, bool bInOrientation, TArray<uint64>& OutFaceList)
+{
+	for (AlTrimRegion* TrimRegion = Shell.firstTrimRegion(); TrimRegion; TrimRegion = TrimRegion->nextRegion())
+	{
+		if (uint64 FaceID = AddTrimRegion(*TrimRegion, InObjectReference, InAlMatrix, bInOrientation))
+		{
+			OutFaceList.Add(FaceID);
+		}
+	}
+}
+
+bool FAliasModelToCoretechConverter::AddBRep(AlDagNode& DagNode, EAliasObjectReference InObjectReference)
+{
+	TArray<uint64> FaceList;
+
+	AlEdge2CTEdge.Empty();
+
+	boolean bAlOrientation;
+	DagNode.getSurfaceOrientation(bAlOrientation);
+	bool bOrientation = (bool)bAlOrientation;
+
+	AlMatrix4x4 AlMatrix;
+	if (InObjectReference == EAliasObjectReference::ParentReference)
+	{
+		DagNode.localTransformationMatrix(AlMatrix);
+	}
+
+	AlObjectType objectType = DagNode.type();
+	switch (objectType)
+	{
+		// Push all leaf nodes into 'leaves'
+	case kShellNodeType:
+	{
+		if (AlShellNode* ShellNode = DagNode.asShellNodePtr())
+		{
+			if (AlShell* Shell = ShellNode->shell())
+			{
+				AddShell(*Shell, InObjectReference, AlMatrix, bOrientation, FaceList);
+			}
+		}
+		break;
+	}
+
+	case kSurfaceNodeType:
+	{
+		if (AlSurfaceNode* SurfaceNode = DagNode.asSurfaceNodePtr())
+		{
+			if (AlSurface* Surface = SurfaceNode->surface())
+			{
+				AddFace(*Surface, InObjectReference, AlMatrix, bOrientation, FaceList);
+			}
+		}
+		break;
+	}
+
+	default:
+		break;
+
+	}
+
+	if (FaceList.Num() == 0)
+	{
+		return false;
+	}
+
+	// Create body from faces
+	uint64 BodyID;
+	if (CADLibrary::CTKIO_CreateBody(FaceList, BodyID))
+	{
+		// Setup parenting
+		return CADLibrary::CTKIO_AddBodies({ BodyID }, MainObjectId) ? true : false;
+	}
+
+	return false;
+}
+
+#endif
diff --git a/Engine/Plugins/Enterprise/DatasmithCADImporter/Source/DatasmithWireTranslator/Private/AliasModelToCoretechConverter.h b/Engine/Plugins/Enterprise/DatasmithCADImporter/Source/DatasmithWireTranslator/Private/AliasModelToCoretechConverter.h
new file mode 100644
--- /dev/null
+++ b/Engine/Plugins/Enterprise/DatasmithCADImporter/Source/DatasmithWireTranslator/Private/AliasModelToCoretechConverter.h
@@ -0,0 +1,52 @@
+// Copyright Epic Games, Inc. All Rights Reserved.
+
+#pragma once
+
+#include "CoreMinimal.h"
+
+#include "AliasBRepConverter.h"
+#include "CADModelToCoretechConverterBase.h"
+#include "CTSession.h"
+
+class AlDagNode;
+class AlShell;
+class AlSurface;
+class AlTrimBoundary;
+class AlTrimCurve;
+class AlTrimRegion;
+class AlCurve;
+
+typedef double AlMatrix4x4[4][4];
+
+class FAliasModelToCoretechConverter : public FCADModelToCoretechConverterBase, public IAliasBRepConverter
+{
+public:
+
+	FAliasModelToCoretechConverter(const TCHAR* InOwner)
+		: FCADModelToCoretechConverterBase(InOwner)
+	{
+		// Unit for CoreTech session is set to cm, 0.01, because Wire's unit is cm. Consequently, Scale factor is set to 1.
+		ImportParams.MetricUnit = 0.01;
+		ImportParams.ScaleFactor = 1;
+		ImportParams.bEnableKernelIOTessellation = true;
+	}
+
+	virtual bool AddBRep(AlDagNode& DagNode, EAliasObjectReference ObjectReference) override;
+
+protected:
+
+	/**
+	 * Create a CT coedge (represent the use of an edge by a face).
+	 * @param TrimCurve: A curve in parametric surface space, part of a trim boundary.
+	 */
+	uint64 AddTrimCurve(AlTrimCurve& TrimCurve);
+	uint64 AddTrimBoundary(AlTrimBoundary& TrimBoundary);
+
+	uint64 AddTrimRegion(AlTrimRegion& InTrimRegion, EAliasObjectReference InObjectReference, AlMatrix4x4& InAlMatrix, bool bInOrientation);
+	void AddFace(AlSurface& InSurface, EAliasObjectReference InObjectReference, AlMatrix4x4& InAlMatrix, bool bInOrientation, TArray<uint64>& OutFaceLis);
+	void AddShell(AlShell& InShell, EAliasObjectReference InObjectReference, AlMatrix4x4& InAlMatrix, bool bInOrientation, TArray<uint64>& OutFaceLis);
+
+protected:
+	TMap<void*, uint64>  AlEdge2CTEdge;
+};
+
diff --git a/Engine/Plugins/Enterprise/DatasmithCADImporter/Source/DatasmithWireTranslator/Private/DatasmithWireTranslator.cpp b/Engine/Plugins/Enterprise/DatasmithCADImporter/Source/DatasmithWireTranslator/Private/DatasmithWireTranslator.cpp
--- a/Engine/Plugins/Enterprise/DatasmithCADImporter/Source/DatasmithWireTranslator/Private/DatasmithWireTranslator.cpp
+++ b/Engine/Plugins/Enterprise/DatasmithCADImporter/Source/DatasmithWireTranslator/Private/DatasmithWireTranslator.cpp
@@ -13,7 +13,6 @@
 #include "Misc/FileHelper.h"
 #include "Misc/Paths.h"
 #include "OpenModelUtils.h"
-#include "Utility/DatasmithMeshHelper.h"
 
 #include "StaticMeshDescription.h"
 #include "StaticMeshOperations.h"
@@ -26,8 +25,9 @@
 #include "Modules/ModuleManager.h"
 #endif
 
-#include "AliasCoretechWrapper.h" // requires CoreTech as public dependency
+#include "AliasModelToCoretechConverter.h" // requires CoreTech as public dependency
 #include "CADInterfacesModule.h"
+#include "AliasBrepConverter.h"
 #include "CoreTechSurfaceExtension.h"
 #include "CoreTechSurfaceHelper.h"
 
@@ -141,28 +141,45 @@ public:
 		DatasmithScene->SetProductName(TEXT("Alias Tools"));
 		DatasmithScene->SetProductVersion(TEXT("Alias 2019"));
 
-		LocalSession = FAliasCoretechWrapper::GetSharedSession();
+		bool bEnableKernelIOTessellation = true;
+		if (IConsoleVariable* CVar = IConsoleManager::Get().FindConsoleVariable(TEXT("ds.CADTranslator.EnableKernelIOTessellation")))
+		{
+			bEnableKernelIOTessellation = CVar->GetInt() != 0;
+		}
+
+		if (bEnableKernelIOTessellation)
+		{
+			TSharedRef<FAliasModelToCoretechConverter> AliasToCoretechConverter = MakeShared<FAliasModelToCoretechConverter>(TEXT("Al2CTSharedSession"));
+			CADModelConverter = AliasToCoretechConverter;
+			AliasBRepConverter = AliasToCoretechConverter;
+		}
+		else
+		{
+			//TSharedRef<FAliasCADKernelWrapper> AliasCoretechWrapper = MakeShared<FAliasCADKernelWrapper>(TEXT("Al2CTSharedSession"));
+			//CADWrapper = AliasCoretechWrapper;
+			//AliasWrapper = AliasCoretechWrapper;
+		}
 	}
 
 	~FWireTranslatorImpl()
 	{
 		AlUniverse::deleteAll();
-		LocalSession.Reset();
 	}
 
 	bool Read();
-	TOptional<FMeshDescription> GetMeshDescription(TSharedRef<IDatasmithMeshElement> MeshElement, CADLibrary::FMeshParameters& MeshParameters);
+	TOptional<FMeshDescription> GetMeshDescription(TSharedRef<IDatasmithMeshElement> MeshElement, CADLibrary::FMeshParameters& OutMeshParameters);
 	TOptional<FMeshDescription> GetMeshDescription(TSharedRef<IDatasmithMeshElement> MeshElement, CADLibrary::FMeshParameters& MeshParameters, TSharedRef<BodyData> BodyTemp);
 
 	void SetTessellationOptions(const FDatasmithTessellationOptions& Options);
 	void SetOutputPath(const FString& Path) { OutputPath = Path; }
 
-	CADLibrary::FImportParameters& GetImportParameters()
+	bool LoadStaticMesh(const TSharedRef<IDatasmithMeshElement> MeshElement, FDatasmithMeshElementPayload& OutMeshPayload, const FDatasmithTessellationOptions& InTessellationOptions);
+
+	void AddSurfaceDataForMesh(const TCHAR* InFilePath, const FMeshParameters& InMeshParameters, const FDatasmithTessellationOptions& InTessellationOptions, FDatasmithMeshElementPayload& OutMeshPayload) const
 	{
-		return LocalSession->GetImportParameters();
+		CADModelConverter->AddSurfaceDataForMesh(InFilePath, InMeshParameters, InTessellationOptions, OutMeshPayload);
 	}
 
-
 private:
 
 	struct FDagNodeInfo
@@ -255,6 +272,7 @@ private:
 	FString SceneFullPath;
 
 	FDatasmithTessellationOptions TessellationOptions;
+
 	// Hash value of the scene file used to check if the file has been modified for re-import
 	uint32 SceneFileHash;
 
@@ -282,7 +300,9 @@ private:
 	// If > 0, then the archive is corrupt.
 	int32 NumCRCErrors;
 
-	TSharedPtr<FAliasCoretechWrapper> LocalSession;
+	TSharedPtr<ICADModelConverter> CADModelConverter;
+	TSharedPtr<IAliasBRepConverter> AliasBRepConverter;
+
 };
 
 void FWireTranslatorImpl::SetTessellationOptions(const FDatasmithTessellationOptions& Options)
@@ -301,7 +321,7 @@ bool FWireTranslatorImpl::Read()
 		return false;
 	}
 
-	LocalSession->SetImportParameters(TessellationOptions.ChordTolerance, TessellationOptions.MaxEdgeLength, TessellationOptions.NormalTolerance, (CADLibrary::EStitchingTechnique) TessellationOptions.StitchingTechnique, true);
+	CADModelConverter->SetImportParameters(TessellationOptions.ChordTolerance, TessellationOptions.MaxEdgeLength, TessellationOptions.NormalTolerance, (CADLibrary::EStitchingTechnique) TessellationOptions.StitchingTechnique, true);
 
 	AlRetrieveOptions options;
 	AlUniverse::retrieveOptions(options);
@@ -1728,12 +1748,9 @@ bool FWireTranslatorImpl::RecurseDagForLeavesNoMerge(AlDagNode* FirstDagNode, co
 
 TOptional<FMeshDescription> FWireTranslatorImpl::MeshDagNodeWithExternalMesher(AlDagNode& DagNode, TSharedRef<IDatasmithMeshElement> MeshElement, CADLibrary::FMeshParameters& MeshParameters)
 {
-	LocalSession->ClearData();
-
-
 	// Wire unit is cm
-	LocalSession->SetSceneUnit(0.01);
-	
+	CADModelConverter->InitializeProcess(0.01);
+
 	EAliasObjectReference ObjectReference = EAliasObjectReference::LocalReference;
 
 	if (MeshParameters.bIsSymmetric)
@@ -1742,62 +1759,53 @@ TOptional<FMeshDescription> FWireTranslatorImpl::MeshDagNodeWithExternalMesher(A
 		ObjectReference = EAliasObjectReference::WorldReference;
 	}
 
-	TArray<AlDagNode*> DagNodeSet;
-	DagNodeSet.Add(&DagNode);
-	LocalSession->AddBRep(DagNodeSet, ObjectReference);
+	AliasBRepConverter->AddBRep(DagNode, ObjectReference);
 
-	FString FilePath = LocalSession->GetImportParameters().DefineCADFilePath(*OutputPath, MeshElement->GetName());
-	if (LocalSession->SaveBrep(FilePath))
-	{
-		MeshElement->SetFile(*FilePath);
-	}
+	CADModelConverter->RepairTopology();
+
+	CADModelConverter->SaveBRep(*OutputPath, MeshElement);
 
 	FMeshDescription MeshDescription;
 	DatasmithMeshHelper::PrepareAttributeForStaticMesh(MeshDescription);
 
-	LocalSession->Tessellate(MeshDescription, MeshParameters);
+	CADModelConverter->Tessellate(MeshParameters, MeshDescription);
 
 	return MoveTemp(MeshDescription);
 }
 
 TOptional<FMeshDescription> FWireTranslatorImpl::MeshDagNodeWithExternalMesher(TSharedRef<BodyData> Body, TSharedRef<IDatasmithMeshElement> MeshElement, CADLibrary::FMeshParameters& MeshParameters)
 {
-	LocalSession->ClearData();
-
 	// Wire unit is cm
-	LocalSession->SetSceneUnit(0.01);
-	
+	CADModelConverter->InitializeProcess(0.01);
+
 	EAliasObjectReference ObjectReference = EAliasObjectReference::LocalReference;
 	if (MeshParameters.bIsSymmetric)
 	{
 		// All actors of a Alias symmetric layer are defined in the world Reference i.e. they have identity transform. So Mesh actor has to be defined in the world reference. 
 		ObjectReference = EAliasObjectReference::WorldReference;
 	}
-	else if (GetImportParameters().StitchingTechnique == StitchingSew)
+	else if (TessellationOptions.StitchingTechnique == EDatasmithCADStitchingTechnique::StitchingSew)
 	{
 		// In the case of StitchingSew, AlDagNode children of a GroupNode are merged together. To be merged, they have to be defined in the reference of parent GroupNode.
 		ObjectReference = EAliasObjectReference::ParentReference;
 	}
 
-	LocalSession->AddBRep(Body->ShellSet, ObjectReference);
-
-	FString FilePath = LocalSession->GetImportParameters().DefineCADFilePath(*OutputPath, MeshElement->GetName());
-	if (LocalSession->SaveBrep(FilePath))
+	for(AlDagNode* DagNode : Body->ShellSet)
 	{
-		MeshElement->SetFile(*FilePath);
+		AliasBRepConverter->AddBRep(*DagNode, ObjectReference);
 	}
 
-	FMeshDescription MeshDescription;
-	DatasmithMeshHelper::PrepareAttributeForStaticMesh(MeshDescription);
+	CADModelConverter->SaveBRep(*OutputPath, MeshElement);
 
-	LocalSession->Tessellate(MeshDescription, MeshParameters);
+	FMeshDescription MeshDescription;
+	CADModelConverter->Tessellate(MeshParameters, MeshDescription);
 
 	return MoveTemp(MeshDescription);
 }
 
 TOptional<FMeshDescription> FWireTranslatorImpl::GetMeshOfShellNode(AlDagNode& DagNode, TSharedRef<IDatasmithMeshElement> MeshElement, CADLibrary::FMeshParameters& MeshParameters)
 {
-	if (LocalSession->IsSessionValid())
+	if (CADModelConverter->IsSessionValid())
 	{
 		TOptional< FMeshDescription > UEMesh = MeshDagNodeWithExternalMesher(DagNode, MeshElement, MeshParameters);
 		return UEMesh;
@@ -1861,7 +1869,6 @@ TOptional<FMeshDescription> FWireTranslatorImpl::GetMeshOfNodeMesh(AlMeshNode& M
 	return TOptional< FMeshDescription >();
 }
 
-
 TOptional<FMeshDescription> FWireTranslatorImpl::GetMeshDescription(TSharedRef<IDatasmithMeshElement> MeshElement, CADLibrary::FMeshParameters& MeshParameters, TSharedRef<BodyData> Body)
 {
 	if(Body->ShellSet.Num() == 0 )
@@ -1896,7 +1903,6 @@ TOptional<FMeshDescription> FWireTranslatorImpl::GetMeshDescription(TSharedRef<I
 	}
 }
 
-
 TOptional<FMeshDescription> FWireTranslatorImpl::GetMeshDescription(TSharedRef<IDatasmithMeshElement> MeshElement, CADLibrary::FMeshParameters& MeshParameters)
 {
 	AlDagNode ** DagNodeTemp = MeshElementToAlDagNodeMap.Find(&MeshElement.Get());
@@ -1971,6 +1977,18 @@ TOptional< FMeshDescription > FWireTranslatorImpl::ImportMesh(AlMesh& CurrentMes
 	return MoveTemp(MeshDescription);
 }
 
+bool FWireTranslatorImpl::LoadStaticMesh(const TSharedRef<IDatasmithMeshElement> MeshElement, FDatasmithMeshElementPayload& OutMeshPayload, const FDatasmithTessellationOptions& InTessellationOptions)
+{
+	CADLibrary::FMeshParameters MeshParameters;
+	if (TOptional< FMeshDescription > Mesh = GetMeshDescription(MeshElement, MeshParameters))
+	{
+		OutMeshPayload.LodMeshes.Add(MoveTemp(Mesh.GetValue()));
+		AddSurfaceDataForMesh(MeshElement->GetFile(), MeshParameters, InTessellationOptions, OutMeshPayload);
+	}
+	return OutMeshPayload.LodMeshes.Num() > 0;
+}
+
+
 #endif
 
 //////////////////////////////////////////////////////////////////////////
@@ -2073,25 +2091,12 @@ void FDatasmithWireTranslator::UnloadScene()
 bool FDatasmithWireTranslator::LoadStaticMesh(const TSharedRef<IDatasmithMeshElement> MeshElement, FDatasmithMeshElementPayload& OutMeshPayload)
 {
 #ifdef USE_OPENMODEL
-	CADLibrary::FImportParameters& ImportParameters = Translator->GetImportParameters();
-	CADLibrary::FMeshParameters MeshParameters;
-	if (TOptional< FMeshDescription > Mesh = Translator->GetMeshDescription(MeshElement, MeshParameters))
+	if (Translator)
 	{
-		OutMeshPayload.LodMeshes.Add(MoveTemp(Mesh.GetValue()));
-
-		if (ImportParameters.bEnableKernelIOTessellation)
-		{
-			CoreTechSurface::AddSurfaceDataForMesh(MeshElement, ImportParameters, MeshParameters, GetCommonTessellationOptions(), OutMeshPayload);
-		}
-		else
-		{
-			//CADKernelSurface::AddSurfaceDataForMesh(MeshElement, ImportParameters, MeshParameters, GetCommonTessellationOptions(), OutMeshPayload);
-		}
+		return Translator->LoadStaticMesh(MeshElement, OutMeshPayload, GetCommonTessellationOptions());
 	}
-	return OutMeshPayload.LodMeshes.Num() > 0;
-#else
-	return false;
 #endif
+	return false;
 }
 
 void FDatasmithWireTranslator::SetSceneImportOptions(TArray<TStrongObjectPtr<UDatasmithOptionsBase>>& Options)
diff --git a/Engine/Plugins/Enterprise/DatasmithCADImporter/Source/ParametricSurface/ParametricSurface.Build.cs b/Engine/Plugins/Enterprise/DatasmithCADImporter/Source/ParametricSurface/ParametricSurface.Build.cs
--- a/Engine/Plugins/Enterprise/DatasmithCADImporter/Source/ParametricSurface/ParametricSurface.Build.cs
+++ b/Engine/Plugins/Enterprise/DatasmithCADImporter/Source/ParametricSurface/ParametricSurface.Build.cs
@@ -20,6 +20,7 @@ namespace UnrealBuildTool.Rules
 				{
 					"Core",
 					"CoreUObject",
+					"CADTools",
                     "DatasmithContent",
 					"DatasmithTranslator",
 				}
diff --git a/Engine/Plugins/Enterprise/DatasmithCADImporter/Source/ParametricSurface/Public/CADModelConverter.h b/Engine/Plugins/Enterprise/DatasmithCADImporter/Source/ParametricSurface/Public/CADModelConverter.h
new file mode 100644
--- /dev/null
+++ b/Engine/Plugins/Enterprise/DatasmithCADImporter/Source/ParametricSurface/Public/CADModelConverter.h
@@ -0,0 +1,68 @@
+// Copyright Epic Games, Inc. All Rights Reserved.
+
+#pragma once
+
+#include "CoreMinimal.h"
+
+#include "CADOptions.h"
+
+class IDatasmithMeshElement;
+struct FDatasmithMeshElementPayload;
+struct FDatasmithTessellationOptions;
+
+namespace CADLibrary
+{
+struct FImportParameters;
+struct FMeshParameters;
+
+/**
+ * Interface to convert CAD Model (like Alias, Rhino, ...) into the internal CAD Modeler (CADKernel, ...)
+ * The derived class must also implement IAddXXXBrepModel according to the CAD model (IAddAliasBrepModel, IAddRhinoBrepModel, ...)
+ * The derived class is instantiated for an import session. 
+ * It processes a subset of the model corresponding to a static mesh actor.
+ * The main step of the process is:
+ *    - InitializeProcess
+ *    - Fill the modeler database with the subset of the model: 
+ *        - this step is done with AddBRep function of the format interface (IAddAliasBrepModel, IAddRhinoBrepModel)
+ *        - many bodies can be added to the model to generate only one static mesh actor
+ *    - RepairTopology if needed
+ *    - SaveBrep(const FString& FilePath) for retessellate purpose
+ *    - Tessellate
+ *
+ * In a second step, the saved file is associated to its UStaticMesh with AddSurfaceDataForMesh function
+ */
+
+class ICADModelConverter
+{
+public:
+	virtual ~ICADModelConverter() = default;
+
+	/**	 
+	 * @param FileMetricUnit: Handle input file unit number of meters per file unit.
+	 * eg. For a file in inches, arg should be 0.0254
+	 */
+	virtual void InitializeProcess(double InMetricUnit) = 0;
+
+	virtual bool RepairTopology() = 0;
+	virtual bool SaveBRep(const TCHAR* OutputPath, TSharedRef<IDatasmithMeshElement>& MeshElement) = 0;
+	virtual bool Tessellate(const CADLibrary::FMeshParameters& InMeshParameters, FMeshDescription& OutMeshDescription) = 0;
+
+	/**
+	 * Set Import parameters,
+	 * Tack care to set scale factor before because import parameters will be scale according to scale factor
+	 * @param ChordTolerance : SAG
+	 * @param MaxEdgeLength : max length of element's edge
+	 * @param NormalTolerance : Angle between two adjacent triangles
+	 * @param StitchingTechnique : CAD topology correction technique
+	 * @param bScaleUVMap : Scale the UV map to a world unit.
+	 */
+	virtual void SetImportParameters(float ChordTolerance, float MaxEdgeLength, float NormalTolerance, CADLibrary::EStitchingTechnique StitchingTechnique, bool bScaleUVMap) = 0;
+
+	virtual bool IsSessionValid() = 0;
+
+	virtual void AddSurfaceDataForMesh(const TCHAR* InFilePath, const FMeshParameters& InMeshParameters, const FDatasmithTessellationOptions& InTessellationOptions, FDatasmithMeshElementPayload& OutMeshPayload) const = 0;
+
+};
+
+}
+
