[Alias] Fix memory leak

diff --git a/Engine/Plugins/Enterprise/DatasmithCADImporter/Source/DatasmithWireTranslator/Private/AliasModelToCoretechConverter.cpp b/Engine/Plugins/Enterprise/DatasmithCADImporter/Source/DatasmithWireTranslator/Private/AliasModelToCoretechConverter.cpp
--- a/Engine/Plugins/Enterprise/DatasmithCADImporter/Source/DatasmithWireTranslator/Private/AliasModelToCoretechConverter.cpp
+++ b/Engine/Plugins/Enterprise/DatasmithCADImporter/Source/DatasmithWireTranslator/Private/AliasModelToCoretechConverter.cpp
@@ -30,7 +30,7 @@ using namespace CADLibrary;
 namespace AliasToCoreTechUtils
 {
 	template<typename Surface_T>
-	uint64 CreateCTNurbs(Surface_T& Surface, EAliasObjectReference InObjectReference, AlMatrix4x4& InAlMatrix)
+	uint64 CreateCTNurbs(Surface_T& Surface, EAliasObjectReference InObjectReference, const AlMatrix4x4& InAlMatrix)
 	{
 		FNurbsSurface CTSurface;
 
@@ -81,7 +81,7 @@ namespace AliasToCoreTechUtils
 	}
 }
 
-uint64 FAliasModelToCoretechConverter::AddTrimCurve(AlTrimCurve& TrimCurve)
+uint64 FAliasModelToCoretechConverter::AddTrimCurve(const AlTrimCurve& TrimCurve)
 {
 	FNurbsCurve CTCurve;
 
@@ -123,10 +123,10 @@ uint64 FAliasModelToCoretechConverter::AddTrimCurve(AlTrimCurve& TrimCurve)
 	return 0;
 }
 
-uint64 FAliasModelToCoretechConverter::AddTrimBoundary(AlTrimBoundary& TrimBoundary)
+uint64 FAliasModelToCoretechConverter::AddTrimBoundary(const AlTrimBoundary& TrimBoundary)
 {
 	TArray<uint64> Edges;
-	for(AlTrimCurve *TrimCurve = TrimBoundary.firstCurve(); TrimCurve; TrimCurve = TrimCurve->nextCurve())
+	for(TUniquePtr<AlTrimCurve> TrimCurve(TrimBoundary.firstCurve()); TrimCurve.IsValid(); TrimCurve = TUniquePtr<AlTrimCurve>(TrimCurve->nextCurve()))
 	{
 		if (uint64 CoedgeID = AddTrimCurve(*TrimCurve))
 		{
@@ -138,7 +138,7 @@ uint64 FAliasModelToCoretechConverter::AddTrimBoundary(AlTrimBoundary& TrimBound
 	return CTKIO_CreateLoop(Edges, LoopID) ? LoopID : 0;
 }
 
-uint64 FAliasModelToCoretechConverter::AddTrimRegion(AlTrimRegion& TrimRegion, EAliasObjectReference InObjectReference, AlMatrix4x4& InAlMatrix, bool bInOrientation)
+uint64 FAliasModelToCoretechConverter::AddTrimRegion(const AlTrimRegion& TrimRegion, EAliasObjectReference InObjectReference, const AlMatrix4x4& InAlMatrix, bool bInOrientation)
 {
 	uint64 SurfaceID = AliasToCoreTechUtils::CreateCTNurbs(TrimRegion, InObjectReference, InAlMatrix);
 	if (SurfaceID == 0)
@@ -147,7 +147,7 @@ uint64 FAliasModelToCoretechConverter::AddTrimRegion(AlTrimRegion& TrimRegion, E
 	}
 
 	TArray<uint64> Boundaries;
-	for(AlTrimBoundary *TrimBoundary = TrimRegion.firstBoundary(); TrimBoundary; TrimBoundary = TrimBoundary->nextBoundary())
+	for (TUniquePtr<AlTrimBoundary> TrimBoundary(TrimRegion.firstBoundary()); TrimBoundary.IsValid(); TrimBoundary = TUniquePtr<AlTrimBoundary>(TrimBoundary->nextBoundary()))
 	{
 		if (uint64 LoopId = AddTrimBoundary(*TrimBoundary))
 		{
@@ -159,20 +159,18 @@ uint64 FAliasModelToCoretechConverter::AddTrimRegion(AlTrimRegion& TrimRegion, E
 	return CTKIO_CreateFace(SurfaceID, bInOrientation, Boundaries, FaceID) ? FaceID : 0;
 }
 
-void FAliasModelToCoretechConverter::AddFace(AlSurface& Surface, EAliasObjectReference InObjectReference, AlMatrix4x4& InAlMatrix, bool bInOrientation, TArray<uint64>& OutFaceList)
+void FAliasModelToCoretechConverter::AddFace(const AlSurface& Surface, EAliasObjectReference InObjectReference, const AlMatrix4x4& InAlMatrix, bool bInOrientation, TArray<uint64>& OutFaceList)
 {
-	if (AlTrimRegion *TrimRegion = Surface.firstTrimRegion())
+	TUniquePtr<AlTrimRegion> TrimRegion(Surface.firstTrimRegion());
+	if (TrimRegion.IsValid())
 	{
-		while (TrimRegion)
+		for (; TrimRegion.IsValid(); TrimRegion = TUniquePtr<AlTrimRegion>(TrimRegion->nextRegion()))
 		{
-			if(uint64 FaceID = AddTrimRegion(*TrimRegion, InObjectReference, InAlMatrix, bInOrientation))
+			if (uint64 FaceID = AddTrimRegion(*TrimRegion, InObjectReference, InAlMatrix, bInOrientation))
 			{
 				OutFaceList.Add(FaceID);
 			}
-
-			TrimRegion = TrimRegion->nextRegion();
 		}
-
 		return;
 	}
 
@@ -185,9 +183,9 @@ void FAliasModelToCoretechConverter::AddFace(AlSurface& Surface, EAliasObjectRef
 	}
 }
 
-void FAliasModelToCoretechConverter::AddShell(AlShell& Shell, EAliasObjectReference InObjectReference, AlMatrix4x4& InAlMatrix, bool bInOrientation, TArray<uint64>& OutFaceList)
+void FAliasModelToCoretechConverter::AddShell(const AlShell& Shell, EAliasObjectReference InObjectReference, const AlMatrix4x4& InAlMatrix, bool bInOrientation, TArray<uint64>& OutFaceList)
 {
-	for (AlTrimRegion* TrimRegion = Shell.firstTrimRegion(); TrimRegion; TrimRegion = TrimRegion->nextRegion())
+	for (TUniquePtr<AlTrimRegion> TrimRegion(Shell.firstTrimRegion()); TrimRegion.IsValid(); TrimRegion = TUniquePtr<AlTrimRegion>(TrimRegion->nextRegion()))
 	{
 		if (uint64 FaceID = AddTrimRegion(*TrimRegion, InObjectReference, InAlMatrix, bInOrientation))
 		{
@@ -213,16 +211,18 @@ bool FAliasModelToCoretechConverter::AddBRep(AlDagNode& DagNode, EAliasObjectRef
 	}
 
 	AlObjectType objectType = DagNode.type();
-	switch (objectType)
+	switch (objectType) 
 	{
 		// Push all leaf nodes into 'leaves'
 	case kShellNodeType:
 	{
 		if (AlShellNode* ShellNode = DagNode.asShellNodePtr())
 		{
-			if (AlShell* Shell = ShellNode->shell())
+			AlShell* ShellPtr = ShellNode->shell();
+			if (AlIsValid(ShellPtr))
 			{
-				AddShell(*Shell, InObjectReference, AlMatrix, bOrientation, FaceList);
+				TUniquePtr<AlShell> Shell(ShellPtr);
+				AddShell(*ShellPtr, InObjectReference, AlMatrix, bOrientation, FaceList);
 			}
 		}
 		break;
@@ -232,9 +232,11 @@ bool FAliasModelToCoretechConverter::AddBRep(AlDagNode& DagNode, EAliasObjectRef
 	{
 		if (AlSurfaceNode* SurfaceNode = DagNode.asSurfaceNodePtr())
 		{
-			if (AlSurface* Surface = SurfaceNode->surface())
+			AlSurface* SurfacePtr = SurfaceNode->surface();
+			if (AlIsValid(SurfacePtr))
 			{
-				AddFace(*Surface, InObjectReference, AlMatrix, bOrientation, FaceList);
+				TUniquePtr<AlSurface> Surface(SurfacePtr);
+				AddFace(*SurfacePtr, InObjectReference, AlMatrix, bOrientation, FaceList);
 			}
 		}
 		break;
diff --git a/Engine/Plugins/Enterprise/DatasmithCADImporter/Source/DatasmithWireTranslator/Private/AliasModelToCoretechConverter.h b/Engine/Plugins/Enterprise/DatasmithCADImporter/Source/DatasmithWireTranslator/Private/AliasModelToCoretechConverter.h
--- a/Engine/Plugins/Enterprise/DatasmithCADImporter/Source/DatasmithWireTranslator/Private/AliasModelToCoretechConverter.h
+++ b/Engine/Plugins/Enterprise/DatasmithCADImporter/Source/DatasmithWireTranslator/Private/AliasModelToCoretechConverter.h
@@ -39,12 +39,12 @@ protected:
 	 * Create a CT coedge (represent the use of an edge by a face).
 	 * @param TrimCurve: A curve in parametric surface space, part of a trim boundary.
 	 */
-	uint64 AddTrimCurve(AlTrimCurve& TrimCurve);
-	uint64 AddTrimBoundary(AlTrimBoundary& TrimBoundary);
+	uint64 AddTrimCurve(const AlTrimCurve& TrimCurve);
+	uint64 AddTrimBoundary(const AlTrimBoundary& TrimBoundary);
 
-	uint64 AddTrimRegion(AlTrimRegion& InTrimRegion, EAliasObjectReference InObjectReference, AlMatrix4x4& InAlMatrix, bool bInOrientation);
-	void AddFace(AlSurface& InSurface, EAliasObjectReference InObjectReference, AlMatrix4x4& InAlMatrix, bool bInOrientation, TArray<uint64>& OutFaceLis);
-	void AddShell(AlShell& InShell, EAliasObjectReference InObjectReference, AlMatrix4x4& InAlMatrix, bool bInOrientation, TArray<uint64>& OutFaceLis);
+	uint64 AddTrimRegion(const AlTrimRegion& InTrimRegion, EAliasObjectReference InObjectReference, const AlMatrix4x4& InAlMatrix, bool bInOrientation);
+	void AddFace(const AlSurface& InSurface, EAliasObjectReference InObjectReference, const AlMatrix4x4& InAlMatrix, bool bInOrientation, TArray<uint64>& OutFaceLis);
+	void AddShell(const AlShell& InShell, EAliasObjectReference InObjectReference, const AlMatrix4x4& InAlMatrix, bool bInOrientation, TArray<uint64>& OutFaceLis);
 
 protected:
 	TMap<void*, uint64>  AlEdge2CTEdge;
diff --git a/Engine/Plugins/Enterprise/DatasmithCADImporter/Source/DatasmithWireTranslator/Private/DatasmithWireTranslator.cpp b/Engine/Plugins/Enterprise/DatasmithCADImporter/Source/DatasmithWireTranslator/Private/DatasmithWireTranslator.cpp
--- a/Engine/Plugins/Enterprise/DatasmithCADImporter/Source/DatasmithWireTranslator/Private/DatasmithWireTranslator.cpp
+++ b/Engine/Plugins/Enterprise/DatasmithCADImporter/Source/DatasmithWireTranslator/Private/DatasmithWireTranslator.cpp
@@ -58,7 +58,7 @@ DEFINE_LOG_CATEGORY_STATIC(LogDatasmithWireTranslator, Log, All);
 
 #define LOCTEXT_NAMESPACE "DatasmithWireTranslator"
 
-#define WRONG_VERSION_TEXT "Unsupported version of Alias detected. Please downgrade to Alias 2020.0 (or earlier version) or upgrade to Alias 2021 (or later version)."
+#define WRONG_VERSION_TEXT "Unsupported version of Alias detected. Please upgrade to Alias 2021.3 (or later version)."
 #define CAD_INTERFACE_UNAVAILABLE "CAD Interface module is unavailable. Meshing will be done by Alias."
 
 #ifdef USE_OPENMODEL
@@ -68,14 +68,14 @@ using namespace CADLibrary;
 class BodyData
 {
 public:
-	BodyData(const char* InShaderName, const char* InLayerName, bool bInCadData)
+	BodyData(const FString& InShaderName, const FString& InLayerName, bool bInCadData)
 		: bCadData(bInCadData)
 	{
 		ShaderName = InShaderName;
 		LayerName = InLayerName;
 	};
 
-	TArray<AlDagNode*> ShellSet;
+	TArray<TSharedPtr<AlDagNode>> ShellSet;
 	FString ShaderName;
 	FString LayerName;
 	FString Label;
@@ -84,20 +84,20 @@ public:
 	// Generates BodyData's unique id from AlDagNode objects
 	uint32 GetUuid(const uint32& ParentUuid)
 	{
-		if(ShellSet.Num() == 0)
+		if (ShellSet.Num() == 0)
 		{
 			return ParentUuid;
 		}
 
 		if (ShellSet.Num() > 1)
 		{
-			ShellSet.Sort([&](AlDagNode& NodeA, AlDagNode& NodeB) { return OpenModelUtils::GetAlDagNodeUuid(NodeA) < OpenModelUtils::GetAlDagNodeUuid(NodeB); });
+			ShellSet.Sort([&](const TSharedPtr<AlDagNode>& NodeA, const TSharedPtr<AlDagNode>& NodeB) { return OpenModelUtils::GetAlDagNodeUuid(NodeA) < OpenModelUtils::GetAlDagNodeUuid(NodeB); });
 		}
 
 		uint32 BodyUUID = 0;
-		for (AlDagNode* DagNode : ShellSet)
+		for (const TSharedPtr<AlDagNode>& DagNode : ShellSet)
 		{
-			BodyUUID = HashCombine(BodyUUID, OpenModelUtils::GetAlDagNodeUuid(*DagNode));
+			BodyUUID = HashCombine(BodyUUID, OpenModelUtils::GetAlDagNodeUuid(DagNode));
 		}
 
 		return HashCombine(ParentUuid, BodyUUID);
@@ -137,9 +137,9 @@ public:
 		// application_name is something like "Catia V5"
 		DatasmithScene->SetHost(TEXT("Alias"));
 		DatasmithScene->SetVendor(TEXT("Autodesk"));
-		DatasmithScene->SetExporterSDKVersion(TEXT("2019"));
+		DatasmithScene->SetExporterSDKVersion(TEXT("2022"));
 		DatasmithScene->SetProductName(TEXT("Alias Tools"));
-		DatasmithScene->SetProductVersion(TEXT("Alias 2019"));
+		DatasmithScene->SetProductVersion(TEXT("Alias 2022"));
 
 		bool bEnableKernelIOTessellation = true;
 		if (IConsoleVariable* CVar = IConsoleManager::Get().FindConsoleVariable(TEXT("ds.CADTranslator.EnableKernelIOTessellation")))
@@ -155,9 +155,9 @@ public:
 		}
 		else
 		{
-			//TSharedRef<FAliasCADKernelWrapper> AliasCoretechWrapper = MakeShared<FAliasCADKernelWrapper>(TEXT("Al2CTSharedSession"));
-			//CADWrapper = AliasCoretechWrapper;
-			//AliasWrapper = AliasCoretechWrapper;
+			//TSharedRef<FAliasModelToCADKernelConverter> AliasToCoretechConverter = MakeShared<FAliasModelToCADKernelConverter>();
+			//CADModelConverter = AliasToCoretechConverter;
+			//AliasBRepConverter = AliasToCoretechConverter;
 		}
 	}
 
@@ -186,33 +186,33 @@ private:
 	{
 		uint32 Uuid;  // Use for actor name
 		FString Label;
-		TSharedPtr< IDatasmithActorElement > ActorElement;
+		TSharedPtr<IDatasmithActorElement> ActorElement;
 	};
 
 	bool GetDagLeaves();
 	bool GetShader();
-	bool RecurseDagForLeaves(AlDagNode* DagNode, const FDagNodeInfo& ParentInfo);
-	bool RecurseDagForLeavesNoMerge(AlDagNode* DagNode, const FDagNodeInfo& ParentInfo);
-	bool ProcessAlGroupNode(AlGroupNode& GroupNode, const FDagNodeInfo& ParentInfo);
-	bool ProcessAlShellNode(AlDagNode& GroupNode, const FDagNodeInfo& ParentInfo, const char* ShaderName);
-	bool ProcessBodyNode(TSharedRef<BodyData> Body, const FDagNodeInfo& ParentInfo);
-	TSharedPtr< IDatasmithMeshElement > FindOrAddMeshElement(TSharedRef<BodyData> Body, const FDagNodeInfo& ParentInfo);
-	TSharedPtr< IDatasmithMeshElement > FindOrAddMeshElement(AlDagNode& ShellNode, const FDagNodeInfo& ParentInfo, const char* ShaderName);
-	void GetDagNodeInfo(AlDagNode& GroupNode, const FDagNodeInfo& ParentInfo, FDagNodeInfo& CurrentNodeInfo);
-	void GetDagNodeInfo(TSharedRef<BodyData> CurrentNode, const FDagNodeInfo& ParentInfo, FDagNodeInfo& CurrentNodeInfo);
-	void GetDagNodeMeta(AlDagNode& CurrentNode, TSharedPtr< IDatasmithActorElement > ActorElement);
+	bool RecurseDagForLeaves(const TSharedPtr<AlDagNode>& DagNode, const FDagNodeInfo& ParentInfo);
+	bool RecurseDagForLeavesNoMerge(const TSharedPtr<AlDagNode>& DagNode, const FDagNodeInfo& ParentInfo);
+	bool ProcessAlGroupNode(AlDagNode& GroupNode, const FDagNodeInfo& ParentInfo);
+	bool ProcessAlShellNode(const TSharedPtr<AlDagNode>& ShellNode, const FDagNodeInfo& ParentInfo, const FString& ShaderName);
+	bool ProcessBodyNode(const TSharedPtr<BodyData>& Body, const FDagNodeInfo& ParentInfo);
+	TSharedPtr<IDatasmithMeshElement> FindOrAddMeshElement(const TSharedPtr<BodyData>& Body, const FDagNodeInfo& ParentInfo);
+	TSharedPtr<IDatasmithMeshElement> FindOrAddMeshElement(const TSharedPtr<AlDagNode>& ShellNode, const FDagNodeInfo& ParentInfo, const FString& ShaderName);
+	void GetDagNodeInfo(const TSharedPtr<AlDagNode>& GroupNode, const FDagNodeInfo& ParentInfo, FDagNodeInfo& CurrentNodeInfo);
+	void GetDagNodeInfo(const TSharedPtr<BodyData>& CurrentNode, const FDagNodeInfo& ParentInfo, FDagNodeInfo& CurrentNodeInfo);
+	void GetDagNodeMeta(const AlDagNode& InDagNode, TSharedPtr<IDatasmithActorElement> OutActorElement);
 
 	TOptional<FMeshDescription> GetMeshOfShellNode(AlDagNode& DagNode, TSharedRef<IDatasmithMeshElement> MeshElement, CADLibrary::FMeshParameters& MeshParameters);
-	TOptional<FMeshDescription> GetMeshOfNodeMesh(AlMeshNode& DagNode, TSharedRef<IDatasmithMeshElement> MeshElement, CADLibrary::FMeshParameters& MeshParameters, AlMatrix4x4* AlMeshInvGlobalMatrix = nullptr);
+	TOptional<FMeshDescription> GetMeshOfNodeMesh(AlDagNode& DagNode, TSharedRef<IDatasmithMeshElement> MeshElement, CADLibrary::FMeshParameters& MeshParameters, AlMatrix4x4* AlMeshInvGlobalMatrix = nullptr);
 	TOptional<FMeshDescription> GetMeshOfShellBody(TSharedRef<BodyData> DagNode, TSharedRef<IDatasmithMeshElement> MeshElement, CADLibrary::FMeshParameters& MeshParameters);
 	TOptional<FMeshDescription> GetMeshOfMeshBody(TSharedRef<BodyData> DagNode, TSharedRef<IDatasmithMeshElement> MeshElement, CADLibrary::FMeshParameters& MeshParameters);
 
-	void AddNodeInBodyGroup(AlDagNode& DagNode, const char* ShaderName, TMap<uint32, TSharedPtr<BodyData>>& ShellToProcess, bool bIsAPatch, uint32 MaxSize);
+	void AddNodeInBodyGroup(TSharedPtr<AlDagNode>& DagNode, const FString& ShaderName, TMap<uint32, TSharedPtr<BodyData>>& ShellToProcess, bool bIsAPatch, uint32 MaxSize);
 
 	TOptional<FMeshDescription> MeshDagNodeWithExternalMesher(AlDagNode& DagNode, TSharedRef<IDatasmithMeshElement> MeshElement, CADLibrary::FMeshParameters& MeshParameters);
 	TOptional<FMeshDescription> MeshDagNodeWithExternalMesher(TSharedRef<BodyData> DagNode, TSharedRef<IDatasmithMeshElement> MeshElement, CADLibrary::FMeshParameters& MeshParameters);
 
- 	TOptional< FMeshDescription > ImportMesh(AlMesh& Mesh, CADLibrary::FMeshParameters& MeshParameters);
+	TOptional<FMeshDescription> ImportMesh(AlMesh& Mesh, CADLibrary::FMeshParameters& MeshParameters);
 
 	FORCEINLINE bool IsTransparent(FColor& TransparencyColor)
 	{
@@ -254,15 +254,15 @@ private:
 		case AlShadingFields::kFLD_SHADING_COMMON_GLOW_INTENSITY:
 			GlowIntensity = Value;
 			return true;
-		default :
+		default:
 			return false;
 		}
 	}
 
-	void AddAlBlinnParameters(AlShader *Shader, TSharedRef<IDatasmithUEPbrMaterialElement> MaterialElement);
-	void AddAlLambertParameters(AlShader *Shader, TSharedRef<IDatasmithUEPbrMaterialElement> MaterialElement);
-	void AddAlLightSourceParameters(AlShader *Shader, TSharedRef<IDatasmithUEPbrMaterialElement> MaterialElement);
-	void AddAlPhongParameters(AlShader *Shader, TSharedRef<IDatasmithUEPbrMaterialElement> MaterialElement);
+	void AddAlBlinnParameters(const TUniquePtr<AlShader>& Shader, TSharedRef<IDatasmithUEPbrMaterialElement> MaterialElement);
+	void AddAlLambertParameters(const TUniquePtr<AlShader>& Shader, TSharedRef<IDatasmithUEPbrMaterialElement> MaterialElement);
+	void AddAlLightSourceParameters(const TUniquePtr<AlShader>& Shader, TSharedRef<IDatasmithUEPbrMaterialElement> MaterialElement);
+	void AddAlPhongParameters(const TUniquePtr<AlShader>& Shader, TSharedRef<IDatasmithUEPbrMaterialElement> MaterialElement);
 
 private:
 	TSharedRef<IDatasmithScene> DatasmithScene;
@@ -276,14 +276,14 @@ private:
 	// Hash value of the scene file used to check if the file has been modified for re-import
 	uint32 SceneFileHash;
 
-	AlDagNode* AlRootNode;
+	TSharedPtr<AlDagNode> AlRootNode;
 
 	/** Table of correspondence between mesh identifier and associated Datasmith mesh element */
 	TMap<uint32, TSharedPtr<IDatasmithMeshElement>> ShellUuidToMeshElementMap;
 	TMap<uint32, TSharedPtr<IDatasmithMeshElement>> BodyUuidToMeshElementMap;
 
 	/** Datasmith mesh elements to OpenModel objects */
-	TMap<IDatasmithMeshElement*, AlDagNode*> MeshElementToAlDagNodeMap;
+	TMap<IDatasmithMeshElement*, TSharedPtr<AlDagNode>> MeshElementToAlDagNodeMap;
 
 	TMap<IDatasmithMeshElement*, TSharedPtr<BodyData>> MeshElementToBodyMap;
 
@@ -321,7 +321,7 @@ bool FWireTranslatorImpl::Read()
 		return false;
 	}
 
-	CADModelConverter->SetImportParameters(TessellationOptions.ChordTolerance, TessellationOptions.MaxEdgeLength, TessellationOptions.NormalTolerance, (CADLibrary::EStitchingTechnique) TessellationOptions.StitchingTechnique, true);
+	CADModelConverter->SetImportParameters(TessellationOptions.ChordTolerance, TessellationOptions.MaxEdgeLength, TessellationOptions.NormalTolerance, (CADLibrary::EStitchingTechnique)TessellationOptions.StitchingTechnique, true);
 
 	AlRetrieveOptions options;
 	AlUniverse::retrieveOptions(options);
@@ -342,7 +342,7 @@ bool FWireTranslatorImpl::Read()
 	return true;
 }
 
-void FWireTranslatorImpl::AddAlBlinnParameters(AlShader *Shader, TSharedRef<IDatasmithUEPbrMaterialElement> MaterialElement)
+void FWireTranslatorImpl::AddAlBlinnParameters(const TUniquePtr<AlShader>& Shader, TSharedRef<IDatasmithUEPbrMaterialElement> MaterialElement)
 {
 	// Default values for a Blinn material
 	FColor Color(145, 148, 153);
@@ -358,7 +358,7 @@ void FWireTranslatorImpl::AddAlBlinnParameters(AlShader *Shader, TSharedRef<IDat
 	double SpecularRolloff = 0.5;
 
 	AlList* List = Shader->fields();
-	for (AlShadingFieldItem* Item = static_cast<AlShadingFieldItem *>(List->first()); Item; Item = Item->nextField())
+	for (AlShadingFieldItem* Item = static_cast<AlShadingFieldItem*>(List->first()); Item; Item = Item->nextField())
 	{
 		double Value = 0.0f;
 		statusCode ErrorCode = Shader->parameter(Item->field(), Value);
@@ -381,7 +381,7 @@ void FWireTranslatorImpl::AddAlBlinnParameters(AlShader *Shader, TSharedRef<IDat
 			Gloss = Value;
 			break;
 		case AlShadingFields::kFLD_SHADING_BLINN_SPECULAR_R:
-			SpecularColor.R = (uint8) (255.f * Value);
+			SpecularColor.R = (uint8)(255.f * Value);
 			break;
 		case AlShadingFields::kFLD_SHADING_BLINN_SPECULAR_G:
 			SpecularColor.G = (uint8)(255.f * Value);;
@@ -619,14 +619,14 @@ void FWireTranslatorImpl::AddAlBlinnParameters(AlShader *Shader, TSharedRef<IDat
 		MaterialElement->GetOpacity().SetExpression(Divide);
 		MaterialElement->SetParentLabel(TEXT("M_DatasmithAliasBlinnTransparent"));
 	}
-	else 
+	else
 	{
 		MaterialElement->SetParentLabel(TEXT("M_DatasmithAliasBlinn"));
 	}
 
 }
 
-void FWireTranslatorImpl::AddAlLambertParameters(AlShader *Shader, TSharedRef<IDatasmithUEPbrMaterialElement> MaterialElement)
+void FWireTranslatorImpl::AddAlLambertParameters(const TUniquePtr<AlShader>& Shader, TSharedRef<IDatasmithUEPbrMaterialElement> MaterialElement)
 {
 	// Default values for a Lambert material
 	FColor Color(145, 148, 153);
@@ -636,7 +636,7 @@ void FWireTranslatorImpl::AddAlLambertParameters(AlShader *Shader, TSharedRef<ID
 	double GlowIntensity = 0.0;
 
 	AlList* List = Shader->fields();
-	for (AlShadingFieldItem* Item = static_cast<AlShadingFieldItem *>(List->first()); Item; Item = Item->nextField())
+	for (AlShadingFieldItem* Item = static_cast<AlShadingFieldItem*>(List->first()); Item; Item = Item->nextField())
 	{
 		double Value = 0.0f;
 		statusCode ErrorCode = Shader->parameter(Item->field(), Value);
@@ -785,7 +785,7 @@ void FWireTranslatorImpl::AddAlLambertParameters(AlShader *Shader, TSharedRef<ID
 
 }
 
-void FWireTranslatorImpl::AddAlLightSourceParameters(AlShader *Shader, TSharedRef<IDatasmithUEPbrMaterialElement> MaterialElement)
+void FWireTranslatorImpl::AddAlLightSourceParameters(const TUniquePtr<AlShader>& Shader, TSharedRef<IDatasmithUEPbrMaterialElement> MaterialElement)
 {
 	// Default values for a LightSource material
 	FColor Color(145, 148, 153);
@@ -794,7 +794,7 @@ void FWireTranslatorImpl::AddAlLightSourceParameters(AlShader *Shader, TSharedRe
 	double GlowIntensity = 0.0;
 
 	AlList* List = Shader->fields();
-	for (AlShadingFieldItem* Item = static_cast<AlShadingFieldItem *>(List->first()); Item; Item = Item->nextField())
+	for (AlShadingFieldItem* Item = static_cast<AlShadingFieldItem*>(List->first()); Item; Item = Item->nextField())
 	{
 		double Value = 0.0f;
 		statusCode ErrorCode = Shader->parameter(Item->field(), Value);
@@ -911,7 +911,7 @@ void FWireTranslatorImpl::AddAlLightSourceParameters(AlShader *Shader, TSharedRe
 
 }
 
-void FWireTranslatorImpl::AddAlPhongParameters(AlShader *Shader, TSharedRef<IDatasmithUEPbrMaterialElement> MaterialElement)
+void FWireTranslatorImpl::AddAlPhongParameters(const TUniquePtr<AlShader>& Shader, TSharedRef<IDatasmithUEPbrMaterialElement> MaterialElement)
 {
 	// Default values for a Phong material
 	FColor Color(145, 148, 153);
@@ -926,7 +926,7 @@ void FWireTranslatorImpl::AddAlPhongParameters(AlShader *Shader, TSharedRef<IDat
 	double Reflectivity = 0.5;
 
 	AlList* List = Shader->fields();
-	for (AlShadingFieldItem* Item = static_cast<AlShadingFieldItem *>(List->first()); Item; Item = Item->nextField())
+	for (AlShadingFieldItem* Item = static_cast<AlShadingFieldItem*>(List->first()); Item; Item = Item->nextField())
 	{
 		double Value = 0.0f;
 		statusCode ErrorCode = Shader->parameter(Item->field(), Value);
@@ -1167,10 +1167,9 @@ void FWireTranslatorImpl::AddAlPhongParameters(AlShader *Shader, TSharedRef<IDat
 // Make material
 bool FWireTranslatorImpl::GetShader()
 {
-	AlShader *Shader = AlUniverse::firstShader();
-	while (Shader)
+	for (TUniquePtr<AlShader> Shader(AlUniverse::firstShader()); Shader.IsValid(); Shader = TUniquePtr<AlShader>(AlUniverse::nextShader(Shader.Get())))
 	{
-		FString ShaderName = Shader->name();
+		FString ShaderName = UTF8_TO_TCHAR(Shader->name());
 		FString ShaderModelName = Shader->shadingModel();
 
 		uint32 ShaderUUID = fabs((float)(int32)GetTypeHash(*ShaderName));
@@ -1199,12 +1198,8 @@ bool FWireTranslatorImpl::GetShader()
 
 		DatasmithScene->AddMaterial(MaterialElement);
 
-		TSharedPtr< IDatasmithMaterialIDElement > MaterialIDElement = FDatasmithSceneFactory::CreateMaterialId(MaterialElement->GetName());
+		TSharedPtr<IDatasmithMaterialIDElement> MaterialIDElement = FDatasmithSceneFactory::CreateMaterialId(MaterialElement->GetName());
 		ShaderNameToUEMaterialId.Add(*ShaderName, MaterialIDElement);
-
-		AlShader *curShader = Shader;
-		Shader = AlUniverse::nextShader(Shader);
-		delete curShader;
 	}
 	return true;
 }
@@ -1214,53 +1209,65 @@ bool FWireTranslatorImpl::GetShader()
 bool FWireTranslatorImpl::GetDagLeaves()
 {
 	FDagNodeInfo RootContainer;
-	AlRootNode = AlUniverse::firstDagNode();
+	AlRootNode = TSharedPtr<AlDagNode>(AlUniverse::firstDagNode());
+	if (!AlRootNode.IsValid())
+	{
+		return false;
+	}
 	return RecurseDagForLeaves(AlRootNode, RootContainer);
 }
 
 
-void FWireTranslatorImpl::GetDagNodeMeta(AlDagNode& CurrentNode, TSharedPtr< IDatasmithActorElement > ActorElement)
+void FWireTranslatorImpl::GetDagNodeMeta(const AlDagNode& InDagNode, TSharedPtr<IDatasmithActorElement> OutActorElement)
 {
-	AlLayer* layer = CurrentNode.layer();
-	if (layer != nullptr)
+	AlLayer* LayerPtr = InDagNode.layer();
+	if (LayerPtr)
 	{
-		FString LayerName = CurrentNode.layer()->name();
-		ActorElement->SetLayer(*LayerName);
+		TUniquePtr<AlLayer> Layer(LayerPtr);
+		FString LayerName = UTF8_TO_TCHAR(Layer->name());
+		OutActorElement->SetLayer(*LayerName);
 	}
-
-    // TODO import other Meta
-
 }
 
-void FWireTranslatorImpl::GetDagNodeInfo(AlDagNode& CurrentNode, const FDagNodeInfo& ParentInfo, FDagNodeInfo& CurrentNodeInfo)
+void FWireTranslatorImpl::GetDagNodeInfo(const TSharedPtr<AlDagNode>& CurrentNode, const FDagNodeInfo& ParentInfo, FDagNodeInfo& CurrentNodeInfo)
 {
-    CurrentNodeInfo.Label = CurrentNode.name();
+
+	CurrentNodeInfo.Label = UTF8_TO_TCHAR(CurrentNode->name());
 	uint32 ThisGroupNodeUuid = OpenModelUtils::GetAlDagNodeUuid(CurrentNode);
 	CurrentNodeInfo.Uuid = HashCombine(ParentInfo.Uuid, ThisGroupNodeUuid);
 }
 
-void FWireTranslatorImpl::GetDagNodeInfo(TSharedRef<BodyData> CurrentNode, const FDagNodeInfo& ParentInfo, FDagNodeInfo& CurrentNodeInfo)
+void FWireTranslatorImpl::GetDagNodeInfo(const TSharedPtr<BodyData>& CurrentNode, const FDagNodeInfo& ParentInfo, FDagNodeInfo& CurrentNodeInfo)
 {
 	CurrentNodeInfo.Label = ParentInfo.Label;
 	CurrentNode->Label = CurrentNodeInfo.Label;
 	CurrentNodeInfo.Uuid = CurrentNode->GetUuid(ParentInfo.Uuid);
 }
 
-bool FWireTranslatorImpl::ProcessAlGroupNode(AlGroupNode& GroupNode, const FDagNodeInfo& ParentInfo)
+bool FWireTranslatorImpl::ProcessAlGroupNode(AlDagNode& GroupNode, const FDagNodeInfo& ParentInfo)
 {
-	FDagNodeInfo ThisGroupNodeInfo;
+	AlGroupNode* AlGroup = GroupNode.asGroupNodePtr();
+	if (!AlIsValid(AlGroup))
+	{
+		return false;
+	}
+
+	AlDagNode* AlChildPtr = AlGroup->childNode();
+	if (!AlIsValid(AlChildPtr))
+	{
+		return false;
+	}
 
-	GetDagNodeInfo(GroupNode, ParentInfo, ThisGroupNodeInfo);
+	TSharedPtr<AlDagNode> ChildNode(AlChildPtr);
+
+	FDagNodeInfo ThisGroupNodeInfo;
+	GetDagNodeInfo(ChildNode, ParentInfo, ThisGroupNodeInfo);
 
 	ThisGroupNodeInfo.ActorElement = FDatasmithSceneFactory::CreateActor(*OpenModelUtils::UuidToString(ThisGroupNodeInfo.Uuid));
 	ThisGroupNodeInfo.ActorElement->SetLabel(*ThisGroupNodeInfo.Label);
 	GetDagNodeMeta(GroupNode, ThisGroupNodeInfo.ActorElement);
 
-	AlDagNode * ChildNode = GroupNode.childNode();
-	if (AlIsValid(ChildNode) == TRUE)
-	{
-		RecurseDagForLeaves(ChildNode, ThisGroupNodeInfo);
-	}
+	RecurseDagForLeaves(ChildNode, ThisGroupNodeInfo);
 
 	// add the resulting actor to the scene
 	if (OpenModelUtils::IsValidActor(ThisGroupNodeInfo.ActorElement))
@@ -1281,21 +1288,21 @@ bool FWireTranslatorImpl::ProcessAlGroupNode(AlGroupNode& GroupNode, const FDagN
 	return true;
 }
 
-TSharedPtr< IDatasmithMeshElement > FWireTranslatorImpl::FindOrAddMeshElement(TSharedRef<BodyData> Body, const FDagNodeInfo& NodeInfo)
+TSharedPtr<IDatasmithMeshElement> FWireTranslatorImpl::FindOrAddMeshElement(const TSharedPtr<BodyData>& Body, const FDagNodeInfo& NodeInfo)
 {
-	TSharedPtr< IDatasmithMeshElement >* MeshElementPtr = BodyUuidToMeshElementMap.Find( NodeInfo.Uuid );
+	TSharedPtr<IDatasmithMeshElement>* MeshElementPtr = BodyUuidToMeshElementMap.Find(NodeInfo.Uuid);
 	if (MeshElementPtr != nullptr)
 	{
 		return *MeshElementPtr;
 	}
-	
-	TSharedPtr< IDatasmithMeshElement > MeshElement = FDatasmithSceneFactory::CreateMesh(*OpenModelUtils::UuidToString(NodeInfo.Uuid));
+
+	TSharedPtr<IDatasmithMeshElement> MeshElement = FDatasmithSceneFactory::CreateMesh(*OpenModelUtils::UuidToString(NodeInfo.Uuid));
 	MeshElement->SetLabel(*NodeInfo.Label);
 	MeshElement->SetLightmapSourceUV(-1);
 
 	if (*Body->ShaderName)
 	{
-		TSharedPtr< IDatasmithMaterialIDElement > MaterialElement = ShaderNameToUEMaterialId[Body->ShaderName];
+		TSharedPtr<IDatasmithMaterialIDElement> MaterialElement = ShaderNameToUEMaterialId[Body->ShaderName];
 		MeshElement->SetMaterial(MaterialElement->GetName(), 0);
 	}
 
@@ -1304,27 +1311,24 @@ TSharedPtr< IDatasmithMeshElement > FWireTranslatorImpl::FindOrAddMeshElement(TS
 	ShellUuidToMeshElementMap.Add(NodeInfo.Uuid, MeshElement);
 	MeshElementToBodyMap.Add(MeshElement.Get(), Body);
 
-	BodyUuidToMeshElementMap.Add( NodeInfo.Uuid, MeshElement );
+	BodyUuidToMeshElementMap.Add(NodeInfo.Uuid, MeshElement);
 
 	return MeshElement;
 }
 
-TSharedPtr< IDatasmithMeshElement > FWireTranslatorImpl::FindOrAddMeshElement(AlDagNode& ShellNode, const FDagNodeInfo& ParentInfo, const char* ShaderName)
+TSharedPtr<IDatasmithMeshElement> FWireTranslatorImpl::FindOrAddMeshElement(const TSharedPtr<AlDagNode>& ShellNode, const FDagNodeInfo& ShellNodeInfo, const FString& ShaderName)
 {
 	uint32 ShellUuid = OpenModelUtils::GetAlDagNodeUuid(ShellNode);
 
 	// Look if geometry has not been already processed, return it if found
-	TSharedPtr< IDatasmithMeshElement >* MeshElementPtr = ShellUuidToMeshElementMap.Find(ShellUuid);
+	TSharedPtr<IDatasmithMeshElement>* MeshElementPtr = ShellUuidToMeshElementMap.Find(ShellUuid);
 	if (MeshElementPtr != nullptr)
 	{
 		return *MeshElementPtr;
 	}
 
-	FDagNodeInfo MeshInfo;
-	GetDagNodeInfo(ShellNode, ParentInfo, MeshInfo);
-
-	TSharedPtr< IDatasmithMeshElement > MeshElement = FDatasmithSceneFactory::CreateMesh(*OpenModelUtils::UuidToString(MeshInfo.Uuid));
-	MeshElement->SetLabel(*MeshInfo.Label);
+	TSharedPtr<IDatasmithMeshElement> MeshElement = FDatasmithSceneFactory::CreateMesh(*OpenModelUtils::UuidToString(ShellNodeInfo.Uuid));
+	MeshElement->SetLabel(*ShellNodeInfo.Label);
 	MeshElement->SetLightmapSourceUV(-1);
 
 	// Set MeshElement FileHash used for re-import task 
@@ -1336,32 +1340,32 @@ TSharedPtr< IDatasmithMeshElement > FWireTranslatorImpl::FindOrAddMeshElement(Al
 	Hash.Set(MD5);
 	MeshElement->SetFileHash(Hash);
 
-	if (ShaderName)
+	if (ShaderName.Len())
 	{
-		TSharedPtr< IDatasmithMaterialIDElement > MaterialElement = ShaderNameToUEMaterialId[FString(ShaderName)];
+		TSharedPtr<IDatasmithMaterialIDElement> MaterialElement = ShaderNameToUEMaterialId[ShaderName];
 		MeshElement->SetMaterial(MaterialElement->GetName(), 0);
 	}
 
 	DatasmithScene->AddMesh(MeshElement);
 
 	ShellUuidToMeshElementMap.Add(ShellUuid, MeshElement);
-	MeshElementToAlDagNodeMap.Add(MeshElement.Get(), &ShellNode);
+	MeshElementToAlDagNodeMap.Add(MeshElement.Get(), ShellNode);
 
 	return MeshElement;
 }
 
-bool FWireTranslatorImpl::ProcessAlShellNode(AlDagNode& ShellNode, const FDagNodeInfo& ParentInfo, const char* ShaderName)
+bool FWireTranslatorImpl::ProcessAlShellNode(const TSharedPtr<AlDagNode>& ShellNode, const FDagNodeInfo& ParentInfo, const FString& ShaderName)
 {
 	FDagNodeInfo ShellInfo;
 	GetDagNodeInfo(ShellNode, ParentInfo, ShellInfo);
 
-	TSharedPtr< IDatasmithMeshElement > MeshElement = FindOrAddMeshElement(ShellNode, ShellInfo, ShaderName);
+	TSharedPtr<IDatasmithMeshElement> MeshElement = FindOrAddMeshElement(ShellNode, ShellInfo, ShaderName);
 	if (!MeshElement.IsValid())
 	{
 		return false;
 	}
 
-	TSharedPtr< IDatasmithMeshActorElement > ActorElement = FDatasmithSceneFactory::CreateMeshActor(*OpenModelUtils::UuidToString(ShellInfo.Uuid));
+	TSharedPtr<IDatasmithMeshActorElement> ActorElement = FDatasmithSceneFactory::CreateMeshActor(*OpenModelUtils::UuidToString(ShellInfo.Uuid));
 	if (!ActorElement.IsValid())
 	{
 		return false;
@@ -1371,15 +1375,16 @@ bool FWireTranslatorImpl::ProcessAlShellNode(AlDagNode& ShellNode, const FDagNod
 	ActorElement->SetStaticMeshPathName(MeshElement->GetName());
 	ShellInfo.ActorElement = ActorElement;
 
-	GetDagNodeMeta(ShellNode, ActorElement);
+	GetDagNodeMeta(*ShellNode, ActorElement);
 
-	OpenModelUtils::SetActorTransform(ShellInfo.ActorElement, ShellNode);
+	OpenModelUtils::SetActorTransform(ShellInfo.ActorElement, *ShellNode);
 
 	// Apply materials on the current part
-	if (ShaderName)
+	if (ShaderName.Len())
 	{
-		TSharedPtr< IDatasmithMaterialIDElement > MaterialIDElement = ShaderNameToUEMaterialId[FString(ShaderName)];
-		if (MaterialIDElement.IsValid()) {
+		const TSharedPtr<IDatasmithMaterialIDElement>& MaterialIDElement = ShaderNameToUEMaterialId[ShaderName];
+		if (MaterialIDElement.IsValid())
+		{
 			for (int32 Index = 0; Index < MeshElement->GetMaterialSlotCount(); ++Index)
 			{
 				MaterialIDElement->SetId(MeshElement->GetMaterialSlotAt(Index)->GetId());
@@ -1402,8 +1407,13 @@ bool FWireTranslatorImpl::ProcessAlShellNode(AlDagNode& ShellNode, const FDagNod
 	return true;
 }
 
-bool FWireTranslatorImpl::ProcessBodyNode(TSharedRef<BodyData> Body, const FDagNodeInfo& ParentInfo)
+bool FWireTranslatorImpl::ProcessBodyNode(const TSharedPtr<BodyData>& Body, const FDagNodeInfo& ParentInfo)
 {
+	if (!Body.IsValid())
+	{
+		return false;
+	}
+
 	if (Body->ShellSet.Num() == 1)
 	{
 		return RecurseDagForLeavesNoMerge(Body->ShellSet[0], ParentInfo);
@@ -1412,13 +1422,13 @@ bool FWireTranslatorImpl::ProcessBodyNode(TSharedRef<BodyData> Body, const FDagN
 	FDagNodeInfo ShellInfo;
 	GetDagNodeInfo(Body, ParentInfo, ShellInfo);
 
-	TSharedPtr< IDatasmithMeshElement > MeshElement = FindOrAddMeshElement(Body, ShellInfo);
+	TSharedPtr<IDatasmithMeshElement> MeshElement = FindOrAddMeshElement(Body, ShellInfo);
 	if (!MeshElement.IsValid())
 	{
 		return false;
 	}
 
-	TSharedPtr< IDatasmithMeshActorElement > ActorElement = FDatasmithSceneFactory::CreateMeshActor(*OpenModelUtils::UuidToString(ShellInfo.Uuid));
+	TSharedPtr<IDatasmithMeshActorElement> ActorElement = FDatasmithSceneFactory::CreateMeshActor(*OpenModelUtils::UuidToString(ShellInfo.Uuid));
 	if (!ActorElement.IsValid())
 	{
 		return false;
@@ -1435,7 +1445,7 @@ bool FWireTranslatorImpl::ProcessBodyNode(TSharedRef<BodyData> Body, const FDagN
 	//// Apply materials on the current part
 	if (*Body->ShaderName)
 	{
-		TSharedPtr< IDatasmithMaterialIDElement > MaterialIDElement = ShaderNameToUEMaterialId[FString(*Body->ShaderName)];
+		TSharedPtr<IDatasmithMaterialIDElement> MaterialIDElement = ShaderNameToUEMaterialId[FString(*Body->ShaderName)];
 		if (MaterialIDElement.IsValid()) {
 			for (int32 Index = 0; Index < MeshElement->GetMaterialSlotCount(); ++Index)
 			{
@@ -1459,21 +1469,21 @@ bool FWireTranslatorImpl::ProcessBodyNode(TSharedRef<BodyData> Body, const FDagN
 	return true;
 }
 
-AlDagNode * GetNextNode(AlDagNode * DagNode)
+TSharedPtr<AlDagNode> GetNextNode(const TSharedPtr<AlDagNode>& DagNode)
 {
 	// Grab the next sibling before deleting the node.
-	AlDagNode * SiblingNode = DagNode->nextNode();
+	AlDagNode* SiblingNode = DagNode->nextNode();
 	if (AlIsValid(SiblingNode))
 	{
-		return SiblingNode;
+		return TSharedPtr<AlDagNode>(SiblingNode);
 	}
 	else
 	{
-		return NULL;
+		return TSharedPtr<AlDagNode>();
 	}
 }
 
-bool IsHidden(AlDagNode * DagNode)
+bool IsHidden(const TSharedPtr<AlDagNode>& DagNode)
 {
 	/*
 	AlObjectType objectType = DagNode->type();
@@ -1504,31 +1514,34 @@ bool IsHidden(AlDagNode * DagNode)
 }
 
 
-uint32 GetBodyGroupUuid(const char* ShaderName, const char* LayerName, bool bCadData)
+uint32 GetBodyGroupUuid(const FString& ShaderName, const FString& LayerName, bool bCadData)
 {
 	uint32 Uuid = HashCombine(GetTypeHash(ShaderName), GetTypeHash(bCadData));
 	Uuid = HashCombine(GetTypeHash(LayerName), Uuid);
 	return Uuid;
 }
 
-uint32 GetPatchCount(AlShell& Shell)
+uint32 GetPatchCount(const TUniquePtr<AlShell>& Shell)
 {
 	uint32 PatchCount = 0;
-	AlTrimRegion *TrimRegion = Shell.firstTrimRegion();
-	while (TrimRegion)
+	TUniquePtr<AlTrimRegion> TrimRegion(Shell->firstTrimRegion());
+	while (TrimRegion.IsValid())
 	{
 		PatchCount++;
-		TrimRegion = TrimRegion->nextRegion();
+		TrimRegion = TUniquePtr<AlTrimRegion>(TrimRegion->nextRegion());
 	}
 	return PatchCount;
 }
 
-void FWireTranslatorImpl::AddNodeInBodyGroup(AlDagNode& DagNode, const char* ShaderName, TMap<uint32, TSharedPtr<BodyData>>& ShellToProcess, bool bIsAPatch, uint32 MaxSize)
+void FWireTranslatorImpl::AddNodeInBodyGroup(TSharedPtr<AlDagNode>& DagNode, const FString& ShaderName, TMap<uint32, TSharedPtr<BodyData>>& ShellToProcess, bool bIsAPatch, uint32 MaxSize)
 {
-	const char* LayerName = nullptr;
-	if (AlLayer* Layer = DagNode.layer())
+	FString LayerName;
+
+	AlLayer* LayerPtr = DagNode->layer();
+	if (AlIsValid(LayerPtr))
 	{
-		LayerName = Layer->name();
+		const TUniquePtr<AlLayer> Layer(LayerPtr);
+		LayerName = UTF8_TO_TCHAR(Layer->name());
 	}
 
 	uint32 SetId = GetBodyGroupUuid(ShaderName, LayerName, bIsAPatch);
@@ -1545,17 +1558,17 @@ void FWireTranslatorImpl::AddNodeInBodyGroup(AlDagNode& DagNode, const char* Sha
 		BodyRef->ShellSet.Reserve(MaxSize);
 		Body = BodyRef;
 	}
-	Body->ShellSet.Add(&DagNode);
+	Body->ShellSet.Add(DagNode);
 }
 
-bool FWireTranslatorImpl::RecurseDagForLeaves(AlDagNode* FirstDagNode, const FDagNodeInfo& ParentInfo)
+bool FWireTranslatorImpl::RecurseDagForLeaves(const TSharedPtr<AlDagNode>& FirstDagNode, const FDagNodeInfo& ParentInfo)
 {
 	if (TessellationOptions.StitchingTechnique != EDatasmithCADStitchingTechnique::StitchingSew)
 	{
 		return RecurseDagForLeavesNoMerge(FirstDagNode, ParentInfo);
 	}
 
-	AlDagNode* DagNode = FirstDagNode;
+	TSharedPtr<AlDagNode> DagNode = FirstDagNode;
 	uint32 MaxSize = 0;
 	while (DagNode)
 	{
@@ -1565,104 +1578,114 @@ bool FWireTranslatorImpl::RecurseDagForLeaves(AlDagNode* FirstDagNode, const FDa
 
 	DagNode = FirstDagNode;
 
-	TSet<AlDagNode*> UnProcessedNodes;
 	TMap<uint32, TSharedPtr<BodyData>> ShellToProcess;
 
 	const char* ShaderName = nullptr;
 
 	while (DagNode)
 	{
-		// Filter invalid nodes.
-		if (!AlIsValid(DagNode))
-		{
-			break;
-		}
-		if (IsHidden(DagNode))
+		// Filter visible nodes.
+		if (!IsHidden(DagNode))
 		{
-			return true;
-		}
+			AlObjectType objectType = DagNode->type();
 
-		AlObjectType objectType = DagNode->type();
+			// Process the current node.
+			switch (objectType)
+			{
+				// Push all leaf nodes into 'leaves'
+			case kShellNodeType:
+			{
+				AlShellNode* ShellNode = DagNode->asShellNodePtr();
+				if (!AlIsValid(ShellNode))
+				{
+					break;
+				}
 
-		// Process the current node.
-		switch (objectType)
-		{
-			// Push all leaf nodes into 'leaves'
-		case kShellNodeType:
-		{
-			AlShellNode *ShellNode = DagNode->asShellNodePtr();
-			AlShell *Shell = ShellNode->shell();
-			uint32 NbPatch = GetPatchCount(*Shell);
+				AlShell* ShellPtr = ShellNode->shell();
+				if (!AlIsValid(ShellPtr))
+				{
+					break;
+				}
 
-			if (AlShader* Shader = Shell->firstShader())
-			{
-				ShaderName = Shader->name();
-			}
+				TUniquePtr<AlShell> Shell(ShellPtr);
+				uint32 NbPatch = GetPatchCount(Shell);
 
-			if (NbPatch == 1)
-			{
-				AddNodeInBodyGroup(*DagNode, ShaderName, ShellToProcess, true, MaxSize);
+				TUniquePtr<AlShader> Shader(Shell->firstShader());
+				if (Shader.IsValid())
+				{
+					ShaderName = Shader->name();
+				}
+
+				if (NbPatch == 1)
+				{
+					AddNodeInBodyGroup(DagNode, ShaderName, ShellToProcess, true, MaxSize);
+				}
+				else
+				{
+					ProcessAlShellNode(DagNode, ParentInfo, ShaderName);
+				}
+				break;
 			}
-			else
+			case kSurfaceNodeType:
 			{
-				ProcessAlShellNode(*DagNode, ParentInfo, ShaderName);
+				AlSurfaceNode* SurfaceNode = DagNode->asSurfaceNodePtr();
+				TUniquePtr<AlSurface> Surface(SurfaceNode->surface());
+				if (Surface.IsValid())
+				{
+					TUniquePtr<AlShader> Shader(Surface->firstShader());
+					if (Shader.IsValid())
+					{
+						ShaderName = Shader->name();
+					}
+				}
+				AddNodeInBodyGroup(DagNode, ShaderName, ShellToProcess, true, MaxSize);
+				break;
 			}
-			break;
-		}
-		case kSurfaceNodeType:
-		{
-			AlSurfaceNode *SurfaceNode = DagNode->asSurfaceNodePtr();
-			AlSurface *Surface = SurfaceNode->surface();
-			if (AlShader * Shader = Surface->firstShader())
+
+			case kMeshNodeType:
 			{
-				ShaderName = Shader->name();
+				AlMeshNode* MeshNode = DagNode->asMeshNodePtr();
+				TUniquePtr<AlMesh> Mesh(MeshNode->mesh());
+				if (Mesh.IsValid())
+				{
+					TUniquePtr<AlShader> Shader(Mesh->firstShader());
+					if (Shader.IsValid())
+					{
+						ShaderName = Shader->name();
+					}
+				}
+				AddNodeInBodyGroup(DagNode, ShaderName, ShellToProcess, false, MaxSize);
+				break;
 			}
-			AddNodeInBodyGroup(*DagNode, ShaderName, ShellToProcess, true, MaxSize);
-			break;
-		}
 
-		case kMeshNodeType:
-		{
-			AlMeshNode *MeshNode = DagNode->asMeshNodePtr();
-			AlMesh *Mesh = MeshNode->mesh();
-			if (AlShader * Shader = Mesh->firstShader())
+			// Traverse down through groups
+			case kGroupNodeType:
 			{
-				ShaderName = Shader->name();
+				ProcessAlGroupNode(*DagNode, ParentInfo);
+				break;
 			}
-			AddNodeInBodyGroup(*DagNode, ShaderName, ShellToProcess, false, MaxSize);
-			break;
-		}
-
-		// Traverse down through groups
-		case kGroupNodeType:
-		{
-			AlGroupNode * GroupNode = DagNode->asGroupNodePtr();
-			if (AlIsValid(GroupNode))
+			default:
 			{
-				ProcessAlGroupNode(*GroupNode, ParentInfo);
+				break;
 			}
-			break;
-		}
-		default:
-		{
-		}
 
+			}
 		}
+
 		DagNode = GetNextNode(DagNode);
 	}
 
-	for (auto Body : ShellToProcess)
+	for (const TPair<uint32, TSharedPtr<BodyData>>& Body : ShellToProcess)
 	{
 		ProcessBodyNode(Body.Value.ToSharedRef(), ParentInfo);
 	}
 	return true;
 }
 
-bool FWireTranslatorImpl::RecurseDagForLeavesNoMerge(AlDagNode* FirstDagNode, const FDagNodeInfo& ParentInfo)
+bool FWireTranslatorImpl::RecurseDagForLeavesNoMerge(const TSharedPtr<AlDagNode>& FirstDagNode, const FDagNodeInfo& ParentInfo)
 {
-	AlDagNode* DagNode = FirstDagNode;
+	TSharedPtr<AlDagNode> DagNode = FirstDagNode;
 
-	TSet<AlDagNode*> UnProcessedNodes;
 	TMap<uint32, BodyData> ShellToProcess;
 
 	const char* ShaderName = nullptr;
@@ -1670,78 +1693,77 @@ bool FWireTranslatorImpl::RecurseDagForLeavesNoMerge(AlDagNode* FirstDagNode, co
 	DagNode = FirstDagNode;
 	while (DagNode)
 	{
-		// Filter invalid nodes.
-		if (AlIsValid(DagNode) == FALSE)
-		{
-			break;
-		}
-
-		if (IsHidden(DagNode))
+		if (!IsHidden(DagNode))
 		{
-			return true;
-		}
+			// Process the current node.
+			AlObjectType objectType = DagNode->type();
+			switch (objectType)
+			{
+				// Push all leaf nodes into 'leaves'
+			case kShellNodeType:
+			{
+				AlShellNode* ShellNode = DagNode->asShellNodePtr();
+				TUniquePtr<AlShell> Shell(ShellNode->shell());
+				if (Shell.IsValid())
+				{
+					TUniquePtr<AlShader> Shader(Shell->firstShader());
+					if (Shader.IsValid())
+					{
+						ShaderName = Shader->name();
+					}
 
-		// Process the current node.
-		AlObjectType objectType = DagNode->type();
-		switch (objectType)
-		{
-			// Push all leaf nodes into 'leaves'
-		case kShellNodeType:
-		{
-			AlShellNode *ShellNode = DagNode->asShellNodePtr();
-			AlShell *Shell = ShellNode->shell();
-			if (AlShader * Shader = Shell->firstShader())
+					ProcessAlShellNode(DagNode, ParentInfo, ShaderName);
+				}
+				break;
+			}
+			case kSurfaceNodeType:
 			{
-				ShaderName = Shader->name();
+				AlSurfaceNode* SurfaceNode = DagNode->asSurfaceNodePtr();
+				TUniquePtr<AlSurface> Surface(SurfaceNode->surface());
+				if (Surface.IsValid())
+				{
+					TUniquePtr<AlShader> Shader(Surface->firstShader());
+					if (Shader.IsValid())
+					{
+						ShaderName = Shader->name();
+					}
+				}
+				ProcessAlShellNode(DagNode, ParentInfo, ShaderName);
+				break;
 			}
 
-			ProcessAlShellNode(*DagNode, ParentInfo, ShaderName);
-			break;
-		}
-		case kSurfaceNodeType:
-		{
-			AlSurfaceNode *SurfaceNode = DagNode->asSurfaceNodePtr();
-			AlSurface *Surface = SurfaceNode->surface();
-			if (AlShader * Shader = Surface->firstShader())
+			case kMeshNodeType:
 			{
-				ShaderName = Shader->name();
+				AlMeshNode* MeshNode = DagNode->asMeshNodePtr();
+				TUniquePtr<AlMesh> Mesh(MeshNode->mesh());
+				if (Mesh.IsValid())
+				{
+					TUniquePtr<AlShader> Shader(Mesh->firstShader());
+					if (Shader.IsValid())
+					{
+						ShaderName = Shader->name();
+					}
+				}
+				ProcessAlShellNode(DagNode, ParentInfo, ShaderName);
+				break;
 			}
 
-			ProcessAlShellNode(*DagNode, ParentInfo, ShaderName);
-			break;
-		}
-
-		case kMeshNodeType:
-		{
-			AlMeshNode *MeshNode = DagNode->asMeshNodePtr();
-			AlMesh *Mesh = MeshNode->mesh();
-			if (AlShader * Shader = Mesh->firstShader())
+			// Traverse down through groups
+			case kGroupNodeType:
 			{
-				ShaderName = Shader->name();
+				ProcessAlGroupNode(*DagNode, ParentInfo);
+				break;
 			}
 
-			ProcessAlShellNode(*DagNode, ParentInfo, ShaderName);
-			break;
-		}
-
-		// Traverse down through groups
-		case kGroupNodeType:
-		{
-			AlGroupNode * GroupNode = DagNode->asGroupNodePtr();
-			if (AlIsValid(GroupNode))
+			default:
 			{
-				ProcessAlGroupNode(*GroupNode, ParentInfo);
+				break;
 			}
-			break;
-		}
 
-		default:
-		{
-		}
+			}
 
+			DagNode = GetNextNode(DagNode);
 		}
-
-		DagNode = GetNextNode(DagNode);
 	}
 	return true;
 }
@@ -1790,7 +1812,7 @@ TOptional<FMeshDescription> FWireTranslatorImpl::MeshDagNodeWithExternalMesher(T
 		ObjectReference = EAliasObjectReference::ParentReference;
 	}
 
-	for(AlDagNode* DagNode : Body->ShellSet)
+	for (const TSharedPtr<AlDagNode>& DagNode : Body->ShellSet)
 	{
 		AliasBRepConverter->AddBRep(*DagNode, ObjectReference);
 	}
@@ -1807,7 +1829,7 @@ TOptional<FMeshDescription> FWireTranslatorImpl::GetMeshOfShellNode(AlDagNode& D
 {
 	if (CADModelConverter->IsSessionValid())
 	{
-		TOptional< FMeshDescription > UEMesh = MeshDagNodeWithExternalMesher(DagNode, MeshElement, MeshParameters);
+		TOptional<FMeshDescription> UEMesh = MeshDagNodeWithExternalMesher(DagNode, MeshElement, MeshParameters);
 		return UEMesh;
 	}
 	else
@@ -1815,23 +1837,22 @@ TOptional<FMeshDescription> FWireTranslatorImpl::GetMeshOfShellNode(AlDagNode& D
 		AlMatrix4x4 AlMatrix;
 		DagNode.inverseGlobalTransformationMatrix(AlMatrix);
 		// TODO: the best way, should be to don't have to apply inverse global transform to the generated mesh
-		AlDagNode *TesselatedNode = OpenModelUtils::TesselateDagLeaf(&DagNode, ETesselatorType::Fast, TessellationOptions.ChordTolerance);
-		if (TesselatedNode != nullptr)
+		TSharedPtr<AlDagNode> TesselatedNode = OpenModelUtils::TesselateDagLeaf(DagNode, ETesselatorType::Fast, TessellationOptions.ChordTolerance);
+		if (TesselatedNode.IsValid())
 		{
 			// Get the meshes from the dag nodes. Note that removing the mesh's DAG.
 			// will also removes the meshes, so we have to do it later.
-			TOptional< FMeshDescription > UEMesh = GetMeshOfNodeMesh(*(TesselatedNode->asMeshNodePtr()), MeshElement, MeshParameters, &AlMatrix);
-			delete TesselatedNode;
+			TOptional<FMeshDescription> UEMesh = GetMeshOfNodeMesh(*TesselatedNode, MeshElement, MeshParameters, &AlMatrix);
 			return UEMesh;
 		}
 	}
-	return TOptional< FMeshDescription >();
+	return TOptional<FMeshDescription>();
 }
 
 
 TOptional<FMeshDescription> FWireTranslatorImpl::GetMeshOfShellBody(TSharedRef<BodyData> Body, TSharedRef<IDatasmithMeshElement> MeshElement, CADLibrary::FMeshParameters& MeshParameters)
 {
-	TOptional< FMeshDescription > UEMesh = MeshDagNodeWithExternalMesher(Body, MeshElement, MeshParameters);
+	TOptional<FMeshDescription> UEMesh = MeshDagNodeWithExternalMesher(Body, MeshElement, MeshParameters);
 	return UEMesh;
 }
 
@@ -1842,58 +1863,79 @@ TOptional<FMeshDescription> FWireTranslatorImpl::GetMeshOfMeshBody(TSharedRef<Bo
 	MeshDescription.Empty();
 	bool True = true;
 
-	for (auto DagNode : Body->ShellSet)
+	for (const TSharedPtr<AlDagNode>& DagNode : Body->ShellSet)
 	{
-		AlMeshNode *MeshNode = DagNode->asMeshNodePtr();
-		AlMesh *Mesh = MeshNode->mesh();
-		if (Mesh)
+		AlMeshNode* MeshNode = DagNode->asMeshNodePtr();
+		if (!AlIsValid(MeshNode))
+		{
+			continue;
+		}
+
+		AlMesh* MeshPtr = MeshNode->mesh();
+		if (!AlIsValid(MeshPtr))
 		{
-			OpenModelUtils::TransferAlMeshToMeshDescription(*Mesh, MeshDescription, MeshParameters, True, true);
+			continue;
 		}
+
+		TUniquePtr<AlMesh> Mesh(MeshPtr);
+		OpenModelUtils::TransferAlMeshToMeshDescription(*MeshPtr, MeshDescription, MeshParameters, True, true);
 	}
 
 	return MoveTemp(MeshDescription);
 }
 
-TOptional<FMeshDescription> FWireTranslatorImpl::GetMeshOfNodeMesh(AlMeshNode& MeshNode, TSharedRef<IDatasmithMeshElement> MeshElement, CADLibrary::FMeshParameters& MeshParameters, AlMatrix4x4* AlMeshInvGlobalMatrix)
+TOptional<FMeshDescription> FWireTranslatorImpl::GetMeshOfNodeMesh(AlDagNode& TesselatedNode, TSharedRef<IDatasmithMeshElement> MeshElement, CADLibrary::FMeshParameters& MeshParameters, AlMatrix4x4* AlMeshInvGlobalMatrix)
 {
-	AlMesh * Mesh = MeshNode.mesh();
-	if (AlIsValid(Mesh))
+	AlMeshNode* MeshNode = TesselatedNode.asMeshNodePtr();
+	if (!AlIsValid(MeshNode))
 	{
-		if (AlMeshInvGlobalMatrix != nullptr)
-		{
-			Mesh->transform(*AlMeshInvGlobalMatrix);
-		}
-		return ImportMesh(*Mesh, MeshParameters);
+		return TOptional<FMeshDescription>();
 	}
-	return TOptional< FMeshDescription >();
+
+	AlMesh* Mesh = MeshNode->mesh();
+	if (!AlIsValid(Mesh))
+	{
+		return TOptional<FMeshDescription>();
+	}
+
+	TUniquePtr<AlMesh> SharedMesh(Mesh);
+	if (AlMeshInvGlobalMatrix != nullptr)
+	{
+		SharedMesh->transform(*AlMeshInvGlobalMatrix);
+	}
+
+	return ImportMesh(*Mesh, MeshParameters);
 }
 
 TOptional<FMeshDescription> FWireTranslatorImpl::GetMeshDescription(TSharedRef<IDatasmithMeshElement> MeshElement, CADLibrary::FMeshParameters& MeshParameters, TSharedRef<BodyData> Body)
 {
-	if(Body->ShellSet.Num() == 0 )
+	if (Body->ShellSet.Num() == 0)
 	{
-		return TOptional< FMeshDescription >();
+		return TOptional<FMeshDescription>();
 	}
 
-	AlDagNode& DagNode = *Body->ShellSet[0];
-
-	if (DagNode.layer() && DagNode.layer()->isSymmetric())
+	TSharedPtr<AlDagNode> DagNode = Body->ShellSet[0];
+	AlLayer* LayerPtr = DagNode->layer();
+	if (AlIsValid(LayerPtr))
 	{
-		MeshParameters.bIsSymmetric = true;
-		double Normal[3], Origin[3];
-		DagNode.layer()->symmetricNormal(Normal[0], Normal[1], Normal[2]);
-		DagNode.layer()->symmetricOrigin(Origin[0], Origin[1], Origin[2]);
+		const TUniquePtr<AlLayer> Layer(LayerPtr);
+		if (LayerPtr->isSymmetric())
+		{
+			MeshParameters.bIsSymmetric = true;
+			double Normal[3], Origin[3];
+			LayerPtr->symmetricNormal(Normal[0], Normal[1], Normal[2]);
+			LayerPtr->symmetricOrigin(Origin[0], Origin[1], Origin[2]);
 
-		MeshParameters.SymmetricOrigin.X = (float)Origin[0];
-		MeshParameters.SymmetricOrigin.Y = (float)Origin[1];
-		MeshParameters.SymmetricOrigin.Z = (float)Origin[2];
-		MeshParameters.SymmetricNormal.X = (float)Normal[0];
-		MeshParameters.SymmetricNormal.Y = (float)Normal[1];
-		MeshParameters.SymmetricNormal.Z = (float)Normal[2];
+			MeshParameters.SymmetricOrigin.X = (float)Origin[0];
+			MeshParameters.SymmetricOrigin.Y = (float)Origin[1];
+			MeshParameters.SymmetricOrigin.Z = (float)Origin[2];
+			MeshParameters.SymmetricNormal.X = (float)Normal[0];
+			MeshParameters.SymmetricNormal.Y = (float)Normal[1];
+			MeshParameters.SymmetricNormal.Z = (float)Normal[2];
+		}
 	}
 
-	if(Body->bCadData)
+	if (Body->bCadData)
 	{
 		return GetMeshOfShellBody(Body, MeshElement, MeshParameters);
 	}
@@ -1905,16 +1947,15 @@ TOptional<FMeshDescription> FWireTranslatorImpl::GetMeshDescription(TSharedRef<I
 
 TOptional<FMeshDescription> FWireTranslatorImpl::GetMeshDescription(TSharedRef<IDatasmithMeshElement> MeshElement, CADLibrary::FMeshParameters& MeshParameters)
 {
-	AlDagNode ** DagNodeTemp = MeshElementToAlDagNodeMap.Find(&MeshElement.Get());
-	if (DagNodeTemp == nullptr || *DagNodeTemp == nullptr)
+	TSharedPtr<AlDagNode>* DagNodeTemp = MeshElementToAlDagNodeMap.Find(&MeshElement.Get());
+	if (DagNodeTemp == nullptr || !DagNodeTemp->IsValid())
 	{
 		TSharedPtr<BodyData>* BodyTemp = MeshElementToBodyMap.Find(&MeshElement.Get());
-		if (BodyTemp != nullptr && (*BodyTemp).IsValid())
+		if (BodyTemp != nullptr && BodyTemp->IsValid())
 		{
 			return GetMeshDescription(MeshElement, MeshParameters, (*BodyTemp).ToSharedRef());
 		}
-
-		return TOptional< FMeshDescription >();
+		return TOptional<FMeshDescription>();
 	}
 
 	AlDagNode& DagNode = **DagNodeTemp;
@@ -1926,53 +1967,58 @@ TOptional<FMeshDescription> FWireTranslatorImpl::GetMeshDescription(TSharedRef<I
 		DagNode.getSurfaceOrientation(bAlOrientation);
 		MeshParameters.bNeedSwapOrientation = (bool)bAlOrientation;
 
-		if (DagNode.layer() && DagNode.layer()->isSymmetric())
+		AlLayer* LayerPtr = DagNode.layer();
+		if (AlIsValid(LayerPtr))
 		{
-			MeshParameters.bIsSymmetric = true;
-			double Normal[3], Origin[3];
-			DagNode.layer()->symmetricNormal(Normal[0], Normal[1], Normal[2]);
-			DagNode.layer()->symmetricOrigin(Origin[0], Origin[1], Origin[2]);
-
-			MeshParameters.SymmetricOrigin.X = (float)Origin[0];
-			MeshParameters.SymmetricOrigin.Y = (float)Origin[1];
-			MeshParameters.SymmetricOrigin.Z = (float)Origin[2];
-			MeshParameters.SymmetricNormal.X = (float)Normal[0];
-			MeshParameters.SymmetricNormal.Y = (float)Normal[1];
-			MeshParameters.SymmetricNormal.Z = (float)Normal[2];
+			const TUniquePtr<AlLayer> Layer(LayerPtr);
+			if (LayerPtr->isSymmetric())
+			{
+				MeshParameters.bIsSymmetric = true;
+				double Normal[3], Origin[3];
+				LayerPtr->symmetricNormal(Normal[0], Normal[1], Normal[2]);
+				LayerPtr->symmetricOrigin(Origin[0], Origin[1], Origin[2]);
+
+				MeshParameters.SymmetricOrigin.X = (float)Origin[0];
+				MeshParameters.SymmetricOrigin.Y = (float)Origin[1];
+				MeshParameters.SymmetricOrigin.Z = (float)Origin[2];
+				MeshParameters.SymmetricNormal.X = (float)Normal[0];
+				MeshParameters.SymmetricNormal.Y = (float)Normal[1];
+				MeshParameters.SymmetricNormal.Z = (float)Normal[2];
+			}
 		}
 	}
 
 	switch (objectType)
 	{
-		case kShellNodeType:
-		case kSurfaceNodeType:
-		{
-			return GetMeshOfShellNode(DagNode, MeshElement, MeshParameters);
-			break;
-		}
+	case kShellNodeType:
+	case kSurfaceNodeType:
+	{
+		return GetMeshOfShellNode(DagNode, MeshElement, MeshParameters);
+		break;
+	}
 
-		case kMeshNodeType:
-		{
-			return GetMeshOfNodeMesh(*(DagNode.asMeshNodePtr()), MeshElement, MeshParameters);
-			break;
-		}
+	case kMeshNodeType:
+	{
+		return GetMeshOfNodeMesh(DagNode, MeshElement, MeshParameters);
+		break;
+	}
 
-		default:
-			break;
+	default:
+		break;
 	}
 
-	return TOptional< FMeshDescription >();
+	return TOptional<FMeshDescription>();
 }
 
 
 // Note that Alias file unit is cm like UE
-TOptional< FMeshDescription > FWireTranslatorImpl::ImportMesh(AlMesh& CurrentMesh, CADLibrary::FMeshParameters& MeshParameters)
+TOptional<FMeshDescription> FWireTranslatorImpl::ImportMesh(AlMesh& InMesh, CADLibrary::FMeshParameters& InMeshParameters)
 {
 
 	FMeshDescription MeshDescription;
 	DatasmithMeshHelper::PrepareAttributeForStaticMesh(MeshDescription);
-	bool True = true;
-	OpenModelUtils::TransferAlMeshToMeshDescription(CurrentMesh, MeshDescription, MeshParameters, True);
+	bool bMerge = true;
+	OpenModelUtils::TransferAlMeshToMeshDescription(InMesh, MeshDescription, InMeshParameters, bMerge);
 
 	return MoveTemp(MeshDescription);
 }
@@ -1980,7 +2026,7 @@ TOptional< FMeshDescription > FWireTranslatorImpl::ImportMesh(AlMesh& CurrentMes
 bool FWireTranslatorImpl::LoadStaticMesh(const TSharedRef<IDatasmithMeshElement> MeshElement, FDatasmithMeshElementPayload& OutMeshPayload, const FDatasmithTessellationOptions& InTessellationOptions)
 {
 	CADLibrary::FMeshParameters MeshParameters;
-	if (TOptional< FMeshDescription > Mesh = GetMeshDescription(MeshElement, MeshParameters))
+	if (TOptional<FMeshDescription> Mesh = GetMeshDescription(MeshElement, MeshParameters))
 	{
 		OutMeshPayload.LodMeshes.Add(MoveTemp(Mesh.GetValue()));
 		AddSurfaceDataForMesh(MeshElement->GetFile(), MeshParameters, InTessellationOptions, OutMeshPayload);
@@ -2030,12 +2076,11 @@ bool FDatasmithWireTranslator::IsSourceSupported(const FDatasmithSceneSource& So
 bool FDatasmithWireTranslator::LoadScene(TSharedRef<IDatasmithScene> OutScene)
 {
 #ifdef USE_OPENMODEL
-	// Check installed version of Alias Tools because binaries from 2020.1 to 2020.3 crashes on load
-	const uint64 LibAlias2020Version = 7318349414924288;
-	const uint64 LibAlias2021Version = 7599824377020416;
+	// Check installed version of Alias Tools because binaries before 2021.3 are not compatible with lib Alias 2022.0
+	const uint64 LibAlias2021_3Version = 7599833027117059;
 	uint64 FileVersion = FPlatformMisc::GetFileVersion(TEXT("libalias_api.dll"));
 
-	if (FileVersion > LibAlias2020Version && FileVersion < LibAlias2021Version)
+	if (FileVersion < LibAlias2021_3Version)
 	{
 #if WITH_EDITOR
 		// Display message and abort import
@@ -2043,7 +2088,7 @@ bool FDatasmithWireTranslator::LoadScene(TSharedRef<IDatasmithScene> OutScene)
 
 		TSharedPtr<IMessageLogListing> LogListing = MessageLogModule.GetLogListing("DatasmithWireTranslator");
 
-		if(LogListing.IsValid())
+		if (LogListing.IsValid())
 		{
 			LogListing->SetLabel(LOCTEXT("MessageLogging", "Datasmith Wire Translator"));
 
@@ -2055,19 +2100,19 @@ bool FDatasmithWireTranslator::LoadScene(TSharedRef<IDatasmithScene> OutScene)
 		else
 		{
 			UE_LOG(LogDatasmithWireTranslator, Error, TEXT(WRONG_VERSION_TEXT));
-		}
+	}
 #else
 		UE_LOG(LogDatasmithWireTranslator, Error, TEXT(WRONG_VERSION_TEXT));
 #endif
 
 		return false;
-	}
+}
 
 	const FString& Filename = GetSource().GetSourceFile();
 
 	Translator = MakeShared<FWireTranslatorImpl>(Filename, OutScene);
 	if (!Translator)
-	{
+	{ 
 		return false;
 	}
 
@@ -2106,7 +2151,7 @@ void FDatasmithWireTranslator::SetSceneImportOptions(TArray<TStrongObjectPtr<UDa
 
 	if (Translator)
 	{
-		Translator->SetTessellationOptions( GetCommonTessellationOptions() );
+		Translator->SetTessellationOptions(GetCommonTessellationOptions());
 	}
 #endif
 }
diff --git a/Engine/Plugins/Enterprise/DatasmithCADImporter/Source/DatasmithWireTranslator/Private/OpenModelUtils.cpp b/Engine/Plugins/Enterprise/DatasmithCADImporter/Source/DatasmithWireTranslator/Private/OpenModelUtils.cpp
--- a/Engine/Plugins/Enterprise/DatasmithCADImporter/Source/DatasmithWireTranslator/Private/OpenModelUtils.cpp
+++ b/Engine/Plugins/Enterprise/DatasmithCADImporter/Source/DatasmithWireTranslator/Private/OpenModelUtils.cpp
@@ -16,9 +16,10 @@
 #include "MeshDescription.h"
 #include "StaticMeshAttributes.h"
 #include "StaticMeshOperations.h"
+
 using namespace OpenModelUtils;
 
-const TCHAR * OpenModelUtils::AlObjectTypeToString(AlObjectType type)
+const TCHAR* OpenModelUtils::AlObjectTypeToString(AlObjectType type)
 {
 	switch (type)
 	{
@@ -254,7 +255,7 @@ const TCHAR * OpenModelUtils::AlObjectTypeToString(AlObjectType type)
 	return TEXT("kUndefined");
 }
 
-const TCHAR * OpenModelUtils::AlShadingFieldToString(AlShadingFields field)
+const TCHAR* OpenModelUtils::AlShadingFieldToString(AlShadingFields field)
 {
 	switch (field)
 	{
@@ -1551,17 +1552,23 @@ const TCHAR * OpenModelUtils::AlShadingFieldToString(AlShadingFields field)
 	return TEXT("");
 }
 
-void OpenModelUtils::SetActorTransform(const TSharedPtr< IDatasmithActorElement >& ActorElement, AlDagNode& DagNode)
+void OpenModelUtils::SetActorTransform(TSharedPtr<IDatasmithActorElement>& OutActorElement, const AlDagNode& InDagNode)
 {
-	if (ActorElement)
+	if (OutActorElement)
 	{
-		if (DagNode.layer() && DagNode.layer()->isSymmetric())
+		AlLayer* LayerPtr = InDagNode.layer();
+		if (AlIsValid(LayerPtr))
 		{
-			return;
+			const TUniquePtr<AlLayer> Layer(LayerPtr);
+
+			if (LayerPtr->isSymmetric())
+			{
+				return;
+			}
 		}
 
 		AlMatrix4x4 AlMatrix;
-		DagNode.localTransformationMatrix(AlMatrix);
+		InDagNode.localTransformationMatrix(AlMatrix);
 
 		// LWC_TODO: Perf improvements with LWC enabled
 		FMatrix44f Matrix;
@@ -1575,14 +1582,14 @@ void OpenModelUtils::SetActorTransform(const TSharedPtr< IDatasmithActorElement
 		}
 		FTransform LocalTransform(Matrix);
 		FTransform LocalUETransform = FDatasmithUtils::ConvertTransform(FDatasmithUtils::EModelCoordSystem::ZUp_RightHanded, LocalTransform);
-		
-		ActorElement->SetTranslation(LocalUETransform.GetTranslation());
-		ActorElement->SetScale(LocalUETransform.GetScale3D());
-		ActorElement->SetRotation(LocalUETransform.GetRotation());
+
+		OutActorElement->SetTranslation(LocalUETransform.GetTranslation());
+		OutActorElement->SetScale(LocalUETransform.GetScale3D());
+		OutActorElement->SetRotation(LocalUETransform.GetRotation());
 	}
 }
 
-bool OpenModelUtils::IsValidActor(const TSharedPtr< IDatasmithActorElement >& ActorElement)
+bool OpenModelUtils::IsValidActor(const TSharedPtr<IDatasmithActorElement>& ActorElement)
 {
 	if (ActorElement != nullptr)
 	{
@@ -1592,7 +1599,7 @@ bool OpenModelUtils::IsValidActor(const TSharedPtr< IDatasmithActorElement >& Ac
 		}
 		else if (ActorElement->IsA(EDatasmithElementType::StaticMeshActor))
 		{
-			const TSharedPtr< IDatasmithMeshActorElement >& MeshActorElement = StaticCastSharedPtr< IDatasmithMeshActorElement >(ActorElement);
+			const TSharedPtr<IDatasmithMeshActorElement>& MeshActorElement = StaticCastSharedPtr<IDatasmithMeshActorElement>(ActorElement);
 			return FCString::Strlen(MeshActorElement->GetStaticMeshPathName()) > 0;
 		}
 	}
@@ -1601,14 +1608,12 @@ bool OpenModelUtils::IsValidActor(const TSharedPtr< IDatasmithActorElement >& Ac
 
 bool OpenModelUtils::TransferAlMeshToMeshDescription(const AlMesh& AliasMesh, FMeshDescription& MeshDescription, CADLibrary::FMeshParameters& MeshParameters, bool& bHasNormal, bool bMerge)
 {
-	// Ref. GP3DMVisitorImpl::visitMesh
-	// Ref. FGPureMeshInterface::CreateMesh
 	if (AliasMesh.numberOfVertices() == 0 || AliasMesh.numberOfTriangles() == 0)
 	{
 		return false;
 	}
 
-	if( !bMerge )
+	if (!bMerge)
 	{
 		MeshDescription.Empty();
 	}
@@ -1642,14 +1647,14 @@ bool OpenModelUtils::TransferAlMeshToMeshDescription(const AlMesh& AliasMesh, FM
 	const int32 VertexInstanceCount = 3 * TriangleCount;
 
 	TArray<FVertexID> VertexPositionIDs;
-	VertexPositionIDs.SetNum( VertexCount*NbStep );
+	VertexPositionIDs.SetNum(VertexCount * NbStep);
 
 	// Reserve space for attributes
 	// At this point, all the faces are triangles
-	MeshDescription.ReserveNewVertices(VertexCount*NbStep);
-	MeshDescription.ReserveNewVertexInstances(VertexInstanceCount*NbStep);
-	MeshDescription.ReserveNewEdges(VertexInstanceCount*NbStep);
-	MeshDescription.ReserveNewPolygons(TriangleCount*NbStep);
+	MeshDescription.ReserveNewVertices(VertexCount * NbStep);
+	MeshDescription.ReserveNewVertexInstances(VertexInstanceCount * NbStep);
+	MeshDescription.ReserveNewEdges(VertexInstanceCount * NbStep);
+	MeshDescription.ReserveNewPolygons(TriangleCount * NbStep);
 
 	// Assume one material per mesh, no partitioning
 	MeshDescription.ReserveNewPolygonGroups(1);
@@ -1658,15 +1663,15 @@ bool OpenModelUtils::TransferAlMeshToMeshDescription(const AlMesh& AliasMesh, FM
 	PolygonGroupImportedMaterialSlotNames[PolyGroupId] = ImportedSlotName;
 
 	// At least one UV set must exist.
-	if( VertexInstanceUVs.GetNumChannels() == 0 )
+	if (VertexInstanceUVs.GetNumChannels() == 0)
 	{
 		VertexInstanceUVs.SetNumChannels(1);
 	}
 
 	// Get Alias mesh info
-	const float * AlVertices = AliasMesh.vertices();
+	const float* AlVertices = AliasMesh.vertices();
 
-	for(int32 Step = 0; Step < NbStep; Step++)
+	for (int32 Step = 0; Step < NbStep; Step++)
 	{
 		// Fill the vertex array
 		if (Step == 0)
@@ -1674,18 +1679,18 @@ bool OpenModelUtils::TransferAlMeshToMeshDescription(const AlMesh& AliasMesh, FM
 			FVertexID* VertexPositionIDPtr = VertexPositionIDs.GetData();
 			for (int Index = 0; Index < VertexCount; ++Index, ++VertexPositionIDPtr)
 			{
-				const float * CurVertex = AlVertices + 3 * Index;
+				const float* CurVertex = AlVertices + 3 * Index;
 				*VertexPositionIDPtr = MeshDescription.CreateVertex();
 				// ConvertVector_ZUp_RightHanded
 				VertexPositions[*VertexPositionIDPtr] = FVector(-CurVertex[0], CurVertex[1], CurVertex[2]);
 			}
 		}
-		else 
+		else
 		{
 			FVertexID* VertexPositionIDPtr = VertexPositionIDs.GetData() + VertexCount;
 			for (int Index = 0, PositionIndex = VertexCount; Index < VertexCount; ++Index, ++VertexPositionIDPtr)
 			{
-				const float * CurVertex = AlVertices + 3 * Index;
+				const float* CurVertex = AlVertices + 3 * Index;
 				*VertexPositionIDPtr = MeshDescription.CreateVertex();
 				// ConvertVector_ZUp_RightHanded
 				VertexPositions[*VertexPositionIDPtr] = SymmetricMatrix.TransformPosition(FVector(-CurVertex[0], CurVertex[1], CurVertex[2]));
@@ -1701,7 +1706,7 @@ bool OpenModelUtils::TransferAlMeshToMeshDescription(const AlMesh& AliasMesh, FM
 
 		// Get Alias mesh info
 		const int* Triangles = AliasMesh.triangles();
-		const float * AlNormals = AliasMesh.normals();
+		const float* AlNormals = AliasMesh.normals();
 		const float* AlUVs = AliasMesh.uvs();
 
 		// Get per-triangle data: indices, normals and uvs
@@ -1716,15 +1721,15 @@ bool OpenModelUtils::TransferAlMeshToMeshDescription(const AlMesh& AliasMesh, FM
 					CornerVertexInstanceIDs[VertexIndex] = MeshDescription.CreateVertexInstance(CornerVertexIDs[VertexIndex]);
 
 					// Set the normal
-					const float * CurNormal = &AlNormals[3 * Triangles[TIndex]];
+					const float* CurNormal = &AlNormals[3 * Triangles[TIndex]];
 					// ConvertVector_ZUp_RightHanded
 					FVector UENormal(-CurNormal[0], CurNormal[1], CurNormal[2]);
 					UENormal = UENormal.GetSafeNormal();
-					if (Step>0) 
+					if (Step > 0)
 					{
 						UENormal = SymmetricMatrix.TransformVector(UENormal);
 					}
-					else 
+					else
 					{
 						UENormal *= -1.;
 					}
@@ -1751,25 +1756,25 @@ bool OpenModelUtils::TransferAlMeshToMeshDescription(const AlMesh& AliasMesh, FM
 				const FPolygonID NewPolygonID = MeshDescription.CreatePolygon(PolyGroupId, CornerVertexInstanceIDs);
 			}
 		}
-		else 
+		else
 		{
 			for (int32 FaceIndex = 0; FaceIndex < TriangleCount; ++FaceIndex, Triangles += 3)
 			{
 				// Create Vertex instances and set their attributes
 				for (int32 VertexIndex = 0; VertexIndex < CornerCount; ++VertexIndex)
 				{
-					CornerVertexIDs[VertexIndex] =  VertexPositionIDs[Triangles[VertexIndex] + VertexCount * Step];
+					CornerVertexIDs[VertexIndex] = VertexPositionIDs[Triangles[VertexIndex] + VertexCount * Step];
 					CornerVertexInstanceIDs[VertexIndex] = MeshDescription.CreateVertexInstance(CornerVertexIDs[VertexIndex]);
 
 					// Set the normal
-					const float * CurNormal = &AlNormals[3 * Triangles[VertexIndex]];
+					const float* CurNormal = &AlNormals[3 * Triangles[VertexIndex]];
 
 					// ConvertVector_ZUp_RightHanded
 					FVector UENormal(-CurNormal[0], CurNormal[1], CurNormal[2]);
 					UENormal = UENormal.GetSafeNormal();
 					if (Step > 0)
 					{
-						UENormal = SymmetricMatrix.TransformVector(UENormal) *-1;
+						UENormal = SymmetricMatrix.TransformVector(UENormal) * -1;
 					}
 					VertexInstanceNormals[CornerVertexInstanceIDs[VertexIndex]] = UENormal;
 				}
@@ -1802,29 +1807,29 @@ bool OpenModelUtils::TransferAlMeshToMeshDescription(const AlMesh& AliasMesh, FM
 }
 
 
-AlDagNode* OpenModelUtils::TesselateDagLeaf(AlDagNode* DagLeaf, ETesselatorType TessType, double Tolerance)
+TSharedPtr<AlDagNode> OpenModelUtils::TesselateDagLeaf(const AlDagNode& DagLeaf, ETesselatorType TessType, double Tolerance)
 {
-	AlDagNode * TesselatedNode;
+	AlDagNode* TesselatedNode = nullptr;
 	statusCode TessStatus;
 
 	switch (TessType)
 	{
 	case(ETesselatorType::Accurate):
-		TessStatus = AlTesselate::chordHeightDeviationAccurate(TesselatedNode, DagLeaf, Tolerance);
+		TessStatus = AlTesselate::chordHeightDeviationAccurate(TesselatedNode, &DagLeaf, Tolerance);
 		break;
 	case(ETesselatorType::Fast):
 	default:
-		TessStatus = AlTesselate::chordHeightDeviationFast(TesselatedNode, DagLeaf, Tolerance);
+		TessStatus = AlTesselate::chordHeightDeviationFast(TesselatedNode, &DagLeaf, Tolerance);
 		break;
 	}
 
 	if ((TessStatus == sSuccess) && (AlIsValid(TesselatedNode) == TRUE))
 	{
-		return TesselatedNode;
+		return TSharedPtr<AlDagNode>(TesselatedNode);
 	}
 	else
 	{
-		return nullptr;
+		return TSharedPtr<AlDagNode>();
 	}
 }
 
diff --git a/Engine/Plugins/Enterprise/DatasmithCADImporter/Source/DatasmithWireTranslator/Private/OpenModelUtils.h b/Engine/Plugins/Enterprise/DatasmithCADImporter/Source/DatasmithWireTranslator/Private/OpenModelUtils.h
--- a/Engine/Plugins/Enterprise/DatasmithCADImporter/Source/DatasmithWireTranslator/Private/OpenModelUtils.h
+++ b/Engine/Plugins/Enterprise/DatasmithCADImporter/Source/DatasmithWireTranslator/Private/OpenModelUtils.h
@@ -39,12 +39,12 @@ enum class EAlShaderModelType : uint8
 
 namespace OpenModelUtils
 {
-	const TCHAR * AlObjectTypeToString(AlObjectType type);
-	const TCHAR * AlShadingFieldToString(AlShadingFields field);
+	const TCHAR* AlObjectTypeToString(AlObjectType type);
+	const TCHAR* AlShadingFieldToString(AlShadingFields field);
 
-	void SetActorTransform(const TSharedPtr< IDatasmithActorElement >& ActorElement, AlDagNode& DagNode);
+	void SetActorTransform(TSharedPtr<IDatasmithActorElement>& OutActorElement, const AlDagNode& InDagNode);
 
-	bool IsValidActor(const TSharedPtr< IDatasmithActorElement >& ActorElement);
+	bool IsValidActor(const TSharedPtr<IDatasmithActorElement>& ActorElement);
 
 	inline FString UuidToString(const uint32& Uuid)
 	{
@@ -58,22 +58,22 @@ namespace OpenModelUtils
 		return HashCombine(IdA, HashCombine(IdB, HashCombine(IdC, IdD)));
 	}
 
-	inline uint32 GetAlDagNodeUuid(AlDagNode& GroupNode)
+	inline uint32 GetAlDagNodeUuid(const TSharedPtr<AlDagNode>& DagNode)
 	{
-		if (GroupNode.hasPersistentID() == sSuccess)
+		if (DagNode->hasPersistentID() == sSuccess)
 		{
 			AlPersistentID* PersistentID;
-			GroupNode.persistentID(PersistentID);
+			DagNode->persistentID(PersistentID);
 			return GetTypeHash(*PersistentID);
 		}
-		FString Label = GroupNode.name();
+		FString Label = UTF8_TO_TCHAR(DagNode->name());
 		return GetTypeHash(Label);
 	}
 
 	// Note that Alias file unit is cm like UE
 	bool TransferAlMeshToMeshDescription(const AlMesh& Mesh, FMeshDescription& MeshDescription, CADLibrary::FMeshParameters& SymmetricParameters, bool& bHasNormal, bool bMerge = false);
 
-	AlDagNode* TesselateDagLeaf(AlDagNode* DagLeaf, ETesselatorType TessType, double Tolerance);
+	TSharedPtr<AlDagNode> TesselateDagLeaf(const AlDagNode& DagLeaf, ETesselatorType TessType, double Tolerance);
 }
 
 
