Pluggin of CADKernel in Wire and Rhino Translator
By setting r.CADTranslator.EnableKernelIOTessellation CVar to 0 (disabling KernelIO Tessellation) CADKernel is activated in wire interface

diff --git a/Engine/Plugins/Enterprise/DatasmithCADImporter/Source/CADKernelSurface/Private/CADKernelSurfaceExtension.cpp b/Engine/Plugins/Enterprise/DatasmithCADImporter/Source/CADKernelSurface/Private/CADKernelSurfaceExtension.cpp
--- a/Engine/Plugins/Enterprise/DatasmithCADImporter/Source/CADKernelSurface/Private/CADKernelSurfaceExtension.cpp
+++ b/Engine/Plugins/Enterprise/DatasmithCADImporter/Source/CADKernelSurface/Private/CADKernelSurfaceExtension.cpp
@@ -9,6 +9,7 @@
 #include "Engine/StaticMesh.h"
 #include "MeshDescription.h"
 #include "MeshDescriptionHelper.h"
+#include "Misc/FileHelper.h"
 #include "StaticMeshAttributes.h"
 #include "UObject/EnterpriseObjectVersion.h"
 
@@ -64,26 +65,14 @@ bool UCADKernelParametricSurfaceData::Tessellate(UStaticMesh& StaticMesh, const
 		CADKernelSession->AddDatabase(RawData);
 
 		TSharedRef<CADKernel::FModel> CADKernelModel = CADKernelSession->GetModel();
-
-		// Tesselate the model
-		TSharedRef<CADKernel::FModelMesh> CADKernelModelMesh = CADKernel::FEntity::MakeShared<CADKernel::FModelMesh>();
-
-		CADLibrary::FCADKernelTools::DefineMeshCriteria(CADKernelModelMesh, ImportParameters);
-
-		CADKernel::FParametricMesher Mesher(CADKernelModelMesh);
-		Mesher.MeshEntity(CADKernelModel);
-
 		TArray<TSharedPtr<CADKernel::FBody>> CADKernelBodies = CADKernelModel->GetBodies();
 		if (CADKernelBodies.Num() != 1)
 		{
 			return bSuccessfulTessellation;
 		}
 
-		CADLibrary::FBodyMesh BodyMesh;
-		uint32 DefaultMaterialHash = 0;
-
-		TSharedRef<CADKernel::FTopologicalEntity> Body = StaticCastSharedRef<CADKernel::FTopologicalEntity>(CADKernelBodies[0].ToSharedRef());
-		if(CADLibrary::FCADKernelTools::Tessellate(Body, ImportParameters, CadMeshParameters, MeshDescription))
+		TSharedRef<CADKernel::FTopologicalEntity> CADKernelEntity = CADKernelModel;
+		if(CADLibrary::FCADKernelTools::Tessellate(CADKernelEntity, ImportParameters, CadMeshParameters, MeshDescription))
 		{
 			// To update the SectionInfoMap 
 			TPolygonGroupAttributesConstRef<FName> MaterialSlotNames = MeshDescriptionAttributes.GetPolygonGroupMaterialSlotNames();
diff --git a/Engine/Plugins/Enterprise/DatasmithCADImporter/Source/CADKernelSurface/Public/CADModelToCADKernelConverterBase.h b/Engine/Plugins/Enterprise/DatasmithCADImporter/Source/CADKernelSurface/Public/CADModelToCADKernelConverterBase.h
new file mode 100644
--- /dev/null
+++ b/Engine/Plugins/Enterprise/DatasmithCADImporter/Source/CADKernelSurface/Public/CADModelToCADKernelConverterBase.h
@@ -0,0 +1,88 @@
+// Copyright Epic Games, Inc. All Rights Reserved.
+
+#pragma once
+
+#include "CoreMinimal.h"
+
+#include "CADKernelSurfaceExtension.h"
+#include "CADKernelTools.h"
+#include "CADModelConverter.h"
+#include "CADOptions.h"
+#include "IDatasmithSceneElements.h"
+
+#include "CADKernel/Core/Session.h"
+#include "CADKernel/Topo/Model.h"
+
+struct FDatasmithMeshElementPayload;
+struct FDatasmithTessellationOptions;
+
+class FCADModelToCADKernelConverterBase : public CADLibrary::ICADModelConverter
+{
+public:
+
+	FCADModelToCADKernelConverterBase()
+		// Unit is set to cm, 0.01, because Wire's unit is cm.
+		: CADKernelSession(0.00001 / 0.01)
+		, GeometricTolerance(0.00001 / 0.01)
+		, SquareTolerance(GeometricTolerance* GeometricTolerance)
+	{
+		ImportParameters.MetricUnit = 0.01;
+		ImportParameters.ScaleFactor = 1;
+		ImportParameters.bEnableKernelIOTessellation = false;
+	}
+
+	virtual void InitializeProcess(double InMetricUnit) override
+	{
+		CADKernelSession.Clear();
+	}
+
+	virtual bool RepairTopology() override
+	{
+		// Apply stitching if applicable
+		return true;
+	}
+
+	virtual bool SaveBRep(const TCHAR* InFolderPath, TSharedRef<IDatasmithMeshElement>& MeshElement) override
+	{
+		FString FilePath = FPaths::Combine(InFolderPath, MeshElement->GetName()) + TEXT(".ugeom");
+		CADKernelSession.SaveDatabase(*FilePath);
+		MeshElement->SetFile(*FilePath);
+		return true;
+	}
+
+	virtual bool Tessellate(const CADLibrary::FMeshParameters& InMeshParameters, FMeshDescription& OutMeshDescription) override
+	{
+		TSharedRef<CADKernel::FTopologicalEntity> CADKernelEntity = CADKernelSession.GetModel();
+		return CADLibrary::FCADKernelTools::Tessellate(CADKernelEntity, ImportParameters, InMeshParameters, OutMeshDescription);
+	}
+
+	virtual void SetImportParameters(float ChordTolerance, float MaxEdgeLength, float NormalTolerance, CADLibrary::EStitchingTechnique StitchingTechnique, bool bScaleUVMap) override
+	{
+		ImportParameters.ChordTolerance = ChordTolerance;
+		ImportParameters.MaxEdgeLength = MaxEdgeLength;
+		ImportParameters.MaxNormalAngle = NormalTolerance;
+		ImportParameters.StitchingTechnique = StitchingTechnique;
+		ImportParameters.bScaleUVMap = bScaleUVMap;
+		ImportParameters.bEnableKernelIOTessellation = false;
+	}
+
+	virtual bool IsSessionValid() override
+	{
+		return true;
+	}
+
+	virtual void AddSurfaceDataForMesh(const TCHAR* InFilePath, const CADLibrary::FMeshParameters& InMeshParameters, const FDatasmithTessellationOptions& InTessellationOptions, FDatasmithMeshElementPayload& OutMeshPayload) const override
+	{
+		CADKernelSurface::AddSurfaceDataForMesh(InFilePath, ImportParameters, InMeshParameters, InTessellationOptions, OutMeshPayload);
+	}
+
+protected:
+
+	CADKernel::FSession CADKernelSession;
+
+	CADLibrary::FImportParameters ImportParameters;
+	double GeometricTolerance;
+	double SquareTolerance;
+
+};
+
diff --git a/Engine/Plugins/Enterprise/DatasmithCADImporter/Source/DatasmithWireTranslator/DatasmithWireTranslator.Build.cs b/Engine/Plugins/Enterprise/DatasmithCADImporter/Source/DatasmithWireTranslator/DatasmithWireTranslator.Build.cs
--- a/Engine/Plugins/Enterprise/DatasmithCADImporter/Source/DatasmithWireTranslator/DatasmithWireTranslator.Build.cs
+++ b/Engine/Plugins/Enterprise/DatasmithCADImporter/Source/DatasmithWireTranslator/DatasmithWireTranslator.Build.cs
@@ -7,7 +7,7 @@ public class DatasmithWireTranslator : ModuleRules
 {
 	public DatasmithWireTranslator(ReadOnlyTargetRules Target) : base(Target)
 	{
-        PublicDependencyModuleNames.AddRange(
+		PublicDependencyModuleNames.AddRange(
 			new string[]
 			{
 				"Core",
@@ -18,19 +18,21 @@ public class DatasmithWireTranslator : ModuleRules
 		PrivateDependencyModuleNames.AddRange(
 			new string[]
 			{
-                "CADInterfaces",
-                "CADLibrary",
-                "CADTools",
+				"CADInterfaces",
+				"CADKernel",
+				"CADKernelSurface",
+				"CADLibrary",
+				"CADTools",
 				"CoreTechSurface",
-                "DatasmithContent",
+				"DatasmithContent",
 				"DatasmithCore",
 				"DatasmithTranslator",
 				"Engine",
 				"MeshDescription",
 				"ParametricSurface",
 				"StaticMeshDescription",
-            }
-        );
+			}
+		);
 
 		if (Target.bBuildEditor == true)
 		{
diff --git a/Engine/Plugins/Enterprise/DatasmithCADImporter/Source/DatasmithWireTranslator/Private/AliasModelToCADKernelConverter.cpp b/Engine/Plugins/Enterprise/DatasmithCADImporter/Source/DatasmithWireTranslator/Private/AliasModelToCADKernelConverter.cpp
new file mode 100644
--- /dev/null
+++ b/Engine/Plugins/Enterprise/DatasmithCADImporter/Source/DatasmithWireTranslator/Private/AliasModelToCADKernelConverter.cpp
@@ -0,0 +1,348 @@
+// Copyright Epic Games, Inc. All Rights Reserved.
+
+#include "AliasModelToCADKernelConverter.h"
+
+#include "CADKernelTools.h"
+#include "Hal/PlatformMemory.h"
+#include "MeshDescriptionHelper.h"
+
+#include "CADKernel/Core/Session.h"
+#include "CADKernel/Geo/Curves/NURBSCurve.h"
+#include "CADKernel/Geo/GeoEnum.h"
+#include "CADKernel/Geo/Surfaces/NURBSSurface.h"
+#include "CADKernel/Geo/Surfaces/Surface.h"
+
+#include "CADKernel/Math/Point.h"
+#include "CADKernel/Mesh/Meshers/ParametricMesher.h"
+#include "CADKernel/Mesh/Structure/ModelMesh.h"
+
+#include "CADKernel/Topo/Body.h"
+//#include "CADKernel/Topo/Joiner.h"
+#include "CADKernel/Topo/Model.h"
+#include "CADKernel/Topo/Shell.h"
+#include "CADKernel/Topo/TopologicalEdge.h"
+#include "CADKernel/Topo/TopologicalFace.h"
+#include "CADKernel/Topo/TopologicalLoop.h"
+
+#ifdef USE_OPENMODEL
+
+// Alias API wrappes object in AlObjects. This is an abstract base class which holds a reference to an anonymous data structure.
+// The only way to compare two AlObjects is to compare their data structure. That is the reason why private fields are made public. 
+#define private public
+#include "AlCurve.h"
+#include "AlDagNode.h"
+#include "AlShell.h"
+#include "AlShellNode.h"
+#include "AlSurface.h"
+#include "AlSurfaceNode.h"
+#include "AlTrimBoundary.h"
+#include "AlTrimCurve.h"
+#include "AlTrimRegion.h"
+#include "AlTM.h"
+
+using namespace CADLibrary;
+using namespace CADKernel;
+
+namespace AliasToCADKernelUtils
+{
+	template<typename Surface_T>
+	TSharedPtr<FSurface> AddNurbsSurface(double GeometricTolerance, Surface_T& AliasSurface, EAliasObjectReference InObjectReference, const AlMatrix4x4& InAlMatrix)
+	{
+		FNurbsSurfaceHomogeneousData NurbsData;
+		NurbsData.bSwapUV= true;
+		NurbsData.bIsRational = true;
+
+		NurbsData.PoleUCount = AliasSurface.uNumberOfCVsInclMultiples();
+		NurbsData.PoleVCount = AliasSurface.vNumberOfCVsInclMultiples();
+
+		NurbsData.UDegree = AliasSurface.uDegree();  // U order of the surface
+		NurbsData.VDegree = AliasSurface.vDegree();  // V order of the surface
+
+		int32 KnotSizeU = AliasSurface.realuNumberOfKnots() + 2;
+		int32 KnotSizeV = AliasSurface.realvNumberOfKnots() + 2;
+
+		NurbsData.UNodalVector.SetNumUninitialized(KnotSizeU);
+		NurbsData.VNodalVector.SetNumUninitialized(KnotSizeV);
+
+		AliasSurface.realuKnotVector(NurbsData.UNodalVector.GetData() + 1);
+		AliasSurface.realvKnotVector(NurbsData.VNodalVector.GetData() + 1);
+
+		NurbsData.UNodalVector[0] = NurbsData.UNodalVector[1];
+		NurbsData.UNodalVector[KnotSizeU - 1] = NurbsData.UNodalVector[KnotSizeU - 2];
+		NurbsData.VNodalVector[0] = NurbsData.VNodalVector[1];
+		NurbsData.VNodalVector[KnotSizeV - 1] = NurbsData.VNodalVector[KnotSizeV - 2];
+
+		const int32 CoordinateCount = NurbsData.PoleUCount * NurbsData.PoleVCount * 4;
+		NurbsData.HomogeneousPoles.SetNumUninitialized(CoordinateCount);
+
+		if (InObjectReference == EAliasObjectReference::WorldReference)
+		{
+			AliasSurface.CVsWorldPositionInclMultiples(NurbsData.HomogeneousPoles.GetData());
+		}
+		else if (InObjectReference == EAliasObjectReference::ParentReference)
+		{
+			AlTM TranformMatrix(InAlMatrix);
+			AliasSurface.CVsAffectedPositionInclMultiples(TranformMatrix, NurbsData.HomogeneousPoles.GetData());
+		}
+		else  // EAliasObjectReference::LocalReference
+		{
+			AliasSurface.CVsUnaffectedPositionInclMultiples(NurbsData.HomogeneousPoles.GetData());
+		}
+
+		return FEntity::MakeShared<FNURBSSurface>(GeometricTolerance, NurbsData);
+	}
+}
+
+TSharedPtr<FTopologicalEdge> FAliasModelToCADKernelConverter::AddEdge(const AlTrimCurve& AliasTrimCurve, TSharedPtr<FSurface>& CarrierSurface)
+{
+	FNurbsCurveData NurbsCurveData;
+
+	curveFormType Form = AliasTrimCurve.form();
+
+	NurbsCurveData.Degree = AliasTrimCurve.degree();
+	int ControlPointCount = AliasTrimCurve.numberOfCVs();
+
+	NurbsCurveData.Dimension = 2;
+	NurbsCurveData.bIsRational = true;
+
+	int32 KnotCount = AliasTrimCurve.realNumberOfKnots() + 2;
+
+	NurbsCurveData.Weights.SetNumUninitialized(ControlPointCount);
+	NurbsCurveData.Poles.SetNumUninitialized(ControlPointCount);
+	NurbsCurveData.NodalVector.SetNumUninitialized(KnotCount);
+
+	using AlPoint = double[3];
+	// Notice that each CV has three coordinates - the three coordinates describe 2D parameter space, with a homogeneous coordinate.
+	// Each control point is u, v and w, where u and v are parameter space and w is the homogeneous coordinate.
+	AliasTrimCurve.CVsUVPosition(NurbsCurveData.NodalVector.GetData() + 1, (AlPoint*) NurbsCurveData.Poles.GetData());
+
+	AliasTrimCurve.realKnotVector(NurbsCurveData.NodalVector.GetData() + 1);
+	NurbsCurveData.NodalVector[0] = NurbsCurveData.NodalVector[1];
+	NurbsCurveData.NodalVector[KnotCount - 1] = NurbsCurveData.NodalVector[KnotCount - 2];
+
+	for (int32 Index = 0; Index < ControlPointCount; ++Index)
+	{
+		NurbsCurveData.Weights[Index] = NurbsCurveData.Poles[Index].Z;
+		NurbsCurveData.Poles[Index].Z = 0;
+	}
+
+	TSharedRef<FNURBSCurve> Nurbs = FEntity::MakeShared<FNURBSCurve>(NurbsCurveData);
+
+	TSharedRef<FRestrictionCurve> RestrictionCurve = FEntity::MakeShared<FRestrictionCurve>(CarrierSurface.ToSharedRef(), Nurbs);
+	TSharedPtr<FTopologicalEdge> Edge = FTopologicalEdge::Make(RestrictionCurve);
+	if (!Edge.IsValid())
+	{
+		return TSharedPtr<FTopologicalEdge>();
+	}
+
+	// Build topo
+	TUniquePtr<AlTrimCurve> TwinCurve(AliasTrimCurve.getTwinCurve());
+	if (TwinCurve.IsValid())
+	{
+		if (TSharedPtr<FTopologicalEdge>* TwinEdge = AlEdge2CADKernelEdge.Find(TwinCurve->fSpline))
+		{
+			Edge->Link((*TwinEdge).ToSharedRef(), SquareTolerance);
+		}
+		// only TrimCurve with twin need to be in the map
+		AlEdge2CADKernelEdge.Add(AliasTrimCurve.fSpline, Edge);
+	}
+
+	return Edge;
+}
+
+TSharedPtr<FTopologicalLoop> FAliasModelToCADKernelConverter::AddLoop(const AlTrimBoundary& TrimBoundary, TSharedPtr<FSurface>& CarrierSurface)
+{
+	TArray<TSharedPtr<FTopologicalEdge>> Edges;
+	TArray<CADKernel::EOrientation> Directions;
+
+	int32 EdgeCount = 0;
+	for(TUniquePtr<AlTrimCurve> TrimCurve(TrimBoundary.firstCurve()); TrimCurve.IsValid(); TrimCurve = TUniquePtr<AlTrimCurve>(TrimCurve->nextCurve()))
+	{
+		EdgeCount++;
+	}
+
+	Edges.Reserve(EdgeCount);
+	Directions.Reserve(EdgeCount);
+
+	for (TUniquePtr<AlTrimCurve> TrimCurve(TrimBoundary.firstCurve()); TrimCurve.IsValid(); TrimCurve = TUniquePtr<AlTrimCurve>(TrimCurve->nextCurve()))
+	{
+		TSharedPtr<FTopologicalEdge> Edge = AddEdge(*TrimCurve, CarrierSurface);
+		if (Edge.IsValid())
+		{
+			Edges.Add(Edge);
+			Directions.Emplace(CADKernel::EOrientation::Front);
+		}
+	}
+
+	if (Edges.Num() == 0)
+	{
+		return TSharedPtr<FTopologicalLoop>();
+	}
+
+	TSharedPtr<FTopologicalLoop> Loop = FTopologicalLoop::Make(Edges, Directions, GeometricTolerance);
+	return Loop;
+}
+
+TSharedPtr<FTopologicalFace> FAliasModelToCADKernelConverter::AddTrimRegion(const AlTrimRegion& TrimRegion, EAliasObjectReference InObjectReference, const AlMatrix4x4& InAlMatrix, bool bInOrientation)
+{
+	TSharedPtr<FSurface> Surface = AliasToCADKernelUtils::AddNurbsSurface(GeometricTolerance, TrimRegion, InObjectReference, InAlMatrix);
+	if (!Surface.IsValid())
+	{
+		return TSharedPtr<FTopologicalFace>();
+	}
+
+	TArray<TSharedPtr<FTopologicalLoop>> Loops;
+	int32 LoopCount = 0;
+	for (TUniquePtr<AlTrimBoundary> TrimBoundary(TrimRegion.firstBoundary()); TrimBoundary.IsValid(); TrimBoundary = TUniquePtr<AlTrimBoundary>(TrimBoundary->nextBoundary()))
+	{
+		LoopCount++;
+	}
+	Loops.Reserve(LoopCount);
+
+	for (TUniquePtr<AlTrimBoundary> TrimBoundary(TrimRegion.firstBoundary()); TrimBoundary.IsValid(); TrimBoundary = TUniquePtr<AlTrimBoundary>(TrimBoundary->nextBoundary()))
+	{
+		TSharedPtr<FTopologicalLoop> Loop = AddLoop(*TrimBoundary, Surface);
+		if (!Loop.IsValid())
+		{
+			continue;
+		}
+
+		TArray<FPoint2D> LoopSampling;
+		Loop->Get2DSampling(LoopSampling);
+		FAABB2D Boundary;
+		Boundary += LoopSampling;
+		Loop->Boundary.Set(Boundary.GetMin(), Boundary.GetMax());
+
+		// Check if the loop is not composed with only degenerated edge
+		bool bDegeneratedLoop = true;
+		for (const FOrientedEdge& Edge : Loop->GetEdges())
+		{
+			if (!Edge.Entity->IsDegenerated())
+			{
+				bDegeneratedLoop = false;
+				break;
+			}
+		}
+		if (bDegeneratedLoop)
+		{
+			continue;
+		}
+
+		Loops.Add(Loop);
+	}
+
+	if (Loops.Num() == 0)
+	{
+		FMessage::Printf(Log, TEXT("The Face %s is degenerate, this face is ignored\n"), TrimRegion.name());
+		return TSharedPtr<FTopologicalFace>();
+	}
+
+	TSharedRef<FTopologicalFace> Face = FEntity::MakeShared<FTopologicalFace>(Surface);
+	Face->SetPatchId(LastFaceId++);
+
+	Face->AddLoops(Loops);
+	return Face;
+}
+
+void FAliasModelToCADKernelConverter::AddFace(const AlSurface& Surface, EAliasObjectReference InObjectReference, const AlMatrix4x4& InAlMatrix, bool bInOrientation, TSharedRef<FShell>& Shell)
+{
+	TUniquePtr<AlTrimRegion> TrimRegion(Surface.firstTrimRegion());
+	if (TrimRegion.IsValid())
+	{
+		for (; TrimRegion.IsValid(); TrimRegion = TUniquePtr<AlTrimRegion>(TrimRegion->nextRegion()))
+		{
+			TSharedPtr<FTopologicalFace> Face = AddTrimRegion(*TrimRegion, InObjectReference, InAlMatrix, bInOrientation);
+			if (Face.IsValid())
+			{
+				Shell->Add(Face.ToSharedRef(), bInOrientation ? CADKernel::EOrientation::Front : CADKernel::EOrientation::Back);
+			}
+		}
+		return;
+	}
+
+	TSharedPtr<FSurface> CADKernelSurface = AliasToCADKernelUtils::AddNurbsSurface(GeometricTolerance, Surface, InObjectReference, InAlMatrix);
+	if (CADKernelSurface.IsValid())
+	{
+		TSharedRef<FTopologicalFace> Face = FEntity::MakeShared<FTopologicalFace>(CADKernelSurface);
+		Face->ApplyNaturalLoops();
+		Shell->Add(Face, bInOrientation ? CADKernel::EOrientation::Front : CADKernel::EOrientation::Back);
+	}
+}
+
+void FAliasModelToCADKernelConverter::AddShell(const AlShell& InShell, EAliasObjectReference InObjectReference, const AlMatrix4x4& InAlMatrix, bool bInOrientation, TSharedRef<FShell>& CADKernelShell)
+{
+	for(TUniquePtr<AlTrimRegion> TrimRegion(InShell.firstTrimRegion()); TrimRegion.IsValid(); TrimRegion = TUniquePtr<AlTrimRegion>(TrimRegion->nextRegion()))
+	{
+		TSharedPtr<FTopologicalFace> Face = AddTrimRegion(*TrimRegion, InObjectReference, InAlMatrix, bInOrientation);
+		if (Face.IsValid())
+		{
+			CADKernelShell->Add(Face.ToSharedRef(), bInOrientation ? CADKernel::EOrientation::Front : CADKernel::EOrientation::Back);
+		}
+	}
+}
+
+bool FAliasModelToCADKernelConverter::AddBRep(AlDagNode& DagNode, EAliasObjectReference InObjectReference)
+{
+	AlEdge2CADKernelEdge.Empty();
+
+	TSharedRef<FBody> CADKernelBody = FEntity::MakeShared<FBody>();
+	TSharedRef<FShell> CADKernelShell = FEntity::MakeShared<FShell>();
+	CADKernelBody->AddShell(CADKernelShell);
+
+	boolean bAlOrientation;
+	DagNode.getSurfaceOrientation(bAlOrientation);
+	bool bOrientation = (bool)bAlOrientation;
+
+	AlMatrix4x4 AlMatrix;
+	if (InObjectReference == EAliasObjectReference::ParentReference)
+	{
+		DagNode.localTransformationMatrix(AlMatrix);
+	}
+
+	AlObjectType objectType = DagNode.type();
+	switch (objectType)
+	{
+		// Push all leaf nodes into 'leaves'
+	case kShellNodeType:
+	{
+		if (AlShellNode* ShellNode = DagNode.asShellNodePtr())
+		{
+			AlShell* ShellPtr = ShellNode->shell();
+			if(AlIsValid(ShellPtr))
+			{
+				TUniquePtr<AlShell> AliasShell(ShellPtr);
+				AddShell(*ShellPtr, InObjectReference, AlMatrix, bOrientation, CADKernelShell);
+			}
+		}
+		break;
+	}
+
+	case kSurfaceNodeType:
+	{
+		if (AlSurfaceNode* SurfaceNode = DagNode.asSurfaceNodePtr())
+		{
+			AlSurface* SurfacePtr = SurfaceNode->surface();
+			if (AlIsValid(SurfacePtr))
+			{
+				TUniquePtr<AlSurface> AliasSurface(SurfacePtr);
+				AddFace(*SurfacePtr, InObjectReference, AlMatrix, bOrientation, CADKernelShell);
+			}
+		}
+		break;
+	}
+	default:
+		break;
+	}
+
+	if (CADKernelShell->FaceCount() == 0)
+	{
+		return false;
+	}
+
+	// Create body from faces
+	CADKernelSession.GetModel()->Add(CADKernelBody);
+	return true;
+}
+
+
+#endif
\ No newline at end of file
diff --git a/Engine/Plugins/Enterprise/DatasmithCADImporter/Source/DatasmithWireTranslator/Private/AliasModelToCADKernelConverter.h b/Engine/Plugins/Enterprise/DatasmithCADImporter/Source/DatasmithWireTranslator/Private/AliasModelToCADKernelConverter.h
new file mode 100644
--- /dev/null
+++ b/Engine/Plugins/Enterprise/DatasmithCADImporter/Source/DatasmithWireTranslator/Private/AliasModelToCADKernelConverter.h
@@ -0,0 +1,52 @@
+// Copyright Epic Games, Inc. All Rights Reserved.
+
+#pragma once
+
+#include "CoreMinimal.h"
+
+#include "AliasBRepConverter.h"
+#include "CADModelToCADKernelConverterBase.h"
+
+#include "CADKernel/Core/Session.h"
+#include "CADKernel/Core/Types.h"
+
+class AlDagNode;
+class AlShell;
+class AlSurface;
+class AlTrimBoundary;
+class AlTrimCurve;
+class AlTrimRegion;
+
+typedef double AlMatrix4x4[4][4];
+
+namespace CADKernel
+{
+	class FShell;
+	class FSurface;
+	class FTopologicalEdge;
+	class FTopologicalFace;
+	class FTopologicalLoop;
+}
+
+class FAliasModelToCADKernelConverter : public FCADModelToCADKernelConverterBase, public IAliasBRepConverter
+{
+
+public:
+	FAliasModelToCADKernelConverter()
+	{
+	}
+
+	virtual bool AddBRep(AlDagNode& DagNode, EAliasObjectReference ObjectReference) override;
+
+protected:
+	TSharedPtr<CADKernel::FTopologicalEdge> AddEdge(const AlTrimCurve& TrimCurve, TSharedPtr<CADKernel::FSurface>& CarrierSurface);
+	TSharedPtr<CADKernel::FTopologicalLoop> AddLoop(const AlTrimBoundary& TrimBoundary, TSharedPtr<CADKernel::FSurface>& CarrierSurface);
+	TSharedPtr<CADKernel::FTopologicalFace> AddTrimRegion(const AlTrimRegion& InTrimRegion, EAliasObjectReference InObjectReference, const AlMatrix4x4& InAlMatrix, bool bInOrientation);
+	void AddFace(const AlSurface& InSurface, EAliasObjectReference InObjectReference, const AlMatrix4x4& InAlMatrix, bool bInOrientation, TSharedRef<CADKernel::FShell>& Shell);
+	void AddShell(const AlShell& InShell, EAliasObjectReference InObjectReference, const AlMatrix4x4& InAlMatrix, bool bInOrientation, TSharedRef<CADKernel::FShell>& Shell);
+
+protected:
+	int32 LastFaceId = 1;
+	TMap<void*, TSharedPtr<CADKernel::FTopologicalEdge>>  AlEdge2CADKernelEdge;
+};
+
diff --git a/Engine/Plugins/Enterprise/DatasmithCADImporter/Source/DatasmithWireTranslator/Private/DatasmithWireTranslator.cpp b/Engine/Plugins/Enterprise/DatasmithCADImporter/Source/DatasmithWireTranslator/Private/DatasmithWireTranslator.cpp
--- a/Engine/Plugins/Enterprise/DatasmithCADImporter/Source/DatasmithWireTranslator/Private/DatasmithWireTranslator.cpp
+++ b/Engine/Plugins/Enterprise/DatasmithCADImporter/Source/DatasmithWireTranslator/Private/DatasmithWireTranslator.cpp
@@ -8,6 +8,7 @@
 #include "DatasmithSceneFactory.h"
 #include "DatasmithTranslator.h"
 #include "DatasmithUtils.h"
+#include "HAL/ConsoleManager.h"
 #include "DatasmithWireTranslatorModule.h"
 #include "IDatasmithSceneElements.h"
 #include "Misc/FileHelper.h"
@@ -26,6 +27,7 @@
 #endif
 
 #include "AliasModelToCoretechConverter.h" // requires CoreTech as public dependency
+#include "AliasModelToCADKernelConverter.h"
 #include "CADInterfacesModule.h"
 #include "AliasBrepConverter.h"
 #include "CoreTechSurfaceExtension.h"
@@ -155,9 +157,9 @@ public:
 		}
 		else
 		{
-			//TSharedRef<FAliasModelToCADKernelConverter> AliasToCoretechConverter = MakeShared<FAliasModelToCADKernelConverter>();
-			//CADModelConverter = AliasToCoretechConverter;
-			//AliasBRepConverter = AliasToCoretechConverter;
+			TSharedRef<FAliasModelToCADKernelConverter> AliasToCoretechConverter = MakeShared<FAliasModelToCADKernelConverter>();
+			CADModelConverter = AliasToCoretechConverter;
+			AliasBRepConverter = AliasToCoretechConverter;
 		}
 	}
 
@@ -203,7 +205,7 @@ private:
 	void GetDagNodeMeta(const AlDagNode& InDagNode, TSharedPtr<IDatasmithActorElement> OutActorElement);
 
 	TOptional<FMeshDescription> GetMeshOfShellNode(AlDagNode& DagNode, TSharedRef<IDatasmithMeshElement> MeshElement, CADLibrary::FMeshParameters& MeshParameters);
-	TOptional<FMeshDescription> GetMeshOfNodeMesh(AlDagNode& DagNode, TSharedRef<IDatasmithMeshElement> MeshElement, CADLibrary::FMeshParameters& MeshParameters, AlMatrix4x4* AlMeshInvGlobalMatrix = nullptr);
+	TOptional<FMeshDescription> GetMeshOfNodeMesh(AlDagNode& DagNode, CADLibrary::FMeshParameters& MeshParameters, AlMatrix4x4* AlMeshInvGlobalMatrix = nullptr);
 	TOptional<FMeshDescription> GetMeshOfShellBody(TSharedRef<BodyData> DagNode, TSharedRef<IDatasmithMeshElement> MeshElement, CADLibrary::FMeshParameters& MeshParameters);
 	TOptional<FMeshDescription> GetMeshOfMeshBody(TSharedRef<BodyData> DagNode, TSharedRef<IDatasmithMeshElement> MeshElement, CADLibrary::FMeshParameters& MeshParameters);
 
@@ -1842,7 +1844,7 @@ TOptional<FMeshDescription> FWireTranslatorImpl::GetMeshOfShellNode(AlDagNode& D
 		{
 			// Get the meshes from the dag nodes. Note that removing the mesh's DAG.
 			// will also removes the meshes, so we have to do it later.
-			TOptional<FMeshDescription> UEMesh = GetMeshOfNodeMesh(*TesselatedNode, MeshElement, MeshParameters, &AlMatrix);
+			TOptional<FMeshDescription> UEMesh = GetMeshOfNodeMesh(*TesselatedNode, MeshParameters, &AlMatrix);
 			return UEMesh;
 		}
 	}
@@ -1884,7 +1886,7 @@ TOptional<FMeshDescription> FWireTranslatorImpl::GetMeshOfMeshBody(TSharedRef<Bo
 	return MoveTemp(MeshDescription);
 }
 
-TOptional<FMeshDescription> FWireTranslatorImpl::GetMeshOfNodeMesh(AlDagNode& TesselatedNode, TSharedRef<IDatasmithMeshElement> MeshElement, CADLibrary::FMeshParameters& MeshParameters, AlMatrix4x4* AlMeshInvGlobalMatrix)
+TOptional<FMeshDescription> FWireTranslatorImpl::GetMeshOfNodeMesh(AlDagNode& TesselatedNode, CADLibrary::FMeshParameters& MeshParameters, AlMatrix4x4* AlMeshInvGlobalMatrix)
 {
 	AlMeshNode* MeshNode = TesselatedNode.asMeshNodePtr();
 	if (!AlIsValid(MeshNode))
@@ -1999,7 +2001,7 @@ TOptional<FMeshDescription> FWireTranslatorImpl::GetMeshDescription(TSharedRef<I
 
 	case kMeshNodeType:
 	{
-		return GetMeshOfNodeMesh(DagNode, MeshElement, MeshParameters);
+		return GetMeshOfNodeMesh(DagNode, MeshParameters);
 		break;
 	}
 
diff --git a/Engine/Source/Runtime/Datasmith/CADKernel/Private/CADKernel/Topo/TopologicalLoop.cpp b/Engine/Source/Runtime/Datasmith/CADKernel/Private/CADKernel/Topo/TopologicalLoop.cpp
--- a/Engine/Source/Runtime/Datasmith/CADKernel/Private/CADKernel/Topo/TopologicalLoop.cpp
+++ b/Engine/Source/Runtime/Datasmith/CADKernel/Private/CADKernel/Topo/TopologicalLoop.cpp
@@ -23,6 +23,28 @@ TSharedPtr<FTopologicalLoop> FTopologicalLoop::Make(const TArray<TSharedPtr<FTop
 	{
 		OrientedEdge.Entity->SetLoop(Loop);
 	}
+
+	TArray<FPoint2D> LoopSampling;
+	Loop->Get2DSampling(LoopSampling);
+	FAABB2D LoopBoundary;
+	LoopBoundary += LoopSampling;
+	Loop->Boundary.Set(LoopBoundary.GetMin(), LoopBoundary.GetMax());
+
+	// Check if the loop is not composed with only degenerated edge
+	bool bDegeneratedLoop = true;
+	for (const FOrientedEdge& Edge : Loop->GetEdges())
+	{
+		if (!Edge.Entity->IsDegenerated())
+		{
+			bDegeneratedLoop = false;
+			break;
+		}
+	}
+
+	if (bDegeneratedLoop)
+	{
+		TSharedPtr<FTopologicalLoop>();
+	}
 	return Loop;
 }
 
