Issue 97378: Undefined symbol 'CADLibrary::FCoreTechSessionBase::Owner'

diff --git a/Engine/Plugins/Enterprise/DatasmithCADImporter/Source/CADInterfaces/Private/CoreTechTypes.cpp b/Engine/Plugins/Enterprise/DatasmithCADImporter/Source/CADInterfaces/Private/CoreTechTypes.cpp
--- a/Engine/Plugins/Enterprise/DatasmithCADImporter/Source/CADInterfaces/Private/CoreTechTypes.cpp
+++ b/Engine/Plugins/Enterprise/DatasmithCADImporter/Source/CADInterfaces/Private/CoreTechTypes.cpp
@@ -227,6 +227,11 @@ namespace CADLibrary
 		}
 	}
 
+	bool FCoreTechSessionBase::IsCoreTechSessionValid() const
+	{
+		return Owner != nullptr && MainObjectId != 0;
+	}
+
 	FCoreTechSessionBase::~FCoreTechSessionBase()
 	{
 		if (Owner != nullptr)
diff --git a/Engine/Plugins/Enterprise/DatasmithCADImporter/Source/CADInterfaces/Public/CoreTechTypes.h b/Engine/Plugins/Enterprise/DatasmithCADImporter/Source/CADInterfaces/Public/CoreTechTypes.h
--- a/Engine/Plugins/Enterprise/DatasmithCADImporter/Source/CADInterfaces/Public/CoreTechTypes.h
+++ b/Engine/Plugins/Enterprise/DatasmithCADImporter/Source/CADInterfaces/Public/CoreTechTypes.h
@@ -258,7 +258,7 @@ namespace CADLibrary
 		 * Make sure CT is initialized, and a main object is ready.
 		 */
 		FCoreTechSessionBase(const TCHAR* Owner);
-		bool IsSessionValid() { return Owner != nullptr && MainObjectId != 0; }
+		bool IsCoreTechSessionValid() const;
 		virtual ~FCoreTechSessionBase();
 
 	protected:
diff --git a/Engine/Plugins/Enterprise/DatasmithCADImporter/Source/CoreTechSurface/Private/CoreTechSurfaceHelper.cpp b/Engine/Plugins/Enterprise/DatasmithCADImporter/Source/CoreTechSurface/Private/CoreTechSurfaceHelper.cpp
--- a/Engine/Plugins/Enterprise/DatasmithCADImporter/Source/CoreTechSurface/Private/CoreTechSurfaceHelper.cpp
+++ b/Engine/Plugins/Enterprise/DatasmithCADImporter/Source/CoreTechSurface/Private/CoreTechSurfaceHelper.cpp
@@ -42,7 +42,7 @@ namespace CoreTechSurface
 	bool LoadFile(const FString& FileName, const CADLibrary::FImportParameters& ImportParameters, const CADLibrary::FMeshParameters& MeshParameters, FMeshDescription& MeshDescription)
 	{
 		CADLibrary::FCoreTechSessionBase Session(TEXT("CoreTechMeshLoader::LoadFile"));
-		if (!Session.IsSessionValid())
+		if (!Session.IsCoreTechSessionValid())
 		{
 			return false;
 		}
diff --git a/Engine/Plugins/Enterprise/DatasmithCADImporter/Source/CoreTechSurface/Public/CADModelToCoretechConverterBase.h b/Engine/Plugins/Enterprise/DatasmithCADImporter/Source/CoreTechSurface/Public/CADModelToCoretechConverterBase.h
--- a/Engine/Plugins/Enterprise/DatasmithCADImporter/Source/CoreTechSurface/Public/CADModelToCoretechConverterBase.h
+++ b/Engine/Plugins/Enterprise/DatasmithCADImporter/Source/CoreTechSurface/Public/CADModelToCoretechConverterBase.h
@@ -65,7 +65,7 @@ public:
 
 	virtual bool IsSessionValid() override
 	{
-		return FCTSession::IsSessionValid();
+		return FCTSession::IsCoreTechSessionValid();
 	}
 
 	void AddSurfaceDataForMesh(const TCHAR* InFilePath, const CADLibrary::FMeshParameters& InMeshParameters, const FDatasmithTessellationOptions& InTessellationOptions, FDatasmithMeshElementPayload& OutMeshPayload) const
diff --git a/Engine/Plugins/Enterprise/DatasmithCADImporter/Source/DatasmithOpenNurbsTranslator/Private/RhinoCoretechWrapper.cpp b/Engine/Plugins/Enterprise/DatasmithCADImporter/Source/DatasmithOpenNurbsTranslator/Private/RhinoCoretechWrapper.cpp
--- a/Engine/Plugins/Enterprise/DatasmithCADImporter/Source/DatasmithOpenNurbsTranslator/Private/RhinoCoretechWrapper.cpp
+++ b/Engine/Plugins/Enterprise/DatasmithCADImporter/Source/DatasmithOpenNurbsTranslator/Private/RhinoCoretechWrapper.cpp
@@ -437,7 +437,7 @@ uint64 BRepToKernelIOBodyTranslator::CreateBody(const ON_3dVector& Offset)
 
 bool FRhinoCoretechWrapper::AddBRep(ON_Brep& Brep, const ON_3dVector& Offset)
 {
-	if (!IsSessionValid())
+	if (!IsCoreTechSessionValid())
 	{
 		ensureMsgf(false, TEXT("bad session init"));
 		return false;
