Pluggin of CADKernel in Rhino

diff --git a/Engine/Plugins/Enterprise/DatasmithCADImporter/Source/CADInterfaces/Private/CTSession.cpp b/Engine/Plugins/Enterprise/DatasmithCADImporter/Source/CADInterfaces/Private/CTSession.cpp
--- a/Engine/Plugins/Enterprise/DatasmithCADImporter/Source/CADInterfaces/Private/CTSession.cpp
+++ b/Engine/Plugins/Enterprise/DatasmithCADImporter/Source/CADInterfaces/Private/CTSession.cpp
@@ -27,11 +27,11 @@ namespace CADLibrary
 
 	void FCTSession::SetSceneUnit(double InMetricUnit)
 	{
-		ImportParams.MetricUnit = InMetricUnit;
+		ImportParams.SetMetricUnit(InMetricUnit);
 		CTKIO_ChangeUnit(InMetricUnit);
 	}
 
-	void FCTSession::SetImportParameters(float ChordTolerance, float MaxEdgeLength, float NormalTolerance, CADLibrary::EStitchingTechnique StitchingTechnique, bool bScaleUVMap)
+	void FCTSession::SetImportParameters(double ChordTolerance, double MaxEdgeLength, double NormalTolerance, CADLibrary::EStitchingTechnique StitchingTechnique, bool bScaleUVMap)
 	{
 		ImportParams.ChordTolerance = ChordTolerance;
 		ImportParams.MaxEdgeLength = MaxEdgeLength;
diff --git a/Engine/Plugins/Enterprise/DatasmithCADImporter/Source/CADInterfaces/Public/CTSession.h b/Engine/Plugins/Enterprise/DatasmithCADImporter/Source/CADInterfaces/Public/CTSession.h
--- a/Engine/Plugins/Enterprise/DatasmithCADImporter/Source/CADInterfaces/Public/CTSession.h
+++ b/Engine/Plugins/Enterprise/DatasmithCADImporter/Source/CADInterfaces/Public/CTSession.h
@@ -21,16 +21,17 @@ public:
 	 * Handle input file unit and an output unit
 	 * @param InOwner:        text that describe the owner of the session (helps to fix initialization issues)
 	 */
-	FCTSession(const TCHAR* InOwner)
+	FCTSession(const TCHAR* InOwner, const CADLibrary::FImportParameters& InImportParameters)
 		: FCoreTechSessionBase(InOwner)
+		, ImportParams(InImportParameters)
 	{
+		ImportParams.bEnableKernelIOTessellation = true;
 	}
 
 	void ClearData();
 
 	bool SaveBrep(const FString& FilePath);
 
-
 	/**
  	 * This function calls, according to the chosen EStitchingTechnique, Kernel_io CT_REPAIR_IO::Sew or CT_REPAIR_IO::Heal. In case of sew, the used tolerance is 100x the geometric tolerance (SewingToleranceFactor = 100). 
 	 * With the case of UE-83379, Alias file, this value is too big (biggest than the geometric features. So Kernel_io hangs during the sew process... In the wait of more test, 100x is still the value used for CAD import except for Alias where the value of the SewingToleranceFactor is set to 1x
@@ -46,13 +47,23 @@ public:
 		ImportParams.ScaleFactor = InScaleFactor;
 	}
 
+	double GetScaleFactor() const
+	{
+		return ImportParams.ScaleFactor;
+	}
+
+	double GetSceneUnit() const
+	{
+		return ImportParams.MetricUnit;
+	}
+
 	/**
 	 * Handle input file unit
 	 * @param FileMetricUnit: number of meters per file unit.
 	 * eg. For a file in inches, arg should be 0.0254
 	 */
 	void SetSceneUnit(double InMetricUnit);
-	
+
 	/**
 	 * Set Import parameters,
 	 * Tack care to set scale factor before because import parameters will be scale according to scale factor
@@ -61,7 +72,7 @@ public:
 	 * @param NormalTolerance : Angle between two adjacent triangles
 	 * @param StitchingTechnique : CAD topology correction technique
 	 */
-	void SetImportParameters(float ChordTolerance, float MaxEdgeLength, float NormalTolerance, CADLibrary::EStitchingTechnique StitchingTechnique, bool bScaleUVMap);
+	void SetImportParameters(double ChordTolerance, double MaxEdgeLength, double NormalTolerance, CADLibrary::EStitchingTechnique StitchingTechnique, bool bScaleUVMap);
 	
 	void SetModelCoordinateSystem(FDatasmithUtils::EModelCoordSystem NewCoordinateSystem)
 	{
diff --git a/Engine/Plugins/Enterprise/DatasmithCADImporter/Source/CADKernelSurface/Public/CADModelToCADKernelConverterBase.h b/Engine/Plugins/Enterprise/DatasmithCADImporter/Source/CADKernelSurface/Public/CADModelToCADKernelConverterBase.h
--- a/Engine/Plugins/Enterprise/DatasmithCADImporter/Source/CADKernelSurface/Public/CADModelToCADKernelConverterBase.h
+++ b/Engine/Plugins/Enterprise/DatasmithCADImporter/Source/CADKernelSurface/Public/CADModelToCADKernelConverterBase.h
@@ -20,14 +20,12 @@ class FCADModelToCADKernelConverterBase : public CADLibrary::ICADModelConverter
 {
 public:
 
-	FCADModelToCADKernelConverterBase()
-		// Unit is set to cm, 0.01, because Wire's unit is cm.
-		: CADKernelSession(0.00001 / 0.01)
-		, GeometricTolerance(0.00001 / 0.01)
+	FCADModelToCADKernelConverterBase(CADLibrary::FImportParameters InImportParameters)
+		: CADKernelSession(0.00001 / InImportParameters.MetricUnit)
+		, ImportParameters(InImportParameters)
+		, GeometricTolerance(0.00001 / InImportParameters.MetricUnit)
 		, SquareTolerance(GeometricTolerance* GeometricTolerance)
 	{
-		ImportParameters.MetricUnit = 0.01;
-		ImportParameters.ScaleFactor = 1;
 		ImportParameters.bEnableKernelIOTessellation = false;
 	}
 
@@ -56,14 +54,28 @@ public:
 		return CADLibrary::FCADKernelTools::Tessellate(CADKernelEntity, ImportParameters, InMeshParameters, OutMeshDescription);
 	}
 
-	virtual void SetImportParameters(float ChordTolerance, float MaxEdgeLength, float NormalTolerance, CADLibrary::EStitchingTechnique StitchingTechnique, bool bScaleUVMap) override
+	virtual void SetImportParameters(double ChordTolerance, double MaxEdgeLength, double NormalTolerance, CADLibrary::EStitchingTechnique StitchingTechnique, bool bScaleUVMap) override
 	{
 		ImportParameters.ChordTolerance = ChordTolerance;
 		ImportParameters.MaxEdgeLength = MaxEdgeLength;
 		ImportParameters.MaxNormalAngle = NormalTolerance;
 		ImportParameters.StitchingTechnique = StitchingTechnique;
 		ImportParameters.bScaleUVMap = bScaleUVMap;
-		ImportParameters.bEnableKernelIOTessellation = false;
+	}
+
+	virtual void SetMetricUnit(double NewMetricUnit) override
+	{
+		ImportParameters.SetMetricUnit(NewMetricUnit);
+	}
+
+	virtual double GetScaleFactor() const override
+	{
+		return ImportParameters.ScaleFactor;
+	}
+
+	virtual double GetMetricUnit() const override
+	{
+		return ImportParameters.MetricUnit;
 	}
 
 	virtual bool IsSessionValid() override
diff --git a/Engine/Plugins/Enterprise/DatasmithCADImporter/Source/CADLibrary/Private/CADKernelTools.cpp b/Engine/Plugins/Enterprise/DatasmithCADImporter/Source/CADLibrary/Private/CADKernelTools.cpp
--- a/Engine/Plugins/Enterprise/DatasmithCADImporter/Source/CADLibrary/Private/CADKernelTools.cpp
+++ b/Engine/Plugins/Enterprise/DatasmithCADImporter/Source/CADLibrary/Private/CADKernelTools.cpp
@@ -389,13 +389,14 @@ namespace CADLibrary
 
 		if (ImportParameters.MaxEdgeLength > SMALL_NUMBER)
 		{
-			TSharedPtr<CADKernel::FCriterion> MaxSizeCriterion = CADKernel::FCriterion::CreateCriterion(CADKernel::ECriterion::MaxSize, ImportParameters.MaxEdgeLength);
+
+			TSharedPtr<CADKernel::FCriterion> MaxSizeCriterion = CADKernel::FCriterion::CreateCriterion(CADKernel::ECriterion::MaxSize, ImportParameters.MaxEdgeLength / ImportParameters.ScaleFactor);
 			MeshModel->AddCriterion(MaxSizeCriterion);
 		}
 
 		if (ImportParameters.ChordTolerance > SMALL_NUMBER)
 		{
-			TSharedPtr<CADKernel::FCriterion> ChordCriterion = CADKernel::FCriterion::CreateCriterion(CADKernel::ECriterion::Sag, ImportParameters.ChordTolerance);
+			TSharedPtr<CADKernel::FCriterion> ChordCriterion = CADKernel::FCriterion::CreateCriterion(CADKernel::ECriterion::Sag, ImportParameters.ChordTolerance / ImportParameters.ScaleFactor);
 			MeshModel->AddCriterion(ChordCriterion);
 		}
 
diff --git a/Engine/Plugins/Enterprise/DatasmithCADImporter/Source/CADTools/Public/CADOptions.h b/Engine/Plugins/Enterprise/DatasmithCADImporter/Source/CADTools/Public/CADOptions.h
--- a/Engine/Plugins/Enterprise/DatasmithCADImporter/Source/CADTools/Public/CADOptions.h
+++ b/Engine/Plugins/Enterprise/DatasmithCADImporter/Source/CADTools/Public/CADOptions.h
@@ -47,6 +47,12 @@ namespace CADLibrary
 		bool bEnableKernelIOTessellation = true;
 		bool bEnableTimeControl = true;
 
+		void SetMetricUnit(double NewMetricUnit)
+		{
+			MetricUnit = NewMetricUnit;
+			ScaleFactor = NewMetricUnit / 0.01;
+		}
+
 		uint32 GetHash() const
 		{
 			uint32 Hash = 0; 
diff --git a/Engine/Plugins/Enterprise/DatasmithCADImporter/Source/CoreTechSurface/CoreTechSurface.Build.cs b/Engine/Plugins/Enterprise/DatasmithCADImporter/Source/CoreTechSurface/CoreTechSurface.Build.cs
--- a/Engine/Plugins/Enterprise/DatasmithCADImporter/Source/CoreTechSurface/CoreTechSurface.Build.cs
+++ b/Engine/Plugins/Enterprise/DatasmithCADImporter/Source/CoreTechSurface/CoreTechSurface.Build.cs
@@ -8,6 +8,11 @@ namespace UnrealBuildTool.Rules
 	{
 		public CoreTechSurface(ReadOnlyTargetRules Target) : base(Target)
 		{
+			// option de compil pour dev
+			OptimizeCode = CodeOptimization.Never;
+			PCHUsage = PCHUsageMode.NoPCHs;
+			// fin option de compil pour dev
+
 			PublicDependencyModuleNames.AddRange(
 				new string[]
 				{
diff --git a/Engine/Plugins/Enterprise/DatasmithCADImporter/Source/CoreTechSurface/Public/CADModelToCoretechConverterBase.h b/Engine/Plugins/Enterprise/DatasmithCADImporter/Source/CoreTechSurface/Public/CADModelToCoretechConverterBase.h
--- a/Engine/Plugins/Enterprise/DatasmithCADImporter/Source/CoreTechSurface/Public/CADModelToCoretechConverterBase.h
+++ b/Engine/Plugins/Enterprise/DatasmithCADImporter/Source/CoreTechSurface/Public/CADModelToCoretechConverterBase.h
@@ -14,13 +14,9 @@ class FCADModelToCoretechConverterBase : public CADLibrary::FCTSession, public C
 {
 public:
 
-	FCADModelToCoretechConverterBase(const TCHAR* InOwner)
-		: CADLibrary::FCTSession(InOwner)
+	FCADModelToCoretechConverterBase(const TCHAR* InOwner, const CADLibrary::FImportParameters& InImportParameters)
+		: CADLibrary::FCTSession(InOwner, InImportParameters)
 	{
-		// Unit for CoreTech session is set to cm, 0.01, because Wire's unit is cm. Consequently, Scale factor is set to 1.
-		ImportParams.MetricUnit = 0.01;
-		ImportParams.ScaleFactor = 1;
-		ImportParams.bEnableKernelIOTessellation = true;
 	}
 
 	virtual bool Tessellate(const CADLibrary::FMeshParameters& InMeshParameters, FMeshDescription& OutMeshDescription) override
@@ -57,10 +53,24 @@ public:
 		return false;
 	}
 
-	virtual void SetImportParameters(float ChordTolerance, float MaxEdgeLength, float NormalTolerance, CADLibrary::EStitchingTechnique StitchingTechnique, bool bScaleUVMap) override
+	virtual void SetImportParameters(double ChordTolerance, double MaxEdgeLength, double NormalTolerance, CADLibrary::EStitchingTechnique StitchingTechnique, bool bScaleUVMap) override
 	{
 		FCTSession::SetImportParameters(ChordTolerance, MaxEdgeLength, NormalTolerance, StitchingTechnique, bScaleUVMap);
-		ImportParams.bEnableKernelIOTessellation = true;
+	}
+
+	virtual void SetMetricUnit(double NewMetricUnit) override
+	{
+		ImportParams.SetMetricUnit(NewMetricUnit);
+	}
+
+	virtual double GetScaleFactor() const override
+	{
+		return ImportParams.ScaleFactor;
+	}
+
+	virtual double GetMetricUnit() const override
+	{
+		return ImportParams.MetricUnit;
 	}
 
 	virtual bool IsSessionValid() override
diff --git a/Engine/Plugins/Enterprise/DatasmithCADImporter/Source/DatasmithOpenNurbsTranslator/DatasmithOpenNurbsTranslator.Build.cs b/Engine/Plugins/Enterprise/DatasmithCADImporter/Source/DatasmithOpenNurbsTranslator/DatasmithOpenNurbsTranslator.Build.cs
--- a/Engine/Plugins/Enterprise/DatasmithCADImporter/Source/DatasmithOpenNurbsTranslator/DatasmithOpenNurbsTranslator.Build.cs
+++ b/Engine/Plugins/Enterprise/DatasmithCADImporter/Source/DatasmithOpenNurbsTranslator/DatasmithOpenNurbsTranslator.Build.cs
@@ -20,9 +20,11 @@ namespace UnrealBuildTool.Rules
 			PrivateDependencyModuleNames.AddRange(
 				new string[]
 				{
-                    "CADInterfaces",
-                    "CADLibrary",
-                    "CoreTechSurface",
+					"CADInterfaces",
+					"CADKernel",
+					"CADKernelSurface",
+					"CADLibrary",
+					"CoreTechSurface",
 					"DatasmithCore",
 					"DatasmithContent",
 					"DatasmithTranslator",
diff --git a/Engine/Plugins/Enterprise/DatasmithCADImporter/Source/DatasmithOpenNurbsTranslator/Private/DatasmithOpenNurbsTranslator.cpp b/Engine/Plugins/Enterprise/DatasmithCADImporter/Source/DatasmithOpenNurbsTranslator/Private/DatasmithOpenNurbsTranslator.cpp
--- a/Engine/Plugins/Enterprise/DatasmithCADImporter/Source/DatasmithOpenNurbsTranslator/Private/DatasmithOpenNurbsTranslator.cpp
+++ b/Engine/Plugins/Enterprise/DatasmithCADImporter/Source/DatasmithOpenNurbsTranslator/Private/DatasmithOpenNurbsTranslator.cpp
@@ -5,9 +5,9 @@
 
 #ifdef USE_OPENNURBS // The whole translation unit is skipped without OpenNurbs TPS library
 
-#include "RhinoCoretechWrapper.h"
-
 #include "CADInterfacesModule.h"
+#include "CADKernelSurfaceExtension.h"
+#include "CADModelConverter.h"
 #include "CoreTechSurfaceExtension.h"
 #include "CoreTechSurfaceHelper.h"
 #include "DatasmithImportOptions.h"
@@ -17,6 +17,9 @@
 #include "DatasmithSceneFactory.h"
 #include "DatasmithSceneSource.h"
 #include "DatasmithUtils.h"
+#include "OpenNurbsBRepConverter.h"
+#include "OpenNurbsBRepToCADKernelConverter.h"
+#include "OpenNurbsBRepToCoretechConverter.h" // requires CoreTech as public dependency
 #include "Utility/DatasmithMeshHelper.h"
 
 #if WITH_EDITOR
@@ -160,7 +163,7 @@ namespace DatasmithOpenNurbsTranslatorUtils
 		}
 	}
 
-	bool CheckForCRCErrors(ON_BinaryArchive& archive, /*ON_TextLog* error_log,*/ const char* sSection, int& NumCRCErrors)
+	bool CheckForCRCErrors(ON_BinaryArchive& archive, const char* sSection, int& NumCRCErrors)
 	{
 		// returns true if new CRC errors are found
 		bool bHasCRCErrors = false;
@@ -168,11 +171,6 @@ namespace DatasmithOpenNurbsTranslatorUtils
 
 		if (NumCRCErrors != CRCCount)
 		{
-			//if (error_log)
-			//{
-			//	error_log->Print("ERROR: Corrupt %s. (CRC errors).\n", sSection);
-			//	error_log->Print("-- Attempting to continue.\n");
-			//}
 			NumCRCErrors = CRCCount;
 			bHasCRCErrors = true;
 		}
@@ -640,8 +638,6 @@ public:
 		, FileVersion(0)
 		, ArchiveOpenNurbsVersion(0)
 		, FileLength(0)
-		, MetricUnit(0.001)
-		, ScalingFactor(0.1)
 		, NumCRCErrors(0)
 	{
 		if (!TranslationCache.IsValid())
@@ -649,13 +645,33 @@ public:
 			TranslationCache = MakeShared<FTranslationCache>();
 		}
 
-		LocalSession = FRhinoCoretechWrapper::GetSharedSession();
+		bool bEnableKernelIOTessellation = true;
+		if (IConsoleVariable* CVar = IConsoleManager::Get().FindConsoleVariable(TEXT("ds.CADTranslator.EnableKernelIOTessellation")))
+		{
+			bEnableKernelIOTessellation = CVar->GetInt() != 0;
+		}
+
+		CADLibrary::FImportParameters ImportParameters;
+		ImportParameters.MetricUnit = 0.001;
+		ImportParameters.ScaleFactor = 0.1;
+		ImportParameters.ModelCoordSys = FDatasmithUtils::EModelCoordSystem::ZUp_RightHanded_FBXLegacy;
+
+		if (bEnableKernelIOTessellation)
+		{
+			TSharedRef<FOpenNurbsBRepToCoretechConverter> OpenNurbsBRepToCoretechConverter = MakeShared<FOpenNurbsBRepToCoretechConverter>(TEXT("Al2CTSharedSession"), ImportParameters);
+			CADModelConverter = OpenNurbsBRepToCoretechConverter;
+			OpenNurbsBRepConverter = OpenNurbsBRepToCoretechConverter;
+		}
+		else
+		{
+			TSharedRef<FOpenNurbsBRepToCADKernelConverter> OpenNurbsBRepToCADKernelConverter = MakeShared<FOpenNurbsBRepToCADKernelConverter>(ImportParameters);
+			CADModelConverter = OpenNurbsBRepToCADKernelConverter;
+			OpenNurbsBRepConverter = OpenNurbsBRepToCADKernelConverter;
+		}
 	}
 
 	~FOpenNurbsTranslatorImpl()
 	{
-		LocalSession.Reset();
-
 		for (FOpenNurbsTranslatorImpl* ChildTranslator : ChildTranslators)
 		{
 			delete ChildTranslator;
@@ -669,11 +685,11 @@ public:
 	void SetBaseOptions(const FDatasmithImportBaseOptions& InBaseOptions);
 	void SetOpenNurbsOptions(const FDatasmithOpenNurbsOptions& Options);
 	void SetOutputPath(const FString& Path) { OutputPath = Path; }
-	double GetScalingFactor() const { return ScalingFactor; }
-	double GetMetricUnit() const { return MetricUnit; }
-
+	
 	void ShowMessageLog(const FString& Filename);
 
+	bool LoadStaticMesh(const TSharedRef<IDatasmithMeshElement> MeshElement, FDatasmithMeshElementPayload& OutMeshPayload, const FDatasmithTessellationOptions& InTessellationOptions);
+
 private:
 	void TranslateTextureMappingTable(const ON_ObjectArray<ON_TextureMapping>& TextureMappingTable);
 	void TranslateMaterialTable(const ON_ObjectArray<ON_Material>& MaterialTable);
@@ -712,7 +728,7 @@ private:
 
 	void SetMaterialToMeshElement(TSharedPtr<IDatasmithMeshElement> MeshElement, TSharedPtr<IDatasmithBaseMaterialElement> MaterialElement, int32 SlotId);
 
-	bool TranslateBRep(ON_Brep* brep, const ON_3dmObjectAttributes& Attributes, FMeshDescription& OutMesh, const TSharedRef< IDatasmithMeshElement >& MeshElement, const FString& Name, bool& bHasNormal);
+	bool TranslateBRep(ON_Brep* brep, const ON_3dmObjectAttributes& Attributes, FMeshDescription& OutMesh, TSharedRef< IDatasmithMeshElement >& MeshElement, const FString& Name, bool& bHasNormal);
 
 
 	bool ComputeObjectGeometryCenter(const FOpenNurbsObjectWrapper& Object, ON_3dVector& OutGeometryCenter);
@@ -730,6 +746,11 @@ private:
 		return -Offset;
 	}
 
+	double GetScaleFactor() const
+	{
+		return CADModelConverter->GetScaleFactor();
+	}
+
 private:
 	TArray<FOpenNurbsTranslatorImpl*> ChildTranslators;
 	TSharedRef<IDatasmithScene> Scene;
@@ -741,8 +762,6 @@ private:
 	uint32 OpenNurbsOptionsHash;
 	FDatasmithImportBaseOptions BaseOptions;
 
-	TSharedPtr<FRhinoCoretechWrapper> LocalSession;
-
 private:
 	// For OpenNurbs archive parsing
 	// Ref. GP3DMReader
@@ -773,9 +792,6 @@ private:
 	// length of archive returned by ON_BinaryArchive::Read3dmEndMark()
 	size_t FileLength;
 
-	double MetricUnit; // meters per file unit (m/u) eg 0.0254 for files in inches
-	double ScalingFactor;
-
 	// Number of crc errors found during archive reading.
 	// If > 0, then the archive is corrupt.
 	int NumCRCErrors;
@@ -821,6 +837,10 @@ private:
 	TMap< TSharedPtr< IDatasmithMeshElement >, ON_3dVector > MeshElementToGeometryCenter;
 
 	TArray<FString> MissingRenderMeshes;
+
+	TSharedPtr<CADLibrary::ICADModelConverter> CADModelConverter;
+	TSharedPtr<IOpenNurbsBRepConverter> OpenNurbsBRepConverter;
+
 };
 
 void FOpenNurbsTranslatorImpl::ShowMessageLog(const FString& Filename)
@@ -1299,7 +1319,7 @@ void FOpenNurbsTranslatorImpl::TranslateLightTable(const ON_ClassArray<FOpenNurb
 			LightType == EDatasmithElementType::SpotLight)
 		{
 			FVector Location(LightObj.Location().x, LightObj.Location().y, LightObj.Location().z);
-			Location *= ScalingFactor;
+			Location *= GetScaleFactor();
 			Location = FDatasmithUtils::ConvertVector(FDatasmithUtils::EModelCoordSystem::ZUp_RightHanded_FBXLegacy, Location);
 			LightElement->SetTranslation(Location);
 		}
@@ -1318,14 +1338,14 @@ void FOpenNurbsTranslatorImpl::TranslateLightTable(const ON_ClassArray<FOpenNurb
 		{
 			TSharedRef<IDatasmithAreaLightElement> AreaLightElement = StaticCastSharedRef<IDatasmithAreaLightElement>(LightElement.ToSharedRef());
 
-			double Length = LightObj.Length().Length() * ScalingFactor;
+			double Length = LightObj.Length().Length() * GetScaleFactor();
 			AreaLightElement->SetLength(Length);
 
 			ON_3dVector Center = LightObj.Location() + 0.5 * LightObj.Length();
 			if (LightStyle == ON::light_style::world_rectangular_light)
 			{
 				Center += 0.5 * LightObj.Width();
-				double Width = LightObj.Width().Length() * ScalingFactor;
+				double Width = LightObj.Width().Length() * GetScaleFactor();
 
 				AreaLightElement->SetWidth(Width);
 				AreaLightElement->SetLightShape(EDatasmithLightShape::Rectangle);
@@ -1352,7 +1372,7 @@ void FOpenNurbsTranslatorImpl::TranslateLightTable(const ON_ClassArray<FOpenNurb
 			FTransform Transform(Matrix);
 			FTransform CorrectedTransform = FDatasmithUtils::ConvertTransform(FDatasmithUtils::EModelCoordSystem::ZUp_RightHanded_FBXLegacy, Transform);
 
-			AreaLightElement->SetTranslation(CorrectedTransform.GetTranslation() * ScalingFactor);
+			AreaLightElement->SetTranslation(CorrectedTransform.GetTranslation() * GetScaleFactor());
 			AreaLightElement->SetScale(CorrectedTransform.GetScale3D());
 			AreaLightElement->SetRotation(CorrectedTransform.GetRotation());
 		}
@@ -1369,8 +1389,8 @@ void FOpenNurbsTranslatorImpl::TranslateLightTable(const ON_ClassArray<FOpenNurb
 			double OuterRadius;
 			LightObj.GetSpotLightRadii(&InnerRadius, &OuterRadius);
 
-			InnerRadius *= ScalingFactor;
-			OuterRadius *= ScalingFactor;
+			InnerRadius *= GetScaleFactor();
+			OuterRadius *= GetScaleFactor();
 
 			float OuterAngle = FMath::Atan(OuterRadius * FMath::Tan(FMath::DegreesToRadians(InnerAngleDegree)) / InnerRadius);
 			SpotLightElement->SetOuterConeAngle(FMath::RadiansToDegrees(OuterAngle));
@@ -1865,7 +1885,7 @@ bool FOpenNurbsTranslatorImpl::TranslateInstance(const FOpenNurbsObjectWrapper&
 	FTransform Transform(Matrix);
 	FTransform CorrectedTransform = FDatasmithUtils::ConvertTransform(FDatasmithUtils::EModelCoordSystem::ZUp_RightHanded_FBXLegacy, Transform);
 
-	ContainerElement->SetTranslation(CorrectedTransform.GetTranslation() * ScalingFactor);
+	ContainerElement->SetTranslation(CorrectedTransform.GetTranslation() * GetScaleFactor());
 	ContainerElement->SetScale(CorrectedTransform.GetScale3D());
 	ContainerElement->SetRotation(CorrectedTransform.GetRotation());
 	bool bIsPartOfInstanceDefinition = false;
@@ -2037,7 +2057,7 @@ TSharedPtr<IDatasmithMeshActorElement> FOpenNurbsTranslatorImpl::GetMeshActorEle
 	if (ComputeObjectGeometryCenter(Object, GeometryCenter))
 	{
 		MeshElementToGeometryCenter.Add(MeshElement.ToSharedRef(), GeometryCenter);
-		FVector ActorOffset = ScalingFactor * FDatasmithUtils::ConvertVector(FDatasmithUtils::EModelCoordSystem::ZUp_RightHanded_FBXLegacy, GeometryCenter);
+		FVector ActorOffset = GetScaleFactor() * FDatasmithUtils::ConvertVector(FDatasmithUtils::EModelCoordSystem::ZUp_RightHanded_FBXLegacy, GeometryCenter);
 		ActorElement->SetTranslation(ActorOffset);
 	}
 
@@ -2083,7 +2103,7 @@ TSharedPtr<IDatasmithActorElement> FOpenNurbsTranslatorImpl::GetPointActorElemen
 	}
 
 	FVector Location(pointObj->point.x, pointObj->point.y, pointObj->point.z);
-	Location *= ScalingFactor;
+	Location *= GetScaleFactor();
 	Location = FDatasmithUtils::ConvertVector(FDatasmithUtils::EModelCoordSystem::ZUp_RightHanded_FBXLegacy, Location);
 
 	ActorElement->SetTranslation(Location);
@@ -2286,7 +2306,6 @@ bool FOpenNurbsTranslatorImpl::Read(ON_BinaryFile& Archive, TSharedRef<IDatasmit
 	// Step 1: REQUIRED - Read Start Section
 	if (!Archive.Read3dmStartSection(&FileVersion, StartSectionComments))
 	{
-		//if (error_log) error_log->Print("ERROR: Unable to read start section. (ON_BinaryArchive::Read3dmStartSection() returned false.)\n");
 		return false;
 	}
 	else if (DatasmithOpenNurbsTranslatorUtils::CheckForCRCErrors(Archive, nullptr, NumCRCErrors))
@@ -2297,7 +2316,6 @@ bool FOpenNurbsTranslatorImpl::Read(ON_BinaryFile& Archive, TSharedRef<IDatasmit
 	// Step 2: REQUIRED - Read properties table
 	if (!Archive.Read3dmProperties(Properties))
 	{
-		//if (error_log) error_log->Print("ERROR: Unable to read properties section. (ON_BinaryArchive::Read3dmProperties() returned false.)\n");
 		return false;
 	}
 	else if (DatasmithOpenNurbsTranslatorUtils::CheckForCRCErrors(Archive, "properties section", NumCRCErrors))
@@ -2317,7 +2335,6 @@ bool FOpenNurbsTranslatorImpl::Read(ON_BinaryFile& Archive, TSharedRef<IDatasmit
 	// Step 3: REQUIRED - Read settings table
 	if (!Archive.Read3dmSettings(Settings))
 	{
-		//if (error_log) error_log->Print("ERROR: Unable to read settings section. (ON_BinaryArchive::Read3dmSettings() returned false.)\n");
 		return false;
 	}
 	else if (DatasmithOpenNurbsTranslatorUtils::CheckForCRCErrors(Archive, "settings section", NumCRCErrors))
@@ -2325,10 +2342,8 @@ bool FOpenNurbsTranslatorImpl::Read(ON_BinaryFile& Archive, TSharedRef<IDatasmit
 		bResult = false;
 	}
 
-	// 	ScalingFactor is defined according to input Rhino file unit. CT don't modify CAD input according to the set file unit.
-	ScalingFactor = 100 / Settings.m_ModelUnitsAndTolerances.Scale(ON::LengthUnitSystem::Meters);
-	MetricUnit = 1 / Settings.m_ModelUnitsAndTolerances.Scale(ON::LengthUnitSystem::Meters);
-	LocalSession->SetScaleFactor(ScalingFactor);
+	double MetricUnit = 1 / Settings.m_ModelUnitsAndTolerances.Scale(ON::LengthUnitSystem::Meters);
+	CADModelConverter->SetMetricUnit(MetricUnit);
 
 	// Step 4: REQUIRED - Read bitmap table (it can be empty)
 	int Count = 0;
@@ -2350,23 +2365,13 @@ bool FOpenNurbsTranslatorImpl::Read(ON_BinaryFile& Archive, TSharedRef<IDatasmit
 			}
 			if (ReturnCode < 0)
 			{
-				//if (error_log)
-				//{
-				//	error_log->Print("ERROR: Corrupt bitmap found. (ON_BinaryArchive::Read3dmBitmap() < 0.)\n");
-				//	error_count++;
-				//	if (error_count > max_error_count)
-				//		return false;
-				//	error_log->Print("-- Attempting to continue.\n");
-				//}
 				bResult = false;
 			}
-			//BitmapTable.Append(pBitmap);
 		}
 
 		// EndRead3dmBitmapTable() must be called when BeginRead3dmBitmapTable() returns true
 		if (!Archive.EndRead3dmBitmapTable())
 		{
-			//if (error_log) error_log->Print("ERROR: Corrupt bitmap table. (ON_BinaryArchive::EndRead3dmBitmapTable() returned false.)\n");
 			return false;
 		}
 		if (DatasmithOpenNurbsTranslatorUtils::CheckForCRCErrors(Archive, "bitmap table", NumCRCErrors))
@@ -2376,11 +2381,6 @@ bool FOpenNurbsTranslatorImpl::Read(ON_BinaryFile& Archive, TSharedRef<IDatasmit
 	}
 	else
 	{
-		//if (error_log)
-		//{
-		//	error_log->Print("WARNING: Missing or corrupt bitmap table. (ON_BinaryArchive::BeginRead3dmBitmapTable() returned false.)\n");
-		//	error_log->Print("-- Attempting to continue.\n");
-		//}
 		bResult = false;
 	}
 
@@ -2399,14 +2399,6 @@ bool FOpenNurbsTranslatorImpl::Read(ON_BinaryFile& Archive, TSharedRef<IDatasmit
 			}
 			if (ReturnCode < 0)
 			{
-				//if (error_log)
-				//{
-				//	error_log->Print("ERROR: Corrupt render texture_mapping found. (ON_BinaryArchive::Read3dmTextureMapping() < 0.)\n");
-				//	error_count++;
-				//	if (error_count > max_error_count)
-				//		return false;
-				//	error_log->Print("-- Attempting to continue.\n");
-				//}
 				continue;
 			}
 			ON_UserDataHolder ud;
@@ -2421,7 +2413,6 @@ bool FOpenNurbsTranslatorImpl::Read(ON_BinaryFile& Archive, TSharedRef<IDatasmit
 		// EndRead3dmTextureMappingTable() must be called when BeginRead3dmTextureMappingTable() returns true.
 		if (!Archive.EndRead3dmTextureMappingTable())
 		{
-			//if (error_log) error_log->Print("ERROR: Corrupt render texture_mapping table. (ON_BinaryArchive::EndRead3dmTextureMappingTable() returned false.)\n");
 			return false;
 		}
 		if (DatasmithOpenNurbsTranslatorUtils::CheckForCRCErrors(Archive, "render texture_mapping table", NumCRCErrors))
@@ -2431,11 +2422,6 @@ bool FOpenNurbsTranslatorImpl::Read(ON_BinaryFile& Archive, TSharedRef<IDatasmit
 	}
 	else
 	{
-		//if (error_log)
-		//{
-		//	error_log->Print("WARNING: Missing or corrupt render texture_mapping table. (ON_BinaryArchive::BeginRead3dmTextureMappingTable() returned false.)\n");
-		//	error_log->Print("-- Attempting to continue.\n");
-		//}
 		bResult = false;
 	}
 
@@ -2454,14 +2440,6 @@ bool FOpenNurbsTranslatorImpl::Read(ON_BinaryFile& Archive, TSharedRef<IDatasmit
 			}
 			if (ReturnCode < 0)
 			{
-				//if (error_log)
-				//{
-				//	error_log->Print("ERROR: Corrupt render material found. (ON_BinaryArchive::Read3dmMaterial() < 0.)\n");
-				//	error_count++;
-				//	if (error_count > max_error_count)
-				//		return false;
-				//	error_log->Print("-- Attempting to continue.\n");
-				//}
 				pMaterial = new ON_Material; // use default
 				pMaterial->SetIndex(Count);
 			}
@@ -2476,7 +2454,6 @@ bool FOpenNurbsTranslatorImpl::Read(ON_BinaryFile& Archive, TSharedRef<IDatasmit
 		// EndRead3dmMaterialTable() must be called when BeginRead3dmMaterialTable() returns true.
 		if (!Archive.EndRead3dmMaterialTable())
 		{
-			//if (error_log) error_log->Print("ERROR: Corrupt render material table. (ON_BinaryArchive::EndRead3dmMaterialTable() returned false.)\n");
 			return false;
 		}
 		if (DatasmithOpenNurbsTranslatorUtils::CheckForCRCErrors(Archive, "render material table", NumCRCErrors))
@@ -2486,11 +2463,6 @@ bool FOpenNurbsTranslatorImpl::Read(ON_BinaryFile& Archive, TSharedRef<IDatasmit
 	}
 	else
 	{
-		//if (error_log)
-		//{
-		//	error_log->Print("WARNING: Missing or corrupt render material table. (ON_BinaryArchive::BeginRead3dmMaterialTable() returned false.)\n");
-		//	error_log->Print("-- Attempting to continue.\n");
-		//}
 		bResult = false;
 	}
 
@@ -2512,14 +2484,6 @@ bool FOpenNurbsTranslatorImpl::Read(ON_BinaryFile& Archive, TSharedRef<IDatasmit
 			}
 			if (ReturnCode < 0)
 			{
-				//if (error_log)
-				//{
-				//	error_log->Print("ERROR: Corrupt render linetype found. (ON_BinaryArchive::Read3dmLinetype() < 0.)\n");
-				//	error_count++;
-				//	if (error_count > max_error_count)
-				//		return false;
-				//	error_log->Print("-- Attempting to continue.\n");
-				//}
 				pLinetype = new ON_Linetype; // use default
 				pLinetype->SetIndex(Count);
 			}
@@ -2534,7 +2498,6 @@ bool FOpenNurbsTranslatorImpl::Read(ON_BinaryFile& Archive, TSharedRef<IDatasmit
 		// EndRead3dmLinetypeTable() must be called when BeginRead3dmLinetypeTable() returns true.
 		if (!Archive.EndRead3dmLinetypeTable())
 		{
-			//if (error_log) error_log->Print("ERROR: Corrupt render linetype table. (ON_BinaryArchive::EndRead3dmLinetypeTable() returned false.)\n");
 			return false;
 		}
 		if (DatasmithOpenNurbsTranslatorUtils::CheckForCRCErrors(Archive, "render linetype table", NumCRCErrors))
@@ -2544,11 +2507,6 @@ bool FOpenNurbsTranslatorImpl::Read(ON_BinaryFile& Archive, TSharedRef<IDatasmit
 	}
 	else
 	{
-		//if (error_log)
-		//{
-		//	error_log->Print("WARNING: Missing or corrupt render linetype table. (ON_BinaryArchive::BeginRead3dmLinetypeTable() returned false.)\n");
-		//	error_log->Print("-- Attempting to continue.\n");
-		//}
 		bResult = false;
 	}
 
@@ -2568,14 +2526,6 @@ bool FOpenNurbsTranslatorImpl::Read(ON_BinaryFile& Archive, TSharedRef<IDatasmit
 			}
 			if (ReturnCode < 0)
 			{
-				//if (error_log)
-				//{
-				//	error_log->Print("ERROR: Corrupt layer found. (ON_BinaryArchive::Read3dmLayer() < 0.)\n");
-				//	error_count++;
-				//	if (error_count > max_error_count)
-				//		return false;
-				//	error_log->Print("-- Attempting to continue.\n");
-				//}
 				pLayer = new ON_Layer; // use default
 				pLayer->SetIndex(Count);
 			}
@@ -2590,7 +2540,6 @@ bool FOpenNurbsTranslatorImpl::Read(ON_BinaryFile& Archive, TSharedRef<IDatasmit
 		// EndRead3dmLayerTable() must be called when BeginRead3dmLayerTable() returns true.
 		if (!Archive.EndRead3dmLayerTable())
 		{
-			//if (error_log) error_log->Print("ERROR: Corrupt render layer table. (ON_BinaryArchive::EndRead3dmLayerTable() returned false.)\n");
 			return false;
 		}
 		if (DatasmithOpenNurbsTranslatorUtils::CheckForCRCErrors(Archive, "layer table", NumCRCErrors))
@@ -2600,11 +2549,6 @@ bool FOpenNurbsTranslatorImpl::Read(ON_BinaryFile& Archive, TSharedRef<IDatasmit
 	}
 	else
 	{
-		//if (error_log)
-		//{
-		//	error_log->Print("WARNING: Missing or corrupt layer table. (ON_BinaryArchive::BeginRead3dmLayerTable() returned false.)\n");
-		//	error_log->Print("-- Attempting to continue.\n");
-		//}
 		bResult = false;
 	}
 
@@ -2624,14 +2568,6 @@ bool FOpenNurbsTranslatorImpl::Read(ON_BinaryFile& Archive, TSharedRef<IDatasmit
 			}
 			if (ReturnCode < 0)
 			{
-				//if (error_log)
-				//{
-				//	error_log->Print("ERROR: Corrupt group found. (ON_BinaryArchive::Read3dmGroup() < 0.)\n");
-				//	error_count++;
-				//	if (error_count > max_error_count)
-				//		return false;
-				//	error_log->Print("-- Attempting to continue.\n");
-				//}
 				pGroup = new ON_Group; // use default
 				pGroup->SetIndex(-1);
 			}
@@ -2646,7 +2582,6 @@ bool FOpenNurbsTranslatorImpl::Read(ON_BinaryFile& Archive, TSharedRef<IDatasmit
 		// EndRead3dmGroupTable() must be called when BeginRead3dmGroupTable() returns true.
 		if (!Archive.EndRead3dmGroupTable())
 		{
-			//if (error_log) error_log->Print("ERROR: Corrupt group table. (ON_BinaryArchive::EndRead3dmGroupTable() returned false.)\n");
 			return false;
 		}
 		if (DatasmithOpenNurbsTranslatorUtils::CheckForCRCErrors(Archive, "group table", NumCRCErrors))
@@ -2656,11 +2591,6 @@ bool FOpenNurbsTranslatorImpl::Read(ON_BinaryFile& Archive, TSharedRef<IDatasmit
 	}
 	else
 	{
-		//if (error_log)
-		//{
-		//	error_log->Print("WARNING: Missing or corrupt group table. (ON_BinaryArchive::BeginRead3dmGroupTable() returned false.)\n");
-		//	error_log->Print("-- Attempting to continue.\n");
-		//}
 		bResult = false;
 	}
 
@@ -2684,14 +2614,6 @@ bool FOpenNurbsTranslatorImpl::Read(ON_BinaryFile& Archive, TSharedRef<IDatasmit
 			}
 			if (ReturnCode < 0)
 			{
-				//if (error_log)
-				//{
-				//	error_log->Print("ERROR: Corrupt dimstyle found. (ON_BinaryArchive::Read3dmDimStyle() < 0.)\n");
-				//	error_count++;
-				//	if (error_count > max_error_count)
-				//		return false;
-				//	error_log->Print("-- Attempting to continue.\n");
-				//}
 				pDimStyle = new ON_DimStyle; // use default
 				pDimStyle->SetIndex(Count);
 			}
@@ -2707,7 +2629,6 @@ bool FOpenNurbsTranslatorImpl::Read(ON_BinaryFile& Archive, TSharedRef<IDatasmit
 		// then you MUST call EndRead3dmDimStyleTable().
 		if (!Archive.EndRead3dmDimStyleTable())
 		{
-			//if (error_log) error_log->Print("ERROR: Corrupt dimstyle table. (ON_BinaryArchive::EndRead3dmDimStyleTable() returned false.)\n");
 			return false;
 		}
 		if (DatasmithOpenNurbsTranslatorUtils::CheckForCRCErrors(Archive, "dimstyle table", NumCRCErrors))
@@ -2717,11 +2638,6 @@ bool FOpenNurbsTranslatorImpl::Read(ON_BinaryFile& Archive, TSharedRef<IDatasmit
 	}
 	else
 	{
-		//if (error_log)
-		//{
-		//	error_log->Print("WARNING: Missing or corrupt dimstyle table. (ON_BinaryArchive::BeginRead3dmDimStyleTable() returned false.)\n");
-		//	error_log->Print("-- Attempting to continue.\n");
-		//}
 		bResult = false;
 	}
 
@@ -2742,14 +2658,6 @@ bool FOpenNurbsTranslatorImpl::Read(ON_BinaryFile& Archive, TSharedRef<IDatasmit
 			}
 			if (ReturnCode < 0)
 			{
-				//if (error_log)
-				//{
-				//	error_log->Print("ERROR: Corrupt render light found. (ON_BinaryArchive::Read3dmLight() < 0.)\n");
-				//	error_count++;
-				//	if (error_count > max_error_count)
-				//		return false;
-				//	error_log->Print("-- Attempting to continue.\n");
-				//}
 				continue;
 			}
 
@@ -2761,7 +2669,6 @@ bool FOpenNurbsTranslatorImpl::Read(ON_BinaryFile& Archive, TSharedRef<IDatasmit
 		// EndRead3dmLightTable() must be called when BeginRead3dmLightTable() returns true.
 		if (!Archive.EndRead3dmLightTable())
 		{
-			//if (error_log) error_log->Print("ERROR: Corrupt render light table. (ON_BinaryArchive::EndRead3dmLightTable() returned false.)\n");
 			return false;
 		}
 		if (DatasmithOpenNurbsTranslatorUtils::CheckForCRCErrors(Archive, "render light table", NumCRCErrors))
@@ -2771,11 +2678,6 @@ bool FOpenNurbsTranslatorImpl::Read(ON_BinaryFile& Archive, TSharedRef<IDatasmit
 	}
 	else
 	{
-		//if (error_log)
-		//{
-		//	error_log->Print("WARNING: Missing or corrupt render light table. (ON_BinaryArchive::BeginRead3dmLightTable() returned false.)\n");
-		//	error_log->Print("-- Attempting to continue.\n");
-		//}
 		bResult = false;
 	}
 
@@ -2809,14 +2711,6 @@ bool FOpenNurbsTranslatorImpl::Read(ON_BinaryFile& Archive, TSharedRef<IDatasmit
 			}
 			if (ReturnCode < 0)
 			{
-				//if (error_log)
-				//{
-				//	error_log->Print("ERROR: Corrupt hatchpattern found. (ON_BinaryArchive::Read3dmHatchPattern() < 0.)\n");
-				//	error_count++;
-				//	if (error_count > max_error_count)
-				//		return false;
-				//	error_log->Print("-- Attempting to continue.\n");
-				//}
 				pHatchPattern = new ON_HatchPattern; // use default
 				pHatchPattern->SetIndex(Count);
 			}
@@ -2831,7 +2725,6 @@ bool FOpenNurbsTranslatorImpl::Read(ON_BinaryFile& Archive, TSharedRef<IDatasmit
 		// EndRead3dmHatchPatternTable() must be called when BeginRead3dmHatchPatternTable() returns true.
 		if (!Archive.EndRead3dmHatchPatternTable())
 		{
-			//if (error_log) error_log->Print("ERROR: Corrupt hatchpattern table. (ON_BinaryArchive::EndRead3dmHatchPatternTable() returned false.)\n");
 			return false;
 		}
 		if (DatasmithOpenNurbsTranslatorUtils::CheckForCRCErrors(Archive, "hatchpattern table", NumCRCErrors))
@@ -2841,11 +2734,6 @@ bool FOpenNurbsTranslatorImpl::Read(ON_BinaryFile& Archive, TSharedRef<IDatasmit
 	}
 	else
 	{
-		//if (error_log)
-		//{
-		//	error_log->Print("WARNING: Missing or corrupt hatchpattern table. (ON_BinaryArchive::BeginRead3dmHatchPatternTable() returned false.)\n");
-		//	error_log->Print("-- Attempting to continue.\n");
-		//}
 		bResult = false;
 	}
 
@@ -2864,29 +2752,14 @@ bool FOpenNurbsTranslatorImpl::Read(ON_BinaryFile& Archive, TSharedRef<IDatasmit
 			}
 			if (ReturnCode < 0)
 			{
-				//if (error_log)
-				//{
-				//	error_log->Print("ERROR: Corrupt instance definition found. (ON_BinaryArchive::Read3dmInstanceDefinition() < 0.)\n");
-				//	error_count++;
-				//	if (error_count > max_error_count)
-				//		return false;
-				//	error_log->Print("-- Attempting to continue.\n");
-				//}
 				continue;
 			}
-			// Note that pIDef is deleted later after step 15
-			////ON_UserDataHolder ud;
-			////ud.MoveUserDataFrom(*pIDef);
-			////m_idef_table.Append(*pIDef);
-			////ud.MoveUserDataTo(*m_idef_table.Last(), false);
-			////delete pIDef;
 			InstanceDefinitionTable.Add(pIDef);
 		}
 
 		// EndRead3dmInstanceDefinitionTable() must be called when BeginRead3dmInstanceDefinitionTable() returns true.
 		if (!Archive.EndRead3dmInstanceDefinitionTable())
 		{
-			//if (error_log) error_log->Print("ERROR: Corrupt instance definition table. (ON_BinaryArchive::EndRead3dmInstanceDefinitionTable() returned false.)\n");
 			return false;
 		}
 		if (DatasmithOpenNurbsTranslatorUtils::CheckForCRCErrors(Archive, "instance definition table", NumCRCErrors))
@@ -2896,11 +2769,6 @@ bool FOpenNurbsTranslatorImpl::Read(ON_BinaryFile& Archive, TSharedRef<IDatasmit
 	}
 	else
 	{
-		//if (error_log)
-		//{
-		//	error_log->Print("WARNING: Missing or corrupt instance definition table. (ON_BinaryArchive::BeginRead3dmInstanceDefinitionTable() returned false.)\n");
-		//	error_log->Print("-- Attempting to continue.\n");
-		//}
 		bResult = false;
 	}
 
@@ -2926,49 +2794,19 @@ bool FOpenNurbsTranslatorImpl::Read(ON_BinaryFile& Archive, TSharedRef<IDatasmit
 			}
 			if (ReturnCode < 0)
 			{
-				//if (error_log)
-				//{
-				//	error_log->Print("ERROR: Object table entry %d is corrupt. (ON_BinaryArchive::Read3dmObject() < 0.)\n", count);
-				//	error_count++;
-				//	if (error_count > max_error_count)
-				//		return false;
-				//	error_log->Print("-- Attempting to continue.\n");
-				//}
 				continue;
 			}
-			//if (m_crc_error_count != archive.BadCRCCount())
-			//{
-			//	if (error_log)
-			//	{
-			//		error_log->Print("ERROR: Object table entry %d is corrupt. (CRC errors).\n", count);
-			//		error_log->Print("-- Attempting to continue.\n");
-			//	}
-			//	m_crc_error_count = archive.BadCRCCount();
-			//}
 			if (pObject)
 			{
 				FOpenNurbsObjectWrapper& mo = ObjectTable.AppendNew();
 				mo.ObjectPtr = pObject;
 				mo.Attributes = attributes;
 			}
-			else
-			{
-				//if (error_log)
-				//{
-				//	if (rc == 2)
-				//		error_log->Print("WARNING: Skipping object table entry %d because it's filtered.\n", count);
-				//	else if (rc == 3)
-				//		error_log->Print("WARNING: Skipping object table entry %d because it's newer than this code.  Update your OpenNURBS toolkit.\n", count);
-				//	else
-				//		error_log->Print("WARNING: Skipping object table entry %d for unknown reason.\n", count);
-				//}
-			}
 		}
 
 		// EndRead3dmObjectTable() must be called when BeginRead3dmObjectTable() returns true.
 		if (!Archive.EndRead3dmObjectTable())
 		{
-			//if (error_log) error_log->Print("ERROR: Corrupt object light table. (ON_BinaryArchive::EndRead3dmObjectTable() returned false.)\n");
 			return false;
 		}
 		if (DatasmithOpenNurbsTranslatorUtils::CheckForCRCErrors(Archive, "object table", NumCRCErrors))
@@ -2978,11 +2816,6 @@ bool FOpenNurbsTranslatorImpl::Read(ON_BinaryFile& Archive, TSharedRef<IDatasmit
 	}
 	else
 	{
-		//if (error_log)
-		//{
-		//	error_log->Print("WARNING: Missing or corrupt object table. (ON_BinaryArchive::BeginRead3dmObjectTable() returned false.)\n");
-		//	error_log->Print("-- Attempting to continue.\n");
-		//}
 		bResult = false;
 	}
 
@@ -3008,42 +2841,18 @@ bool FOpenNurbsTranslatorImpl::Read(ON_BinaryFile& Archive, TSharedRef<IDatasmit
 			}
 			if (ReturnCode < 0)
 			{
-				//if (error_log)
-				//{
-				//	error_log->Print("ERROR: History record table entry %d is corrupt. (ON_BinaryArchive::Read3dmHistoryRecord() < 0.)\n", count);
-				//	error_count++;
-				//	if (error_count > max_error_count)
-				//		return false;
-				//	error_log->Print("-- Attempting to continue.\n");
-				//}
 				continue;
 			}
-			//if (m_crc_error_count != archive.BadCRCCount())
-			//{
-			//	if (error_log)
-			//	{
-			//		error_log->Print("ERROR: History record table entry %d is corrupt. (CRC errors).\n", count);
-			//		error_log->Print("-- Attempting to continue.\n");
-			//	}
-			//	m_crc_error_count = archive.BadCRCCount();
-			//}
+
 			if (pHistoryRecord)
 			{
 				HistoryRecordTable.Append(pHistoryRecord);
 			}
-			else
-			{
-				//if (error_log)
-				//{
-				//	error_log->Print("WARNING: Skipping history record table entry %d for unknown reason.\n", count);
-				//}
-			}
 		}
 
 		// EndRead3dmHistoryRecordTable() must be called when BeginRead3dmHistoryRecordTable() returns true.
 		if (!Archive.EndRead3dmHistoryRecordTable())
 		{
-			//if (error_log) error_log->Print("ERROR: Corrupt object light table. (ON_BinaryArchive::EndRead3dmObjectTable() returned false.)\n");
 			return false;
 		}
 		if (DatasmithOpenNurbsTranslatorUtils::CheckForCRCErrors(Archive, "history record table", NumCRCErrors))
@@ -3053,11 +2862,6 @@ bool FOpenNurbsTranslatorImpl::Read(ON_BinaryFile& Archive, TSharedRef<IDatasmit
 	}
 	else
 	{
-		//if (error_log)
-		//{
-		//	error_log->Print("WARNING: Missing or corrupt history record table. (ON_BinaryArchive::BeginRead3dmHistoryRecordTable() returned false.)\n");
-		//	error_log->Print("-- Attempting to continue.\n");
-		//}
 		bResult = false;
 	}
 
@@ -3122,7 +2926,6 @@ bool FOpenNurbsTranslatorImpl::Read(ON_BinaryFile& Archive, TSharedRef<IDatasmit
 
 		if (!Archive.Read3dmAnonymousUserTable(usertable_3dm_version, usertable_opennurbs_version, ud.m_goo))
 		{
-			//if (error_log) error_log->Print("ERROR: User data table entry %d is corrupt. (ON_BinaryArchive::Read3dmAnonymousUserTable() is false.)\n", count);
 			break;
 		}
 
@@ -3130,23 +2933,10 @@ bool FOpenNurbsTranslatorImpl::Read(ON_BinaryFile& Archive, TSharedRef<IDatasmit
 		// EndRead3dmUserTable() must be called when BeginRead3dmObjectTable() returns true.
 		if (!Archive.EndRead3dmUserTable())
 		{
-			//if (error_log) error_log->Print("ERROR: Corrupt user data table. (ON_BinaryArchive::EndRead3dmUserTable() returned false.)\n");
 			break;
 		}
 	}
 
-	// Nothing to do
-
-	// Step 18: OPTIONAL when reading: EndMark
-	if (!Archive.Read3dmEndMark(&FileLength))
-	{
-		if (Archive.Archive3dmVersion() != 1)
-		{
-			// some v1 files are missing end-of-archive markers
-			//if (error_log) error_log->Print("ERROR: ON_BinaryArchive::Read3dmEndMark(&m_file_length) returned false.\n");
-		}
-	}
-
 	// Clean up the scene by removing unused layer actors
 	TArray<TSharedPtr<IDatasmithActorElement>> ActorsToRemove;
 	int32 NumActors = Scene->GetActorsCount();
@@ -3167,7 +2957,7 @@ bool FOpenNurbsTranslatorImpl::Read(ON_BinaryFile& Archive, TSharedRef<IDatasmit
 	return true;
 }
 
-bool FOpenNurbsTranslatorImpl::TranslateBRep(ON_Brep* Brep, const ON_3dmObjectAttributes& Attributes, FMeshDescription& OutMesh, const TSharedRef< IDatasmithMeshElement >& MeshElement, const FString& Name, bool& bHasNormal)
+bool FOpenNurbsTranslatorImpl::TranslateBRep(ON_Brep* Brep, const ON_3dmObjectAttributes& Attributes, FMeshDescription& OutMesh, TSharedRef< IDatasmithMeshElement >& MeshElement, const FString& Name, bool& bHasNormal)
 {
 	if (Brep == nullptr)
 	{
@@ -3181,24 +2971,17 @@ bool FOpenNurbsTranslatorImpl::TranslateBRep(ON_Brep* Brep, const ON_3dmObjectAt
 	{
 		// Ref. visitBRep
 		const FDatasmithOpenNurbsOptions& TessellationOptions = OpenNurbsOptions;
-		LocalSession->SetImportParameters(TessellationOptions.ChordTolerance, TessellationOptions.MaxEdgeLength, TessellationOptions.NormalTolerance, (CADLibrary::EStitchingTechnique) TessellationOptions.StitchingTechnique, false);
-		LocalSession->SetModelCoordinateSystem(FDatasmithUtils::EModelCoordSystem::ZUp_RightHanded_FBXLegacy);
+		CADModelConverter->SetImportParameters(TessellationOptions.ChordTolerance, TessellationOptions.MaxEdgeLength, TessellationOptions.NormalTolerance, (CADLibrary::EStitchingTechnique) TessellationOptions.StitchingTechnique, false);
 
-		LocalSession->ClearData();
-		LocalSession->SetSceneUnit(GetMetricUnit());
+		CADModelConverter->InitializeProcess(CADModelConverter->GetMetricUnit());
 
-		LocalSession->AddBRep(*Brep, Offset);
+		OpenNurbsBRepConverter->AddBRep(*Brep, Offset);
+		CADModelConverter->RepairTopology();
 
-		FString FilePath = FPaths::Combine(OutputPath, Name) += TEXT(".ct");
-		if (LocalSession->SaveBrep(FilePath))
-		{
-			MeshElement->SetFile(*FilePath);
-		}
-
-		LocalSession->TopoFixes();
+		CADModelConverter->SaveBRep(*OutputPath, MeshElement);
 
 		CADLibrary::FMeshParameters MeshParameters;
-		return LocalSession->Tessellate(OutMesh, MeshParameters);
+		return CADModelConverter->Tessellate(MeshParameters, OutMesh);
 	}
 	else
 	{
@@ -3237,7 +3020,7 @@ bool FOpenNurbsTranslatorImpl::TranslateBRep(ON_Brep* Brep, const ON_3dmObjectAt
 		}
 		if (bHasFaceMaterialChannel)
 		{
-			if (!DatasmithOpenNurbsTranslatorUtils::TranslateMesh(RenderMeshes, RenderMeshes.Count(), OutMesh, bHasNormal, ScalingFactor, Offset, bHasFaceMaterialChannel, FaceMaterialChannel.GetData()))
+			if (!DatasmithOpenNurbsTranslatorUtils::TranslateMesh(RenderMeshes, RenderMeshes.Count(), OutMesh, bHasNormal, GetScaleFactor(), Offset, bHasFaceMaterialChannel, FaceMaterialChannel.GetData()))
 			{
 				return false;
 			}
@@ -3247,7 +3030,7 @@ bool FOpenNurbsTranslatorImpl::TranslateBRep(ON_Brep* Brep, const ON_3dmObjectAt
 			ON_Mesh Mesh;
 			Mesh.Append(RenderMeshes.Count(), RenderMeshes);
 
-			if (!DatasmithOpenNurbsTranslatorUtils::TranslateMesh(&Mesh, OutMesh, bHasNormal, ScalingFactor, Offset))
+			if (!DatasmithOpenNurbsTranslatorUtils::TranslateMesh(&Mesh, OutMesh, bHasNormal, GetScaleFactor(), Offset))
 			{
 				return false;
 			}
@@ -3357,7 +3140,7 @@ TOptional<FMeshDescription> FOpenNurbsTranslatorImpl::GetMeshDescription(TShared
 	if (Object.ObjectPtr->IsKindOf(&ON_Mesh::m_ON_Mesh_class_rtti))
 	{
 		ON_3dVector Offset = GetGeometryOffset(MeshElement);
-		bIsValid = DatasmithOpenNurbsTranslatorUtils::TranslateMesh(ON_Mesh::Cast(Object.ObjectPtr), MeshDescription, bHasNormal, SelectedTranslator->ScalingFactor, Offset);
+		bIsValid = DatasmithOpenNurbsTranslatorUtils::TranslateMesh(ON_Mesh::Cast(Object.ObjectPtr), MeshDescription, bHasNormal, SelectedTranslator->GetScaleFactor(), Offset);
 	}
 	else if (Object.ObjectPtr->IsKindOf(&ON_Brep::m_ON_Brep_class_rtti))
 	{
@@ -3370,7 +3153,7 @@ TOptional<FMeshDescription> FOpenNurbsTranslatorImpl::GetMeshDescription(TShared
 			ON_3dVector Offset = GetGeometryOffset(MeshElement);
 			if (const ON_Mesh* Mesh = extrusion->m_mesh_cache.Mesh(ON::mesh_type::render_mesh))
 			{
-				if (DatasmithOpenNurbsTranslatorUtils::TranslateMesh(Mesh, MeshDescription, bHasNormal, ScalingFactor, Offset))
+				if (DatasmithOpenNurbsTranslatorUtils::TranslateMesh(Mesh, MeshDescription, bHasNormal, GetScaleFactor(), Offset))
 				{
 					return MeshDescription;
 				}
@@ -3495,27 +3278,21 @@ void FDatasmithOpenNurbsTranslator::UnloadScene()
 
 bool FDatasmithOpenNurbsTranslator::LoadStaticMesh(const TSharedRef<IDatasmithMeshElement> MeshElement, FDatasmithMeshElementPayload& OutMeshPayload)
 {
-	if (TOptional< FMeshDescription > Mesh = Translator->GetMeshDescription(MeshElement))
+	if (Translator)
 	{
-		OutMeshPayload.LodMeshes.Add(MoveTemp(Mesh.GetValue()));
-
-		CADLibrary::FImportParameters ImportParameters;
-		ImportParameters.ModelCoordSys = FDatasmithUtils::EModelCoordSystem::ZUp_RightHanded_FBXLegacy;
-		ImportParameters.MetricUnit = Translator->GetMetricUnit();
-		ImportParameters.ScaleFactor = Translator->GetScalingFactor();
+		return Translator->LoadStaticMesh(MeshElement, OutMeshPayload, GetCommonTessellationOptions());
+	}
+	return false;
+}
 
+bool FOpenNurbsTranslatorImpl::LoadStaticMesh(const TSharedRef<IDatasmithMeshElement> MeshElement, FDatasmithMeshElementPayload & OutMeshPayload, const FDatasmithTessellationOptions& InTessellationOptions)
+{
+	if (TOptional<FMeshDescription> Mesh = GetMeshDescription(MeshElement))
+	{
 		CADLibrary::FMeshParameters MeshParameters;
-
-		if (ImportParameters.bEnableKernelIOTessellation)
-		{
-			CoreTechSurface::AddSurfaceDataForMesh(MeshElement->GetFile(), ImportParameters, MeshParameters, OpenNurbsOptions, OutMeshPayload);
-		}
-		else
-		{
-			//CADKernelSurface::AddSurfaceDataForMesh(MeshElement, ImportParameters, MeshParameters, OpenNurbsOptions, OutMeshPayload);
-		}
+		OutMeshPayload.LodMeshes.Add(MoveTemp(Mesh.GetValue()));
+		CADModelConverter->AddSurfaceDataForMesh(MeshElement->GetFile(), MeshParameters, InTessellationOptions, OutMeshPayload);
 	}
-
 	return OutMeshPayload.LodMeshes.Num() > 0;
 }
 
diff --git a/Engine/Plugins/Enterprise/DatasmithCADImporter/Source/DatasmithOpenNurbsTranslator/Private/DatasmithOpenNurbsTranslator.h b/Engine/Plugins/Enterprise/DatasmithCADImporter/Source/DatasmithOpenNurbsTranslator/Private/DatasmithOpenNurbsTranslator.h
--- a/Engine/Plugins/Enterprise/DatasmithCADImporter/Source/DatasmithOpenNurbsTranslator/Private/DatasmithOpenNurbsTranslator.h
+++ b/Engine/Plugins/Enterprise/DatasmithCADImporter/Source/DatasmithOpenNurbsTranslator/Private/DatasmithOpenNurbsTranslator.h
@@ -6,12 +6,13 @@
 
 #include "DatasmithTranslator.h"
 #include "DatasmithOpenNurbsImportOptions.h"
+#include "ParametricSurfaceTranslator.h"
 
 #include "CoreMinimal.h"
 
 class FOpenNurbsTranslatorImpl;
 
-class FDatasmithOpenNurbsTranslator : public IDatasmithTranslator
+class FDatasmithOpenNurbsTranslator : public FParametricSurfaceTranslator
 {
 public:
 	virtual FName GetFName() const override { return "DatasmithOpenNurbsTranslator"; }
diff --git a/Engine/Plugins/Enterprise/DatasmithCADImporter/Source/DatasmithOpenNurbsTranslator/Private/OpenNurbsBRepConverter.h b/Engine/Plugins/Enterprise/DatasmithCADImporter/Source/DatasmithOpenNurbsTranslator/Private/OpenNurbsBRepConverter.h
new file mode 100644
--- /dev/null
+++ b/Engine/Plugins/Enterprise/DatasmithCADImporter/Source/DatasmithOpenNurbsTranslator/Private/OpenNurbsBRepConverter.h
@@ -0,0 +1,13 @@
+// Copyright Epic Games, Inc. All Rights Reserved.
+
+#pragma once
+
+class ON_Brep;
+class ON_3dVector;
+
+class IOpenNurbsBRepConverter
+{
+public:
+	virtual bool AddBRep(ON_Brep& Brep, const ON_3dVector& Offset) = 0;
+};
+
diff --git a/Engine/Plugins/Enterprise/DatasmithCADImporter/Source/DatasmithOpenNurbsTranslator/Private/OpenNurbsBRepToCADKernelConverter.cpp b/Engine/Plugins/Enterprise/DatasmithCADImporter/Source/DatasmithOpenNurbsTranslator/Private/OpenNurbsBRepToCADKernelConverter.cpp
new file mode 100644
--- /dev/null
+++ b/Engine/Plugins/Enterprise/DatasmithCADImporter/Source/DatasmithOpenNurbsTranslator/Private/OpenNurbsBRepToCADKernelConverter.cpp
@@ -0,0 +1,330 @@
+// Copyright Epic Games, Inc. All Rights Reserved.
+
+#include "OpenNurbsBRepToCADKernelConverter.h"
+
+#ifdef USE_OPENNURBS
+#include "CoreTechSurfaceHelper.h"
+
+#pragma warning(push)
+#pragma warning(disable:4265)
+#pragma warning(disable:4005) // TEXT macro redefinition
+#include "opennurbs.h"
+#pragma warning(pop)
+
+#include <vector>
+
+#include "CADKernelTools.h"
+
+#include "CADKernel/Core/Session.h"
+#include "CADKernel/Geo/Curves/NURBSCurve.h"
+#include "CADKernel/Geo/Surfaces/NURBSSurface.h"
+#include "CADKernel/Geo/Surfaces/Surface.h"
+
+#include "CADKernel/Math/Boundary.h"
+#include "CADKernel/Math/Point.h"
+#include "CADKernel/Mesh/Meshers/ParametricMesher.h"
+#include "CADKernel/Mesh/Structure/ModelMesh.h"
+
+#include "CADKernel/Topo/Body.h"
+#include "CADKernel/Topo/Model.h"
+#include "CADKernel/Topo/Shell.h"
+#include "CADKernel/Topo/TopologicalEdge.h"
+#include "CADKernel/Topo/TopologicalFace.h"
+#include "CADKernel/Topo/TopologicalLoop.h"
+
+
+using namespace CADLibrary;
+using namespace CADKernel;
+
+namespace
+{
+	enum EAxis { U, V };
+	struct PerAxisInfo
+	{
+		ON_NurbsSurface& OpenNurbsSurface;
+		EAxis Axis;
+		int32& Degree; // Degree + 1
+		TArray<double>& Knots; // t values with superflux values
+		int32& CtrlVertCount; // number of control points
+
+		uint32 KnotSize;  // ON knots
+		uint32 KnotCount; // CT knots
+
+		TArray<uint32> KnotMultiplicities; // from ON, not relevant as we send n-plicated knots to CT (dbg only)
+
+		PerAxisInfo(EAxis InAxis, ON_NurbsSurface& InSurface, FNurbsSurfaceHomogeneousData& NurbsInfo)
+			: OpenNurbsSurface(InSurface)
+			, Axis(InAxis)
+			, Degree(Axis == U ? NurbsInfo.UDegree : NurbsInfo.VDegree)
+			, Knots(Axis == U ? NurbsInfo.UNodalVector : NurbsInfo.VNodalVector)
+			, CtrlVertCount(Axis == U ? NurbsInfo.PoleUCount : NurbsInfo.PoleVCount)
+		{
+			// detect cases not handled by CADKernel, that is knot vectors with multiplicity < order on either end
+			Degree = OpenNurbsSurface.Order(Axis) - 1;
+			if (OpenNurbsSurface.KnotMultiplicity(Axis, 0) < Degree || OpenNurbsSurface.KnotMultiplicity(Axis, KnotSize - 3) < Degree)
+			{
+				OpenNurbsSurface.IncreaseDegree(Axis, OpenNurbsSurface.Degree(Axis) + 1);
+			}
+			Populate();
+		}
+
+	private:
+		void Populate()
+		{
+			Degree = OpenNurbsSurface.Order(Axis) - 1;
+			CtrlVertCount = OpenNurbsSurface.CVCount(Axis);
+			KnotSize = Degree + CtrlVertCount + 1;
+			KnotCount = OpenNurbsSurface.KnotCount(Axis);
+
+			KnotMultiplicities.SetNumUninitialized(KnotSize - 2);
+			for (uint32 Index = 0; Index < KnotSize - 2; ++Index)
+			{
+				KnotMultiplicities[Index] = OpenNurbsSurface.KnotMultiplicity(Axis, Index); // 0 and < Order + CV_count - 2
+			}
+
+			Knots.Reserve(KnotSize);
+			Knots.Add(OpenNurbsSurface.SuperfluousKnot(Axis, 0));
+			for (uint32 i = 0; i < KnotCount; ++i)
+			{
+				Knots.Add(OpenNurbsSurface.Knot(Axis, i));
+			}
+			Knots.Add(OpenNurbsSurface.SuperfluousKnot(Axis, 1));
+		}
+	};
+}
+
+TSharedRef<FSurface> FOpenNurbsBRepToCADKernelConverter::AddSurface(ON_NurbsSurface& OpenNurbsSurface)
+{
+	FNurbsSurfaceHomogeneousData NurbsData;
+
+	int32 ControlVertexDimension = OpenNurbsSurface.CVSize();
+
+	PerAxisInfo UInfo(U, OpenNurbsSurface, NurbsData);
+	PerAxisInfo VInfo(V, OpenNurbsSurface, NurbsData);
+
+	NurbsData.HomogeneousPoles.SetNumUninitialized(UInfo.CtrlVertCount * VInfo.CtrlVertCount * OpenNurbsSurface.CVSize());
+	double* ControlPoints = NurbsData.HomogeneousPoles.GetData();
+	NurbsData.bIsRational = OpenNurbsSurface.IsRational();
+	ON::point_style PointStyle = OpenNurbsSurface.IsRational() ? ON::point_style::euclidean_rational : ON::point_style::not_rational;
+	for (int32 UIndex = 0; UIndex < UInfo.CtrlVertCount; ++UIndex)
+	{
+		for (int32 VIndex = 0; VIndex < VInfo.CtrlVertCount; ++VIndex, ControlPoints += ControlVertexDimension)
+		{
+			OpenNurbsSurface.GetCV(UIndex, VIndex, PointStyle, ControlPoints);
+		}
+	}
+
+#ifdef REMOVE_NEGATIVE_WEIGHT
+	if (NurbsData.bIsRational)
+	{
+		bool bHasNegativeWeight = false;
+		for (int32 Index = 3; Index < NurbsData.HomogeneousPoles.Num(); Index += 4)
+		{
+			double Weight = ControlPoints[Index];
+			if (Weight < 0.)
+			{
+				bHasNegativeWeight = true;
+				break;
+			}
+		}
+
+		if (bHasNegativeWeight)
+		{
+			UInfo.IncreaseDegree();
+			VInfo.IncreaseDegree();
+			BuildHull();
+		}
+	}
+#endif
+
+	return FEntity::MakeShared<FNURBSSurface>(GeometricTolerance, NurbsData);
+}
+
+TSharedPtr<FTopologicalLoop> FOpenNurbsBRepToCADKernelConverter::AddLoop(const ON_BrepLoop& OpenNurbsLoop, TSharedRef<FSurface>& CarrierSurface)
+{
+	if (!OpenNurbsLoop.IsValid())
+	{
+		return TSharedPtr<FTopologicalLoop>();
+	}
+
+	ON_BrepLoop::TYPE LoopType = OpenNurbsLoop.m_type;
+	bool bIsOuter = (LoopType == ON_BrepLoop::TYPE::outer);
+
+	int32 EdgeCount = OpenNurbsLoop.TrimCount();
+	TArray<TSharedPtr<FTopologicalEdge>> Edges;
+	TArray<CADKernel::EOrientation> Directions;
+
+	Edges.Reserve(EdgeCount);
+	Directions.Reserve(EdgeCount);
+
+	for (int32 Index = 0; Index < EdgeCount; ++Index)
+	{
+		ON_BrepTrim& Trim = *OpenNurbsLoop.Trim(Index);
+
+		TSharedPtr<FTopologicalEdge> Edge = AddEdge(Trim, CarrierSurface);
+		if (Edge.IsValid())
+		{
+			Edges.Add(Edge);
+			Directions.Emplace(CADKernel::EOrientation::Front);
+		}
+	}
+
+	if (Edges.Num() == 0)
+	{
+		return TSharedPtr<FTopologicalLoop>();
+	}
+
+	return FTopologicalLoop::Make(Edges, Directions, GeometricTolerance);
+}
+
+TSharedPtr<FTopologicalEdge> FOpenNurbsBRepToCADKernelConverter::AddEdge(const ON_BrepTrim& OpenNurbsTrim, TSharedRef<FSurface>& CarrierSurface)
+{
+	ON_BrepEdge* OpenNurbsEdge = OpenNurbsTrim.Edge();
+	if (OpenNurbsEdge == nullptr)
+	{
+		return TSharedPtr<FTopologicalEdge>();
+	}
+
+	ON_NurbsCurve OpenNurbsCurve;
+	int NurbFormSuccess = OpenNurbsTrim.GetNurbForm(OpenNurbsCurve); // 0:Nok 1:Ok 2:OkBut
+	if (NurbFormSuccess == 0)
+	{
+		return TSharedPtr<FTopologicalEdge>();
+	}
+
+	FNurbsCurveData NurbsCurveData;
+
+	NurbsCurveData.Dimension = 2;
+	NurbsCurveData.Degree = OpenNurbsCurve.Order() - 1;
+
+	int32 KnotCount = OpenNurbsCurve.KnotCount();
+	int32 ControlPointCount = OpenNurbsCurve.CVCount();
+
+	NurbsCurveData.NodalVector.Reserve(KnotCount + 2);
+	NurbsCurveData.NodalVector.Emplace(OpenNurbsCurve.SuperfluousKnot(0));
+	for (int Index = 0; Index < KnotCount; ++Index)
+	{
+		NurbsCurveData.NodalVector.Emplace(OpenNurbsCurve.Knot(Index));
+	}
+	NurbsCurveData.NodalVector.Emplace(OpenNurbsCurve.SuperfluousKnot(1));
+
+	NurbsCurveData.bIsRational = OpenNurbsCurve.IsRational();
+
+	NurbsCurveData.Poles.SetNumUninitialized(ControlPointCount);
+
+	double* ControlPoints = (double*) NurbsCurveData.Poles.GetData();
+	for (int32 Index = 0; Index < ControlPointCount; ++Index, ControlPoints += 3)
+	{
+		OpenNurbsCurve.GetCV(Index, OpenNurbsCurve.IsRational() ? ON::point_style::euclidean_rational : ON::point_style::not_rational, ControlPoints);
+	}
+
+	if(NurbsCurveData.bIsRational)
+	{
+		NurbsCurveData.Weights.SetNumUninitialized(ControlPointCount);
+		for (int32 Index = 0; Index < ControlPointCount; ++Index)
+		{
+			NurbsCurveData.Weights[Index] = NurbsCurveData.Poles[Index].Z;
+		}
+	}
+
+	for (int32 Index = 0; Index < ControlPointCount; ++Index)
+	{
+		NurbsCurveData.Poles[Index].Z = 0;
+	}
+
+	TSharedRef<FNURBSCurve> Nurbs = FEntity::MakeShared<FNURBSCurve>(NurbsCurveData);
+
+	TSharedRef<FRestrictionCurve> RestrictionCurve = FEntity::MakeShared<FRestrictionCurve>(CarrierSurface, Nurbs);
+
+	ON_Interval dom = OpenNurbsCurve.Domain();
+	FLinearBoundary Boundary(dom.m_t[0], dom.m_t[1]);
+	TSharedPtr<FTopologicalEdge> Edge = FTopologicalEdge::Make(RestrictionCurve, Boundary);
+	if (!Edge.IsValid())
+	{
+		return TSharedPtr<FTopologicalEdge>();
+	}
+
+	// Build topo i.e. find in the twin trims the the first edge and link both edges
+	for (int32 Index = 0; Index < OpenNurbsEdge->m_ti.Count(); ++Index)
+	{
+		int32 LinkedEdgeId = OpenNurbsEdge->m_ti[Index];
+		if (LinkedEdgeId == OpenNurbsTrim.m_trim_index)
+		{
+			continue;
+		}
+
+		TSharedPtr<CADKernel::FTopologicalEdge>* TwinEdge = OpenNurbsTrimId2CADKernelEdge.Find(LinkedEdgeId);
+		if (TwinEdge)
+		{
+			Edge->Link(TwinEdge->ToSharedRef(), SquareTolerance);
+			break;
+		}
+	}
+	OpenNurbsTrimId2CADKernelEdge.Add(OpenNurbsTrim.m_trim_index, Edge);
+
+	return Edge;
+}
+
+TSharedPtr<FTopologicalFace> FOpenNurbsBRepToCADKernelConverter::AddFace(const ON_BrepFace& OpenNurbsFace)
+{
+
+	ON_NurbsSurface OpenNurbsSurface;
+	OpenNurbsFace.NurbsSurface(&OpenNurbsSurface);
+
+	TSharedRef<FSurface> Surface = AddSurface(OpenNurbsSurface);
+
+	TSharedRef<FTopologicalFace> Face = FEntity::MakeShared<FTopologicalFace>(Surface);
+
+	const ON_BrepLoop* OuterLoop = OpenNurbsFace.OuterLoop();
+	if (OuterLoop == nullptr)
+	{
+		Face->ApplyNaturalLoops();
+		return Face;
+	}
+
+	int LoopCount = OpenNurbsFace.LoopCount();
+	for (int LoopIndex = 0; LoopIndex < LoopCount; ++LoopIndex)
+	{
+		const ON_BrepLoop& OpenNurbsLoop = *OpenNurbsFace.Loop(LoopIndex);
+		TSharedPtr<FTopologicalLoop> Loop = AddLoop(OpenNurbsLoop, Surface);
+		if(Loop.IsValid())
+		{
+			Face->AddLoop(Loop);
+		}
+	}
+
+	return Face;
+}
+
+bool FOpenNurbsBRepToCADKernelConverter::AddBRep(ON_Brep& BRep, const ON_3dVector& Offset)
+{
+	OpenNurbsTrimId2CADKernelEdge.Empty();
+
+	TSharedRef<FBody> Body = FEntity::MakeShared<FBody>();
+	TSharedRef<FShell> Shell = FEntity::MakeShared<FShell>();
+	Body->AddShell(Shell);
+
+	BRep.Translate(Offset);
+
+	BRep.FlipReversedSurfaces();
+
+	// Create faces
+	int FaceCount = BRep.m_F.Count();
+	for (int index = 0; index < FaceCount; index++)
+	{
+		const ON_BrepFace& OpenNurbsFace = BRep.m_F[index];
+		TSharedPtr<FTopologicalFace> Face = AddFace(OpenNurbsFace);
+		if (Face.IsValid())
+		{
+			Shell->Add(Face.ToSharedRef(), CADKernel::EOrientation::Front);
+		}
+	}
+
+	BRep.Translate(-Offset);
+
+	CADKernelSession.GetModel()->Add(Body);
+
+	return true;
+}
+
+#endif // defined(USE_OPENNURBS)
diff --git a/Engine/Plugins/Enterprise/DatasmithCADImporter/Source/DatasmithOpenNurbsTranslator/Private/OpenNurbsBRepToCADKernelConverter.h b/Engine/Plugins/Enterprise/DatasmithCADImporter/Source/DatasmithOpenNurbsTranslator/Private/OpenNurbsBRepToCADKernelConverter.h
new file mode 100644
--- /dev/null
+++ b/Engine/Plugins/Enterprise/DatasmithCADImporter/Source/DatasmithOpenNurbsTranslator/Private/OpenNurbsBRepToCADKernelConverter.h
@@ -0,0 +1,65 @@
+// Copyright Epic Games, Inc. All Rights Reserved.
+
+#pragma once
+#include "CoreMinimal.h"
+
+#ifdef USE_OPENNURBS
+
+#include "CADModelToCADKernelConverterBase.h"
+#include "OpenNurbsBRepConverter.h"
+
+#include "CADKernel/Core/Session.h"
+#include "CADKernel/Core/Types.h"
+
+class ON_BoundingBox;
+class ON_Brep;
+class ON_BrepFace;
+class ON_BrepLoop;
+class ON_BrepTrim;
+class ON_NurbsSurface;
+class ON_3dVector;
+
+ namespace CADKernel
+ {
+	 class FShell;
+	 class FSurface;
+	 class FTopologicalEdge;
+	 class FTopologicalFace;
+	 class FTopologicalLoop;
+ }
+
+class FOpenNurbsBRepToCADKernelConverter : public FCADModelToCADKernelConverterBase, public IOpenNurbsBRepConverter
+{
+public:
+
+	FOpenNurbsBRepToCADKernelConverter(const CADLibrary::FImportParameters& InImportParameters)
+		: FCADModelToCADKernelConverterBase(InImportParameters)
+	{
+	}
+
+	virtual ~FOpenNurbsBRepToCADKernelConverter() = default;
+
+	/**
+	 * Set BRep to tessellate, offsetting it prior to tessellation(used to set mesh pivot at the center of the surface bounding box)
+	 *
+	 * @param  Brep	a brep to tessellate
+	 * @param  Offset translate brep by this value before tessellating 
+	 */
+	bool AddBRep(ON_Brep& Brep, const ON_3dVector& Offset);
+	
+	static TSharedPtr<FOpenNurbsBRepToCADKernelConverter> GetSharedSession();
+
+private:
+	TSharedPtr<CADKernel::FTopologicalFace> AddFace(const ON_BrepFace& OpenNurbsFace);
+	TSharedRef<CADKernel::FSurface> AddSurface(ON_NurbsSurface& Surface);
+
+	TSharedPtr<CADKernel::FTopologicalLoop> AddLoop(const ON_BrepLoop& OpenNurbsLoop, TSharedRef<CADKernel::FSurface>& CarrierSurface);
+	TSharedPtr<CADKernel::FTopologicalEdge> AddEdge(const ON_BrepTrim& OpenNurbsTrim, TSharedRef<CADKernel::FSurface>& CarrierSurface);
+
+protected:
+
+	TMap<int32, TSharedPtr<CADKernel::FTopologicalEdge>>  OpenNurbsTrimId2CADKernelEdge;
+
+};
+
+#endif // USE_OPENNURBS
\ No newline at end of file
diff --git a/Engine/Plugins/Enterprise/DatasmithCADImporter/Source/DatasmithOpenNurbsTranslator/Private/OpenNurbsBRepToCoretechConverter.cpp b/Engine/Plugins/Enterprise/DatasmithCADImporter/Source/DatasmithOpenNurbsTranslator/Private/OpenNurbsBRepToCoretechConverter.cpp
new file mode 100644
--- /dev/null
+++ b/Engine/Plugins/Enterprise/DatasmithCADImporter/Source/DatasmithOpenNurbsTranslator/Private/OpenNurbsBRepToCoretechConverter.cpp
@@ -0,0 +1,438 @@
+// Copyright Epic Games, Inc. All Rights Reserved.
+
+#include "OpenNurbsBRepToCoretechConverter.h"
+
+#ifdef USE_OPENNURBS
+#include "CoreTechSurfaceHelper.h"
+
+#pragma warning(push)
+#pragma warning(disable:4265)
+#pragma warning(disable:4005) // TEXT macro redefinition
+#include "opennurbs.h"
+#pragma warning(pop)
+
+#include <vector>
+
+// As of 2018 SP2, CoreTech tessellation for a face is broken when the outer loop is the whole uv range, and some edges are singularities.
+// In those cases, inner loops behave like outers, outer is ignored.
+// current fix will split the uv plane in two parts.
+#define FIX_HOLE_IN_WHOLE_FACE 1
+
+// Handle surface state without u/v duplication
+// #ue_opennurbs : Remove SurfaceInfo and use CADLibrary::FNurbsSurface directly 
+struct SurfaceInfo
+{
+	ON_NurbsSurface& OpenNurbssSurface;
+
+	enum EAxis { U, V };
+	struct PerAxisInfo
+	{
+		uint32 Order; // Degree + 1
+		uint32 CtrlVertCount; // number of control points
+		uint32 KnotSize;  // ON knots
+		uint32 KnotCount; // CT knots
+
+		TArray<uint32> KnotMultiplicities; // from ON, not relevant as we send n-plicated knots to CT (dbg only)
+		TArray<uint32> CTKnotMultiplicityArray; // array of '1' for CT...
+		TArray<double> Knots; // t values with superflux values
+
+		PerAxisInfo(EAxis InA, ON_NurbsSurface& InSurf)
+			: OpenNurbsSurface(InSurf)
+			, Axis(InA)
+		{
+			Populate();
+		}
+
+		// increasing a nurbs degree doesn't change the shape, but this operation generates a new hull with new weights.
+		// This can fix exotic nurbs
+		void IncreaseDegree()
+		{
+			OpenNurbsSurface.IncreaseDegree(Axis, OpenNurbsSurface.Degree(Axis) + 1);
+			Populate();
+		}
+
+		// detect cases not handled by CT, that is knot vectors with multiplicity < order on either end
+		bool FixMultiplicity()
+		{
+			if (KnotMultiplicities[0] + 1 < Order || KnotMultiplicities.Last() + 1 < Order)
+			{
+				IncreaseDegree();
+				return true;
+			}
+			return false;
+		}
+
+	private:
+		void Populate()
+		{
+			Order = OpenNurbsSurface.Order(Axis);
+			CtrlVertCount = OpenNurbsSurface.CVCount(Axis);
+			KnotSize = Order + CtrlVertCount;
+			KnotCount = OpenNurbsSurface.KnotCount(Axis);
+
+			KnotMultiplicities.Init(-1, KnotSize - 2);
+			for (uint32 Index = 0; Index < KnotSize - 2; ++Index)
+			{
+				KnotMultiplicities[Index] = OpenNurbsSurface.KnotMultiplicity(Axis, Index); // 0 and < Order + CV_count - 2
+			}
+
+			Knots.Init(-1, KnotSize);
+			Knots[0] = OpenNurbsSurface.SuperfluousKnot(Axis, 0);
+			for (uint32 Index = 0; Index < KnotCount; ++Index)
+			{
+				Knots[Index + 1] = OpenNurbsSurface.Knot(Axis, Index);
+			}
+			Knots.Last() = OpenNurbsSurface.SuperfluousKnot(Axis, 1);
+
+			CTKnotMultiplicityArray.Init(1, KnotSize);
+		}
+
+		ON_NurbsSurface& OpenNurbsSurface;
+		EAxis Axis;
+	};
+
+	PerAxisInfo UInfo;
+	PerAxisInfo VInfo;
+	int ControlPointDimension; // Number of doubles per ctrl vertex
+	TArray<double> ControlPoints; // [xyzw...]
+
+	SurfaceInfo(ON_NurbsSurface& inSurf)
+		: OpenNurbssSurface(inSurf)
+		, UInfo(U, OpenNurbssSurface)
+		, VInfo(V, OpenNurbssSurface)
+		, ControlPointDimension(OpenNurbssSurface.CVSize())
+	{
+		BuildHull();
+	}
+
+	void BuildHull()
+	{
+		ControlPoints.Init(-1, UInfo.CtrlVertCount * VInfo.CtrlVertCount * ControlPointDimension);
+		double* CtrlHullPtr = ControlPoints.GetData();
+		ON::point_style pt_style = OpenNurbssSurface.IsRational() ? ON::point_style::euclidean_rational : ON::point_style::not_rational;
+		for (uint32 UIndex = 0; UIndex < UInfo.CtrlVertCount; ++UIndex)
+		{
+			for (uint32 VIndex = 0; VIndex < VInfo.CtrlVertCount; ++VIndex, CtrlHullPtr += ControlPointDimension)
+			{
+				OpenNurbssSurface.GetCV(UIndex, VIndex, pt_style, CtrlHullPtr);
+			}
+		}
+	}
+
+	// CT doesn't allow weights < 0
+	void FixNegativeWeights()
+	{
+		if (OpenNurbssSurface.IsRational())
+		{
+			bool bHasNegativeWeight = false;
+			for (uint32 PointIndex = 0; PointIndex < UInfo.CtrlVertCount * VInfo.CtrlVertCount; ++PointIndex)
+			{
+				double PointDataOffset = PointIndex * ControlPointDimension;
+				double Weight = ControlPoints[PointDataOffset + ControlPointDimension - 1];
+				if (Weight < 0.)
+				{
+					bHasNegativeWeight = true;
+					break;
+				}
+			}
+
+			if (bHasNegativeWeight)
+			{
+				UInfo.IncreaseDegree();
+				VInfo.IncreaseDegree();
+				BuildHull();
+			}
+		}
+	}
+
+	// CT doesn't allow multiplicity < order
+	void FixUnsupportedMultiplicity()
+	{
+		bool bUOrderIncreased = UInfo.FixMultiplicity();
+		bool bVOrderIncreased = VInfo.FixMultiplicity();
+		if (bUOrderIncreased || bVOrderIncreased)
+		{
+			BuildHull();
+	}
+	}
+
+	void FixUnsupportedParameters()
+	{
+		FixNegativeWeights();
+		FixUnsupportedMultiplicity();
+	}
+};
+
+
+uint64 FBRepToKernelIOBodyTranslator::CreateCTSurface(ON_NurbsSurface& Surface)
+{
+	if (Surface.Dimension() < 3)
+		return 0;
+
+	SurfaceInfo SurfaceInfo(Surface);
+	SurfaceInfo.FixUnsupportedParameters();
+
+	CADLibrary::FNurbsSurface CTSurface;
+
+	CTSurface.OrderU = SurfaceInfo.UInfo.Order;
+	CTSurface.OrderV = SurfaceInfo.VInfo.Order;
+	CTSurface.KnotSizeU = SurfaceInfo.UInfo.KnotSize;
+	CTSurface.KnotSizeV = SurfaceInfo.VInfo.KnotSize;
+	CTSurface.ControlPointDimension = SurfaceInfo.ControlPointDimension;
+	CTSurface.ControlPointSizeU = SurfaceInfo.UInfo.CtrlVertCount;
+	CTSurface.ControlPointSizeV = SurfaceInfo.VInfo.CtrlVertCount;
+
+	CTSurface.KnotValuesU = MoveTemp(SurfaceInfo.UInfo.Knots);
+	CTSurface.KnotValuesV = MoveTemp(SurfaceInfo.VInfo.Knots);
+	CTSurface.KnotMultiplicityU = MoveTemp(SurfaceInfo.UInfo.CTKnotMultiplicityArray);
+	CTSurface.KnotMultiplicityV = MoveTemp(SurfaceInfo.VInfo.CTKnotMultiplicityArray);
+	CTSurface.ControlPoints = MoveTemp(SurfaceInfo.ControlPoints);
+
+	uint64 CTSurfaceID = 0;
+	return CADLibrary::CTKIO_CreateNurbsSurface(CTSurface, CTSurfaceID) ? CTSurfaceID : 0;
+}
+
+void FBRepToKernelIOBodyTranslator::CreateCTFace_internal(const ON_BrepFace& Face, TArray<uint64>& dest, ON_BoundingBox& outerBBox, ON_NurbsSurface& Surface, bool ignoreInner)
+{
+	uint64 SurfaceID = CreateCTSurface(Surface);
+	if (SurfaceID == 0)
+		return;
+
+	ON_BoundingBox realOuterBBox = Face.OuterLoop()->BoundingBox();
+	bool outerRedefined = realOuterBBox != outerBBox;
+
+	int LoopCount = ignoreInner ? 1 : Face.LoopCount();
+	TArray<uint64> Loops;
+	Loops.Reserve(LoopCount);
+	for (int LoopIndex = 0; LoopIndex < LoopCount; ++LoopIndex)
+	{
+		const ON_BrepLoop& on_Loop = *Face.Loop(LoopIndex);
+		if (!on_Loop.IsValid())
+			continue;
+
+		ON_BrepLoop::TYPE onLoopType = on_Loop.m_type;
+		bool bIsOuter = (onLoopType == ON_BrepLoop::TYPE::outer);
+		//assert(bIsOuter == (LoopIndex == 0));
+
+		int TrimCount = on_Loop.TrimCount();
+		TArray<uint64> Coedges;
+		Coedges.Reserve(TrimCount);
+		for (int i = 0; i < TrimCount; ++i)
+		{
+			ON_BrepTrim& Trim = *on_Loop.Trim(i);
+
+			ON_BrepEdge* on_edge = Trim.Edge();
+			if (on_edge == nullptr)
+				continue;
+
+			ON_NurbsCurve nurbs_curve;
+			int nurbFormSuccess = Trim.GetNurbForm(nurbs_curve); // 0:Nok 1:Ok 2:OkBut
+			if (nurbFormSuccess == 0)
+				continue;
+
+			uint64 NewCoedge;
+			if (!CADLibrary::CTKIO_CreateCoedge(Trim.m_bRev3d, NewCoedge))
+			{
+				continue;
+			}
+
+			BrepTrimToCoedge[Trim.m_trim_index] = NewCoedge;
+
+			// Find a Trim that use this edge, that is not current Trim
+			// If the Trim has been converted into a coedge, link both Trims
+			for (int32 Index = 0; Index < on_edge->m_ti.Count(); ++Index)
+			{
+				int32 LinkedEdgeIndex = on_edge->m_ti[Index];
+				if (LinkedEdgeIndex == Trim.m_trim_index)
+				{
+					continue;
+				}
+
+				uint64 LinkedCoedgeId = BrepTrimToCoedge[LinkedEdgeIndex];
+				if (LinkedCoedgeId)
+				{
+					CADLibrary::CTKIO_MatchCoedges(LinkedCoedgeId, NewCoedge);
+					break;
+				}
+			}
+
+			CADLibrary::FNurbsCurve CTCurve;
+			CTCurve.Order = nurbs_curve.Order();
+
+			int KnotCount = nurbs_curve.KnotCount();
+			int CtrlVertCount = nurbs_curve.CVCount();
+			CTCurve.KnotSize = CTCurve.Order + CtrlVertCount; // cvCount + degree - 1 for OpenNurb, cvCount + degree + 1 for OpenNurb,
+
+			// knot data
+			CTCurve.KnotValues.SetNum(CTCurve.KnotSize);
+			CTCurve.KnotValues[0] = nurbs_curve.SuperfluousKnot(0);
+			for (int j = 0; j < KnotCount; ++j)
+			{
+				CTCurve.KnotValues[j + 1] = nurbs_curve.Knot(j);
+			}
+			CTCurve.KnotValues.Last() = nurbs_curve.SuperfluousKnot(1);
+
+			// Control hull
+			CTCurve.ControlPointDimension = nurbs_curve.CVSize(); // = IsRational() ? Dim()+1 : Dim()
+			CTCurve.ControlPointSize = nurbs_curve.CVCount();
+			CTCurve.ControlPoints.SetNum(CTCurve.ControlPointSize * CTCurve.ControlPointDimension);
+			double* ControlPoints = CTCurve.ControlPoints.GetData();
+			for (uint32 j = 0; j < CTCurve.ControlPointSize; ++j, ControlPoints += CTCurve.ControlPointDimension)
+			{
+				nurbs_curve.GetCV(j, nurbs_curve.IsRational() ? ON::point_style::euclidean_rational : ON::point_style::not_rational, ControlPoints);
+			}
+			if (outerRedefined && bIsOuter)
+			{
+				for (int32 j = 0; j < CTCurve.ControlPoints.Num(); j += CTCurve.ControlPointSize)
+				{
+					CTCurve.ControlPoints[j] = min(max(CTCurve.ControlPoints[j], outerBBox.m_min.x), outerBBox.m_max.x);
+				}
+			}
+
+			// knot multiplicity (ignored as knots are stored multiple times already)
+			CTCurve.KnotMultiplicity.Init(1, CTCurve.KnotSize);
+
+			ON_Interval dom = nurbs_curve.Domain();
+			if (!CADLibrary::CTKIO_SetUVCurve(CTCurve, dom.m_t[0], dom.m_t[1], NewCoedge))
+			{
+				continue;
+			}
+
+			Coedges.Add(NewCoedge);
+		}
+
+		uint64 Loop;
+		if (!CADLibrary::CTKIO_CreateLoop(Coedges, Loop))
+		{
+			continue;
+		}
+
+		Loops.Add(Loop);
+	}
+
+	uint64 FaceID;
+	if (CADLibrary::CTKIO_CreateFace(SurfaceID, true, Loops, FaceID))
+	{
+		dest.Add(FaceID);
+	}
+}
+
+void FBRepToKernelIOBodyTranslator::CreateCTFace(const ON_BrepFace& Face, TArray<uint64>& dest)
+{
+	const ON_BrepLoop* OuterLoop = Face.OuterLoop();
+	if (OuterLoop == nullptr)
+	{
+		return;
+	}
+
+	ON_NurbsSurface Surface;
+	Face.NurbsSurface(&Surface);
+
+	ON_BoundingBox outerBBox = Face.OuterLoop()->BoundingBox();
+
+#if FIX_HOLE_IN_WHOLE_FACE
+	int LoopCount = Face.LoopCount();
+	bool bBadLoopHack = BRep.LoopIsSurfaceBoundary(OuterLoop->m_loop_index) && LoopCount >= 2;
+	if (bBadLoopHack)
+	{
+		TArray<uint64> Coedges;
+		int TrimCount = OuterLoop->TrimCount();
+		bool bHasSingularTrim = false;
+		for (int i = 0; i < TrimCount; ++i)
+		{
+			ON_BrepTrim& Trim = *OuterLoop->Trim(i);
+			bHasSingularTrim |= (Trim.m_type == ON_BrepTrim::TYPE::singular);
+		}
+		bBadLoopHack &= bHasSingularTrim;
+	}
+
+	if (bBadLoopHack)
+	{
+		// we try to split in 2 faces: one west one east.
+		double span = outerBBox.Diagonal().x;
+
+		ON_BoundingBox innerBBox = Face.Loop(1)->BoundingBox();
+		for (int LoopIndex = 2; LoopIndex < LoopCount; ++LoopIndex)
+		{
+			const ON_BrepLoop& innerLoop = *Face.Loop(LoopIndex);
+			innerBBox.Union(innerLoop.BoundingBox());
+		}
+		double spaceWest = innerBBox.Min().x - outerBBox.Min().x;
+		double spaceEast = outerBBox.Max().x - innerBBox.Max().x;
+		double e = 0.01;
+		if (spaceEast < e*span && spaceWest < e*span)
+		{
+			// can't split... We could ignore inner loops:
+			CreateCTFace_internal(Face, dest, outerBBox, Surface, true);
+			return;
+		}
+
+		// split the outer bbox in a normal part (with outer) and a 'rest' (with no iner)
+		ON_BoundingBox outerTrimedPartBBox = outerBBox;
+		if (spaceEast > spaceWest)
+		{
+			outerBBox.m_max.x -= 0.5 * spaceEast;
+			outerTrimedPartBBox.m_min.x = outerBBox.m_max.x;
+		}
+		else
+		{
+			outerBBox.m_min.x += 0.5 * spaceWest;
+			outerTrimedPartBBox.m_max.x = outerBBox.m_min.x;
+		}
+
+		// We now need to create a bonus face to handle the hole we just created
+		CreateCTFace_internal(Face, dest, outerTrimedPartBBox, Surface, true);
+	}
+#endif // FIX_HOLE_IN_WHOLE_FACE
+
+	CreateCTFace_internal(Face, dest, outerBBox, Surface, false);
+}
+
+uint64 FBRepToKernelIOBodyTranslator::CreateBody(const ON_3dVector& Offset)
+{
+	BrepTrimToCoedge.SetNumZeroed(BRep.m_T.Count());
+
+	double boxmin[3];
+	double boxmax[3];
+	BRep.GetBBox(boxmin, boxmax);
+
+	BRep.Translate(Offset);
+	BRep.GetBBox(boxmin, boxmax);
+
+	// Create ct faces
+	BRep.FlipReversedSurfaces();
+	int FaceCount = BRep.m_F.Count();
+	TArray<uint64> FaceList;
+	FaceList.Reserve(FaceCount);
+	for (int index = 0; index < FaceCount; index++)
+	{
+		const ON_BrepFace& On_face = BRep.m_F[index];
+		CreateCTFace(On_face, FaceList);
+	}
+	BRep.Translate(-Offset);
+
+	if (FaceList.Num() == 0)
+	{
+		return 0;
+	}
+
+	// Create body from faces
+	uint64 BodyID;
+	return CADLibrary::CTKIO_CreateBody(FaceList, BodyID) ? BodyID : 0;
+}
+
+bool FOpenNurbsBRepToCoretechConverter::AddBRep(ON_Brep& Brep, const ON_3dVector& Offset)
+{
+	if (!IsCoreTechSessionValid())
+	{
+		ensureMsgf(false, TEXT("bad session init"));
+		return false;
+	}
+
+	FBRepToKernelIOBodyTranslator BodyTranslator(Brep);
+	uint64 BodyID = BodyTranslator.CreateBody(Offset);
+
+	return BodyID ? CADLibrary::CTKIO_AddBodies({ BodyID }, MainObjectId) : false;
+}
+
+#endif // defined(USE_OPENNURBS)
diff --git a/Engine/Plugins/Enterprise/DatasmithCADImporter/Source/DatasmithOpenNurbsTranslator/Private/OpenNurbsBRepToCoretechConverter.h b/Engine/Plugins/Enterprise/DatasmithCADImporter/Source/DatasmithOpenNurbsTranslator/Private/OpenNurbsBRepToCoretechConverter.h
new file mode 100644
--- /dev/null
+++ b/Engine/Plugins/Enterprise/DatasmithCADImporter/Source/DatasmithOpenNurbsTranslator/Private/OpenNurbsBRepToCoretechConverter.h
@@ -0,0 +1,65 @@
+// Copyright Epic Games, Inc. All Rights Reserved.
+
+#pragma once
+
+#include "CoreMinimal.h"
+
+#ifdef USE_OPENNURBS
+
+#include "OpenNurbsBRepConverter.h"
+#include "CADModelToCoretechConverterBase.h"
+#include "CTSession.h"
+
+class ON_Brep;
+class ON_NurbsSurface;
+class ON_BrepFace;
+class ON_BoundingBox;
+class ON_3dVector;
+
+class FBRepToKernelIOBodyTranslator
+{
+public:
+	FBRepToKernelIOBodyTranslator(ON_Brep& InBRep)
+		: BRep(InBRep)
+	{}
+
+	uint64 CreateBody(const ON_3dVector& ON_Offset);
+private:
+	uint64 CreateCTSurface(ON_NurbsSurface& Surface);
+	void CreateCTFace(const ON_BrepFace& Face, TArray<uint64>& dest);
+	void CreateCTFace_internal(const ON_BrepFace& Face, TArray<uint64>& dest, ON_BoundingBox& outerBBox, ON_NurbsSurface& Surface, bool ignoreInner);
+
+private:
+	ON_Brep& BRep;
+
+	// BRep.m_T is an array that store all Trims, so we can also use an array to make the map between Trim index to Coedge Id
+	TArray<uint64> BrepTrimToCoedge;
+};
+
+class FOpenNurbsBRepToCoretechConverter : public FCADModelToCoretechConverterBase, public IOpenNurbsBRepConverter
+{
+public:
+
+	/**
+	 * Make sure CT is initialized, and a main object is ready.
+	 * Handle input file unit and an output unit
+	 * @param InOwner
+	 */
+	explicit FOpenNurbsBRepToCoretechConverter(const TCHAR* InOwner, const CADLibrary::FImportParameters& ImportParameters)
+		: FCADModelToCoretechConverterBase(InOwner, ImportParameters)
+	{
+	}
+
+	virtual ~FOpenNurbsBRepToCoretechConverter() = default;
+
+	/**
+	 * Set BRep to tessellate, offsetting it prior to tessellation(used to set mesh pivot at the center of the surface bounding box)
+	 *
+	 * @param  Brep	a brep to tessellate
+	 * @param  Offset translate brep by this value before tessellating 
+	 */
+	bool AddBRep(ON_Brep& Brep, const ON_3dVector& Offset);
+	
+};
+
+#endif // USE_OPENNURBS
\ No newline at end of file
diff --git a/Engine/Plugins/Enterprise/DatasmithCADImporter/Source/DatasmithOpenNurbsTranslator/Private/RhinoCoretechWrapper.cpp b/Engine/Plugins/Enterprise/DatasmithCADImporter/Source/DatasmithOpenNurbsTranslator/Private/RhinoCoretechWrapper.cpp
deleted file mode 100644
diff --git a/Engine/Plugins/Enterprise/DatasmithCADImporter/Source/DatasmithOpenNurbsTranslator/Private/RhinoCoretechWrapper.h b/Engine/Plugins/Enterprise/DatasmithCADImporter/Source/DatasmithOpenNurbsTranslator/Private/RhinoCoretechWrapper.h
deleted file mode 100644
diff --git a/Engine/Plugins/Enterprise/DatasmithCADImporter/Source/DatasmithWireTranslator/Private/AliasModelToCADKernelConverter.cpp b/Engine/Plugins/Enterprise/DatasmithCADImporter/Source/DatasmithWireTranslator/Private/AliasModelToCADKernelConverter.cpp
--- a/Engine/Plugins/Enterprise/DatasmithCADImporter/Source/DatasmithWireTranslator/Private/AliasModelToCADKernelConverter.cpp
+++ b/Engine/Plugins/Enterprise/DatasmithCADImporter/Source/DatasmithWireTranslator/Private/AliasModelToCADKernelConverter.cpp
@@ -46,50 +46,50 @@ using namespace CADKernel;
 namespace AliasToCADKernelUtils
 {
 	template<typename Surface_T>
-	TSharedPtr<FSurface> AddNurbsSurface(double GeometricTolerance, Surface_T& AliasSurface, EAliasObjectReference InObjectReference, const AlMatrix4x4& InAlMatrix)
+	TSharedPtr<FSurface> AddNURBSSurface(double GeometricTolerance, Surface_T& AliasSurface, EAliasObjectReference InObjectReference, const AlMatrix4x4& InAlMatrix)
 	{
-		FNurbsSurfaceHomogeneousData NurbsData;
-		NurbsData.bSwapUV= true;
-		NurbsData.bIsRational = true;
+		FNurbsSurfaceHomogeneousData NURBSData;
+		NURBSData.bSwapUV= true;
+		NURBSData.bIsRational = true;
 
-		NurbsData.PoleUCount = AliasSurface.uNumberOfCVsInclMultiples();
-		NurbsData.PoleVCount = AliasSurface.vNumberOfCVsInclMultiples();
+		NURBSData.PoleUCount = AliasSurface.uNumberOfCVsInclMultiples();
+		NURBSData.PoleVCount = AliasSurface.vNumberOfCVsInclMultiples();
 
-		NurbsData.UDegree = AliasSurface.uDegree();  // U order of the surface
-		NurbsData.VDegree = AliasSurface.vDegree();  // V order of the surface
+		NURBSData.UDegree = AliasSurface.uDegree();  // U order of the surface
+		NURBSData.VDegree = AliasSurface.vDegree();  // V order of the surface
 
 		int32 KnotSizeU = AliasSurface.realuNumberOfKnots() + 2;
 		int32 KnotSizeV = AliasSurface.realvNumberOfKnots() + 2;
 
-		NurbsData.UNodalVector.SetNumUninitialized(KnotSizeU);
-		NurbsData.VNodalVector.SetNumUninitialized(KnotSizeV);
+		NURBSData.UNodalVector.SetNumUninitialized(KnotSizeU);
+		NURBSData.VNodalVector.SetNumUninitialized(KnotSizeV);
 
-		AliasSurface.realuKnotVector(NurbsData.UNodalVector.GetData() + 1);
-		AliasSurface.realvKnotVector(NurbsData.VNodalVector.GetData() + 1);
+		AliasSurface.realuKnotVector(NURBSData.UNodalVector.GetData() + 1);
+		AliasSurface.realvKnotVector(NURBSData.VNodalVector.GetData() + 1);
 
-		NurbsData.UNodalVector[0] = NurbsData.UNodalVector[1];
-		NurbsData.UNodalVector[KnotSizeU - 1] = NurbsData.UNodalVector[KnotSizeU - 2];
-		NurbsData.VNodalVector[0] = NurbsData.VNodalVector[1];
-		NurbsData.VNodalVector[KnotSizeV - 1] = NurbsData.VNodalVector[KnotSizeV - 2];
+		NURBSData.UNodalVector[0] = NURBSData.UNodalVector[1];
+		NURBSData.UNodalVector[KnotSizeU - 1] = NURBSData.UNodalVector[KnotSizeU - 2];
+		NURBSData.VNodalVector[0] = NURBSData.VNodalVector[1];
+		NURBSData.VNodalVector[KnotSizeV - 1] = NURBSData.VNodalVector[KnotSizeV - 2];
 
-		const int32 CoordinateCount = NurbsData.PoleUCount * NurbsData.PoleVCount * 4;
-		NurbsData.HomogeneousPoles.SetNumUninitialized(CoordinateCount);
+		const int32 CoordinateCount = NURBSData.PoleUCount * NURBSData.PoleVCount * 4;
+		NURBSData.HomogeneousPoles.SetNumUninitialized(CoordinateCount);
 
 		if (InObjectReference == EAliasObjectReference::WorldReference)
 		{
-			AliasSurface.CVsWorldPositionInclMultiples(NurbsData.HomogeneousPoles.GetData());
+			AliasSurface.CVsWorldPositionInclMultiples(NURBSData.HomogeneousPoles.GetData());
 		}
 		else if (InObjectReference == EAliasObjectReference::ParentReference)
 		{
 			AlTM TranformMatrix(InAlMatrix);
-			AliasSurface.CVsAffectedPositionInclMultiples(TranformMatrix, NurbsData.HomogeneousPoles.GetData());
+			AliasSurface.CVsAffectedPositionInclMultiples(TranformMatrix, NURBSData.HomogeneousPoles.GetData());
 		}
 		else  // EAliasObjectReference::LocalReference
 		{
-			AliasSurface.CVsUnaffectedPositionInclMultiples(NurbsData.HomogeneousPoles.GetData());
+			AliasSurface.CVsUnaffectedPositionInclMultiples(NURBSData.HomogeneousPoles.GetData());
 		}
 
-		return FEntity::MakeShared<FNURBSSurface>(GeometricTolerance, NurbsData);
+		return FEntity::MakeShared<FNURBSSurface>(GeometricTolerance, NURBSData);
 	}
 }
 
@@ -97,8 +97,6 @@ TSharedPtr<FTopologicalEdge> FAliasModelToCADKernelConverter::AddEdge(const AlTr
 {
 	FNurbsCurveData NurbsCurveData;
 
-	curveFormType Form = AliasTrimCurve.form();
-
 	NurbsCurveData.Degree = AliasTrimCurve.degree();
 	int ControlPointCount = AliasTrimCurve.numberOfCVs();
 
@@ -155,15 +153,6 @@ TSharedPtr<FTopologicalLoop> FAliasModelToCADKernelConverter::AddLoop(const AlTr
 	TArray<TSharedPtr<FTopologicalEdge>> Edges;
 	TArray<CADKernel::EOrientation> Directions;
 
-	int32 EdgeCount = 0;
-	for(TUniquePtr<AlTrimCurve> TrimCurve(TrimBoundary.firstCurve()); TrimCurve.IsValid(); TrimCurve = TUniquePtr<AlTrimCurve>(TrimCurve->nextCurve()))
-	{
-		EdgeCount++;
-	}
-
-	Edges.Reserve(EdgeCount);
-	Directions.Reserve(EdgeCount);
-
 	for (TUniquePtr<AlTrimCurve> TrimCurve(TrimBoundary.firstCurve()); TrimCurve.IsValid(); TrimCurve = TUniquePtr<AlTrimCurve>(TrimCurve->nextCurve()))
 	{
 		TSharedPtr<FTopologicalEdge> Edge = AddEdge(*TrimCurve, CarrierSurface);
@@ -185,50 +174,20 @@ TSharedPtr<FTopologicalLoop> FAliasModelToCADKernelConverter::AddLoop(const AlTr
 
 TSharedPtr<FTopologicalFace> FAliasModelToCADKernelConverter::AddTrimRegion(const AlTrimRegion& TrimRegion, EAliasObjectReference InObjectReference, const AlMatrix4x4& InAlMatrix, bool bInOrientation)
 {
-	TSharedPtr<FSurface> Surface = AliasToCADKernelUtils::AddNurbsSurface(GeometricTolerance, TrimRegion, InObjectReference, InAlMatrix);
+	TSharedPtr<FSurface> Surface = AliasToCADKernelUtils::AddNURBSSurface(GeometricTolerance, TrimRegion, InObjectReference, InAlMatrix);
 	if (!Surface.IsValid())
 	{
 		return TSharedPtr<FTopologicalFace>();
 	}
 
 	TArray<TSharedPtr<FTopologicalLoop>> Loops;
-	int32 LoopCount = 0;
-	for (TUniquePtr<AlTrimBoundary> TrimBoundary(TrimRegion.firstBoundary()); TrimBoundary.IsValid(); TrimBoundary = TUniquePtr<AlTrimBoundary>(TrimBoundary->nextBoundary()))
-	{
-		LoopCount++;
-	}
-	Loops.Reserve(LoopCount);
-
 	for (TUniquePtr<AlTrimBoundary> TrimBoundary(TrimRegion.firstBoundary()); TrimBoundary.IsValid(); TrimBoundary = TUniquePtr<AlTrimBoundary>(TrimBoundary->nextBoundary()))
 	{
 		TSharedPtr<FTopologicalLoop> Loop = AddLoop(*TrimBoundary, Surface);
-		if (!Loop.IsValid())
+		if (Loop.IsValid())
 		{
-			continue;
+			Loops.Add(Loop);
 		}
-
-		TArray<FPoint2D> LoopSampling;
-		Loop->Get2DSampling(LoopSampling);
-		FAABB2D Boundary;
-		Boundary += LoopSampling;
-		Loop->Boundary.Set(Boundary.GetMin(), Boundary.GetMax());
-
-		// Check if the loop is not composed with only degenerated edge
-		bool bDegeneratedLoop = true;
-		for (const FOrientedEdge& Edge : Loop->GetEdges())
-		{
-			if (!Edge.Entity->IsDegenerated())
-			{
-				bDegeneratedLoop = false;
-				break;
-			}
-		}
-		if (bDegeneratedLoop)
-		{
-			continue;
-		}
-
-		Loops.Add(Loop);
 	}
 
 	if (Loops.Num() == 0)
@@ -260,7 +219,7 @@ void FAliasModelToCADKernelConverter::AddFace(const AlSurface& Surface, EAliasOb
 		return;
 	}
 
-	TSharedPtr<FSurface> CADKernelSurface = AliasToCADKernelUtils::AddNurbsSurface(GeometricTolerance, Surface, InObjectReference, InAlMatrix);
+	TSharedPtr<FSurface> CADKernelSurface = AliasToCADKernelUtils::AddNURBSSurface(GeometricTolerance, Surface, InObjectReference, InAlMatrix);
 	if (CADKernelSurface.IsValid())
 	{
 		TSharedRef<FTopologicalFace> Face = FEntity::MakeShared<FTopologicalFace>(CADKernelSurface);
diff --git a/Engine/Plugins/Enterprise/DatasmithCADImporter/Source/DatasmithWireTranslator/Private/AliasModelToCADKernelConverter.h b/Engine/Plugins/Enterprise/DatasmithCADImporter/Source/DatasmithWireTranslator/Private/AliasModelToCADKernelConverter.h
--- a/Engine/Plugins/Enterprise/DatasmithCADImporter/Source/DatasmithWireTranslator/Private/AliasModelToCADKernelConverter.h
+++ b/Engine/Plugins/Enterprise/DatasmithCADImporter/Source/DatasmithWireTranslator/Private/AliasModelToCADKernelConverter.h
@@ -32,7 +32,8 @@ class FAliasModelToCADKernelConverter : public FCADModelToCADKernelConverterBase
 {
 
 public:
-	FAliasModelToCADKernelConverter()
+	FAliasModelToCADKernelConverter(CADLibrary::FImportParameters InImportParameters)
+		: FCADModelToCADKernelConverterBase(InImportParameters)
 	{
 	}
 
diff --git a/Engine/Plugins/Enterprise/DatasmithCADImporter/Source/DatasmithWireTranslator/Private/AliasModelToCoretechConverter.cpp b/Engine/Plugins/Enterprise/DatasmithCADImporter/Source/DatasmithWireTranslator/Private/AliasModelToCoretechConverter.cpp
--- a/Engine/Plugins/Enterprise/DatasmithCADImporter/Source/DatasmithWireTranslator/Private/AliasModelToCoretechConverter.cpp
+++ b/Engine/Plugins/Enterprise/DatasmithCADImporter/Source/DatasmithWireTranslator/Private/AliasModelToCoretechConverter.cpp
@@ -85,7 +85,6 @@ uint64 FAliasModelToCoretechConverter::AddTrimCurve(const AlTrimCurve& TrimCurve
 {
 	FNurbsCurve CTCurve;
 
-	curveFormType Form = TrimCurve.form();
 	CTCurve.Order = TrimCurve.degree() + 1;
 	CTCurve.ControlPointSize = TrimCurve.numberOfCVs();
 	CTCurve.ControlPointDimension = 3;
diff --git a/Engine/Plugins/Enterprise/DatasmithCADImporter/Source/DatasmithWireTranslator/Private/AliasModelToCoretechConverter.h b/Engine/Plugins/Enterprise/DatasmithCADImporter/Source/DatasmithWireTranslator/Private/AliasModelToCoretechConverter.h
--- a/Engine/Plugins/Enterprise/DatasmithCADImporter/Source/DatasmithWireTranslator/Private/AliasModelToCoretechConverter.h
+++ b/Engine/Plugins/Enterprise/DatasmithCADImporter/Source/DatasmithWireTranslator/Private/AliasModelToCoretechConverter.h
@@ -22,13 +22,9 @@ class FAliasModelToCoretechConverter : public FCADModelToCoretechConverterBase,
 {
 public:
 
-	FAliasModelToCoretechConverter(const TCHAR* InOwner)
-		: FCADModelToCoretechConverterBase(InOwner)
+	FAliasModelToCoretechConverter(const TCHAR* InOwner, CADLibrary::FImportParameters InImportParameters)
+		: FCADModelToCoretechConverterBase(InOwner, InImportParameters)
 	{
-		// Unit for CoreTech session is set to cm, 0.01, because Wire's unit is cm. Consequently, Scale factor is set to 1.
-		ImportParams.MetricUnit = 0.01;
-		ImportParams.ScaleFactor = 1;
-		ImportParams.bEnableKernelIOTessellation = true;
 	}
 
 	virtual bool AddBRep(AlDagNode& DagNode, EAliasObjectReference ObjectReference) override;
diff --git a/Engine/Plugins/Enterprise/DatasmithCADImporter/Source/DatasmithWireTranslator/Private/DatasmithWireTranslator.cpp b/Engine/Plugins/Enterprise/DatasmithCADImporter/Source/DatasmithWireTranslator/Private/DatasmithWireTranslator.cpp
--- a/Engine/Plugins/Enterprise/DatasmithCADImporter/Source/DatasmithWireTranslator/Private/DatasmithWireTranslator.cpp
+++ b/Engine/Plugins/Enterprise/DatasmithCADImporter/Source/DatasmithWireTranslator/Private/DatasmithWireTranslator.cpp
@@ -10,6 +10,7 @@
 #include "DatasmithUtils.h"
 #include "HAL/ConsoleManager.h"
 #include "DatasmithWireTranslatorModule.h"
+#include "HAL/ConsoleManager.h"
 #include "IDatasmithSceneElements.h"
 #include "Misc/FileHelper.h"
 #include "Misc/Paths.h"
@@ -149,15 +150,19 @@ public:
 			bEnableKernelIOTessellation = CVar->GetInt() != 0;
 		}
 
+		CADLibrary::FImportParameters ImportParameters;
+		ImportParameters.MetricUnit = 0.01;
+		ImportParameters.ScaleFactor = 1;
+
 		if (bEnableKernelIOTessellation)
 		{
-			TSharedRef<FAliasModelToCoretechConverter> AliasToCoretechConverter = MakeShared<FAliasModelToCoretechConverter>(TEXT("Al2CTSharedSession"));
+			TSharedRef<FAliasModelToCoretechConverter> AliasToCoretechConverter = MakeShared<FAliasModelToCoretechConverter>(TEXT("Al2CTSharedSession"), ImportParameters);
 			CADModelConverter = AliasToCoretechConverter;
 			AliasBRepConverter = AliasToCoretechConverter;
 		}
 		else
 		{
-			TSharedRef<FAliasModelToCADKernelConverter> AliasToCoretechConverter = MakeShared<FAliasModelToCADKernelConverter>();
+			TSharedRef<FAliasModelToCADKernelConverter> AliasToCoretechConverter = MakeShared<FAliasModelToCADKernelConverter>(ImportParameters);
 			CADModelConverter = AliasToCoretechConverter;
 			AliasBRepConverter = AliasToCoretechConverter;
 		}
@@ -177,11 +182,6 @@ public:
 
 	bool LoadStaticMesh(const TSharedRef<IDatasmithMeshElement> MeshElement, FDatasmithMeshElementPayload& OutMeshPayload, const FDatasmithTessellationOptions& InTessellationOptions);
 
-	void AddSurfaceDataForMesh(const TCHAR* InFilePath, const FMeshParameters& InMeshParameters, const FDatasmithTessellationOptions& InTessellationOptions, FDatasmithMeshElementPayload& OutMeshPayload) const
-	{
-		CADModelConverter->AddSurfaceDataForMesh(InFilePath, InMeshParameters, InTessellationOptions, OutMeshPayload);
-	}
-
 private:
 
 	struct FDagNodeInfo
@@ -2031,7 +2031,7 @@ bool FWireTranslatorImpl::LoadStaticMesh(const TSharedRef<IDatasmithMeshElement>
 	if (TOptional<FMeshDescription> Mesh = GetMeshDescription(MeshElement, MeshParameters))
 	{
 		OutMeshPayload.LodMeshes.Add(MoveTemp(Mesh.GetValue()));
-		AddSurfaceDataForMesh(MeshElement->GetFile(), MeshParameters, InTessellationOptions, OutMeshPayload);
+		CADModelConverter->AddSurfaceDataForMesh(MeshElement->GetFile(), MeshParameters, InTessellationOptions, OutMeshPayload);
 	}
 	return OutMeshPayload.LodMeshes.Num() > 0;
 }
diff --git a/Engine/Plugins/Enterprise/DatasmithCADImporter/Source/ParametricSurface/Public/CADModelConverter.h b/Engine/Plugins/Enterprise/DatasmithCADImporter/Source/ParametricSurface/Public/CADModelConverter.h
--- a/Engine/Plugins/Enterprise/DatasmithCADImporter/Source/ParametricSurface/Public/CADModelConverter.h
+++ b/Engine/Plugins/Enterprise/DatasmithCADImporter/Source/ParametricSurface/Public/CADModelConverter.h
@@ -56,7 +56,12 @@ public:
 	 * @param StitchingTechnique : CAD topology correction technique
 	 * @param bScaleUVMap : Scale the UV map to a world unit.
 	 */
-	virtual void SetImportParameters(float ChordTolerance, float MaxEdgeLength, float NormalTolerance, CADLibrary::EStitchingTechnique StitchingTechnique, bool bScaleUVMap) = 0;
+	virtual void SetImportParameters(double ChordTolerance, double MaxEdgeLength, double NormalTolerance, CADLibrary::EStitchingTechnique StitchingTechnique, bool bScaleUVMap) = 0;
+
+	//virtual void SetScaleFactor(double NewScaleValue) = 0;
+	virtual void SetMetricUnit(double NewMetricUnit) = 0;
+	virtual double GetScaleFactor() const = 0;
+	virtual double GetMetricUnit() const = 0;
 
 	virtual bool IsSessionValid() = 0;
 
diff --git a/Engine/Source/Runtime/Datasmith/CADKernel/Private/CADKernel/Geo/Surfaces/NURBSSurface.cpp b/Engine/Source/Runtime/Datasmith/CADKernel/Private/CADKernel/Geo/Surfaces/NURBSSurface.cpp
--- a/Engine/Source/Runtime/Datasmith/CADKernel/Private/CADKernel/Geo/Surfaces/NURBSSurface.cpp
+++ b/Engine/Source/Runtime/Datasmith/CADKernel/Private/CADKernel/Geo/Surfaces/NURBSSurface.cpp
@@ -132,45 +132,5 @@ void FNURBSSurface::FillNurbs(FNurbsSurfaceHomogeneousData& NurbsData)
 	Finalize();
 }
 
-void FNURBSSurface::FillNurbs(FNurbsSurfaceData& NurbsData)
-{
-	bIsRational = (NurbsData.Weights.Num() > 0);
-
-	PoleUCount = NurbsData.PoleUCount;
-	PoleVCount = NurbsData.PoleVCount;
-
-	UDegree = NurbsData.UDegree;
-	VDegree = NurbsData.VDegree;
-
-	Swap(UNodalVector, NurbsData.UNodalVector);
-	Swap(VNodalVector, NurbsData.VNodalVector);
-
-	int32 PoleCount = PoleUCount * PoleVCount;
-
-	Poles.Reserve(PoleCount);
-	for (int32 Vndex = 0; Vndex < PoleVCount; ++Vndex)
-	{
-		int32 Index = Vndex;
-		for (int32 Undex = 0; Undex < PoleUCount; ++Undex, Index += PoleVCount)
-		{
-			Poles.Emplace(NurbsData.Poles[Index]);
-		}
-	}
-
-	if (bIsRational)
-	{
-		Weights.Reserve(PoleCount);
-		for (int32 Vndex = 0; Vndex < PoleVCount; ++Vndex)
-		{
-			int32 Index = Vndex;
-			for (int32 Undex = 0; Undex < PoleUCount; ++Undex, Index += PoleVCount)
-			{
-				Weights.Emplace(NurbsData.Weights[Index]);
-			}
-		}
-	}
-
-	Finalize();
-}
 
 
diff --git a/Engine/Source/Runtime/Datasmith/CADKernel/Private/CADKernel/Topo/TopologicalLoop.cpp b/Engine/Source/Runtime/Datasmith/CADKernel/Private/CADKernel/Topo/TopologicalLoop.cpp
--- a/Engine/Source/Runtime/Datasmith/CADKernel/Private/CADKernel/Topo/TopologicalLoop.cpp
+++ b/Engine/Source/Runtime/Datasmith/CADKernel/Private/CADKernel/Topo/TopologicalLoop.cpp
@@ -2,6 +2,7 @@
 
 #include "CADKernel/Topo/TopologicalLoop.h"
 
+#include "Algo/AllOf.h"
 #include "CADKernel/Core/KernelParameters.h"
 #include "CADKernel/Geo/Curves/RestrictionCurve.h"
 #include "CADKernel/Math/SlopeUtils.h"
@@ -30,20 +31,9 @@ TSharedPtr<FTopologicalLoop> FTopologicalLoop::Make(const TArray<TSharedPtr<FTop
 	LoopBoundary += LoopSampling;
 	Loop->Boundary.Set(LoopBoundary.GetMin(), LoopBoundary.GetMax());
 
-	// Check if the loop is not composed with only degenerated edge
-	bool bDegeneratedLoop = true;
-	for (const FOrientedEdge& Edge : Loop->GetEdges())
+	if (Algo::AllOf(Loop->GetEdges(), [](const FOrientedEdge& Edge) { return Edge.Entity->IsDegenerated(); }))
 	{
-		if (!Edge.Entity->IsDegenerated())
-		{
-			bDegeneratedLoop = false;
-			break;
-		}
-	}
-
-	if (bDegeneratedLoop)
-	{
-		TSharedPtr<FTopologicalLoop>();
+		return TSharedPtr<FTopologicalLoop>();
 	}
 	return Loop;
 }
@@ -391,51 +381,6 @@ void FTopologicalLoop::ReplaceEdges(TArray<FOrientedEdge>& OldEdges, TSharedPtr<
 	ensureCADKernel(false);
 }
 
-// Need to be reviewed
-/*void FTopologicalLoop::ReplaceEdgesWithMergedEdge(TArray<TSharedPtr<FTopologicalEdge>>& OldEdges, TSharedPtr<FTopologicalVertex>& MiddleVertex, TSharedPtr<FTopologicalEdge>& NewEdge)
-{
-	ensureCADKernel(OldEdges.Num() == 2);
-	ensureCADKernel(EdgeCount() != 2);
-
-	NewEdge->SetLoop(StaticCastSharedRef<FTopologicalLoop>(AsShared()));
-	for (TSharedPtr<FTopologicalEdge> Edge : OldEdges)
-	{
-		Edge->RemoveLoop();
-		Edge->GetStartVertex()->RemoveConnectedEdge(Edge.ToSharedRef());
-		Edge->GetEndVertex()->RemoveConnectedEdge(Edge.ToSharedRef());
-	}
-
-	int32 IndexFirstOldEdge = GetEdgeIndex(OldEdges[0]);
-	int32 IndexLastOldEdge = GetEdgeIndex(OldEdges[1]);
-
-	ensureCADKernel(IndexFirstOldEdge != INDEX_NONE && IndexLastOldEdge != INDEX_NONE);
-
-	Sort(IndexFirstOldEdge, IndexLastOldEdge);
-	if (IndexFirstOldEdge == 0 && IndexLastOldEdge != 1)
-	{
-		IndexLastOldEdge = 0;
-		IndexFirstOldEdge = EdgeCount() - 1;
-	}
-	int32 PreviousEdgeIndex = IndexFirstOldEdge ? IndexFirstOldEdge - 1 : EdgeCount() - 1;
-	TSharedPtr<FTopologicalVertex> PreviousVertex = Edges[PreviousEdgeIndex].Direction == EOrientation::Front ? Edges[PreviousEdgeIndex].Entity->GetEndVertex() : Edges[PreviousEdgeIndex].Entity->GetStartVertex();
-	EOrientation NewEdgeOrientation = NewEdge->GetStartVertex() == PreviousVertex ? EOrientation::Front : EOrientation::Back;
-
-	if (IndexLastOldEdge)
-	{
-		Edges.RemoveAt(IndexFirstOldEdge);
-		Edges.RemoveAt(IndexFirstOldEdge + 1);
-		Edges.EmplaceAt(IndexFirstOldEdge, NewEdge, NewEdgeOrientation);
-	}
-	else
-	{
-		Edges.RemoveAt(0);
-		Edges.Pop();
-		Edges.Emplace(NewEdge, NewEdgeOrientation);
-	}
-}*/
-
-
-
 void FTopologicalLoop::FindSurfaceCorners(TArray<TSharedPtr<FTopologicalVertex>>& OutCorners, TArray<int32>& OutStartSideIndex) const
 {
 	TArray<double> BreakValues;
diff --git a/Engine/Source/Runtime/Datasmith/CADKernel/Public/CADKernel/Geo/Surfaces/NURBSSurface.h b/Engine/Source/Runtime/Datasmith/CADKernel/Public/CADKernel/Geo/Surfaces/NURBSSurface.h
--- a/Engine/Source/Runtime/Datasmith/CADKernel/Public/CADKernel/Geo/Surfaces/NURBSSurface.h
+++ b/Engine/Source/Runtime/Datasmith/CADKernel/Public/CADKernel/Geo/Surfaces/NURBSSurface.h
@@ -23,23 +23,6 @@ namespace CADKernel
 		TArray<double> HomogeneousPoles;
 	};
 
-	struct FNurbsSurfaceData
-	{
-		bool bSwapUV;
-		int32 PoleUCount;
-		int32 PoleVCount;
-
-		int32 UDegree;
-		int32 VDegree;
-
-		TArray<double> UNodalVector;
-		TArray<double> VNodalVector;
-
-		// if Weights.Num() == 0  => bIsRational = false
-		TArray<double> Weights;
-		TArray<FPoint> Poles;
-	};
-
 	class CADKERNEL_API FNURBSSurface : public FSurface
 	{
 		friend FEntity;
@@ -225,7 +208,6 @@ namespace CADKernel
 	private:
 		void Finalize();
 		void FillNurbs(FNurbsSurfaceHomogeneousData& NurbsData);
-		void FillNurbs(FNurbsSurfaceData& NurbsData);
 	};
 
 } // namespace CADKernel
