Check if triangle is degenerated using 2 times the same vertex

diff --git a/Engine/Plugins/Enterprise/DatasmithCADImporter/Source/CADLibrary/Private/MeshDescriptionHelper.cpp b/Engine/Plugins/Enterprise/DatasmithCADImporter/Source/CADLibrary/Private/MeshDescriptionHelper.cpp
--- a/Engine/Plugins/Enterprise/DatasmithCADImporter/Source/CADLibrary/Private/MeshDescriptionHelper.cpp
+++ b/Engine/Plugins/Enterprise/DatasmithCADImporter/Source/CADLibrary/Private/MeshDescriptionHelper.cpp
@@ -454,12 +454,17 @@ namespace CADLibrary
 					VertexIndex[1] = Tessellation.PositionIndices[Tessellation.VertexIndices[Index + Orientation[1]]];
 					VertexIndex[2] = Tessellation.PositionIndices[Tessellation.VertexIndices[Index + Orientation[2]]];
 
-					// Verify the 3 input indices are not defining a degenerated triangle
 					if (VertexIndex[0] == INDEX_NONE || VertexIndex[1] == INDEX_NONE || VertexIndex[2] == INDEX_NONE)
 					{
 						continue;
 					}
 
+					// Verify the 3 input indices are not defining a degenerated triangle
+					if (VertexIndex[0] == VertexIndex[1] || VertexIndex[0] == VertexIndex[2] || VertexIndex[1] == VertexIndex[2])
+					{
+						continue;
+					}
+
 					VertexIDs[0] = VertexIdSet[VertexIndex[0]];
 					VertexIDs[1] = VertexIdSet[VertexIndex[1]];
 					VertexIDs[2] = VertexIdSet[VertexIndex[2]];
