Fix transform (FMatrix used to be in float, now in double)

diff --git a/Engine/Plugins/Enterprise/DatasmithCADImporter/Source/CADInterfaces/Private/CoreTechFileReader.cpp b/Engine/Plugins/Enterprise/DatasmithCADImporter/Source/CADInterfaces/Private/CoreTechFileReader.cpp
--- a/Engine/Plugins/Enterprise/DatasmithCADImporter/Source/CADInterfaces/Private/CoreTechFileReader.cpp
+++ b/Engine/Plugins/Enterprise/DatasmithCADImporter/Source/CADInterfaces/Private/CoreTechFileReader.cpp
@@ -567,19 +567,11 @@ namespace CADLibrary
 		}
 
 		// Ask the transformation of the instance
-		double Matrix[16];
-		if (CT_INSTANCE_IO::AskTransformation(InstanceNodeId, Matrix) == IO_OK)
+		if (CT_INSTANCE_IO::AskTransformation(InstanceNodeId, (double*) Instance.TransformMatrix.M) == IO_OK)
 		{
-			float* MatrixFloats = (float*) Instance.TransformMatrix.M;
-			for (int32 index = 0; index < 16; index++)
+			if (Instance.TransformMatrix.ContainsNaN())
 			{
-				// check if the matrix is not degenerate, otherwise return identity matrix
-				if (FMath::IsNaN(Matrix[index]) || !FMath::IsFinite(Matrix[index]))
-				{
-					Instance.TransformMatrix.SetIdentity();
-					break;
-				}
-				MatrixFloats[index] = (float)Matrix[index];
+				Instance.TransformMatrix.SetIdentity();
 			}
 		}
 	
diff --git a/Engine/Plugins/Enterprise/DatasmithCADImporter/Source/CADTools/Private/CADToolsModule.cpp b/Engine/Plugins/Enterprise/DatasmithCADImporter/Source/CADTools/Private/CADToolsModule.cpp
--- a/Engine/Plugins/Enterprise/DatasmithCADImporter/Source/CADTools/Private/CADToolsModule.cpp
+++ b/Engine/Plugins/Enterprise/DatasmithCADImporter/Source/CADTools/Private/CADToolsModule.cpp
@@ -18,7 +18,7 @@ bool FCADToolsModule::IsAvailable()
 
 uint32 FCADToolsModule::GetCacheVersion()
 {
-	return 9;
+	return 10;
 }
 
 IMPLEMENT_MODULE(FCADToolsModule, CADTools);
