The crash is due to an empty Texture Coordinate Array

diff --git a/Engine/Plugins/Enterprise/DatasmithCADImporter/Source/CADLibrary/Private/MeshDescriptionHelper.cpp b/Engine/Plugins/Enterprise/DatasmithCADImporter/Source/CADLibrary/Private/MeshDescriptionHelper.cpp
--- a/Engine/Plugins/Enterprise/DatasmithCADImporter/Source/CADLibrary/Private/MeshDescriptionHelper.cpp
+++ b/Engine/Plugins/Enterprise/DatasmithCADImporter/Source/CADLibrary/Private/MeshDescriptionHelper.cpp
@@ -1,6 +1,7 @@
 // Copyright Epic Games, Inc. All Rights Reserved.
 #include "MeshDescriptionHelper.h"
 
+#include "Algo/AnyOf.h"
 #include "CADData.h"
 #include "CADOptions.h"
 #include "DatasmithMaterialElements.h"
@@ -403,7 +404,10 @@ namespace CADLibrary
 			Material.Value = PolyGroupID;
 		}
 
-		VertexInstanceUVs.SetNumChannels(1);
+		if (Algo::AnyOf(BodyTessellation.Faces, [](const FTessellationData& FaceTessellation) { return !FaceTessellation.TexCoordArray.IsEmpty(); }))
+		{
+			VertexInstanceUVs.SetNumChannels(1);
+		}
 
 		int32 NbStep = 1;
 		if (MeshParameters.bIsSymmetric)
@@ -484,18 +488,24 @@ namespace CADLibrary
 					PatchGroups[PolygonID] = Tessellation.PatchId;
 				}
 
-				for (int32 IndexFace = 0; IndexFace < MeshVertexInstanceIDs.Num(); IndexFace += 3)
+				if (!Tessellation.TexCoordArray.IsEmpty())
 				{
-					for (int32 Index = 0; Index < TriangleCount; Index++)
+					for (int32 IndexFace = 0; IndexFace < MeshVertexInstanceIDs.Num(); IndexFace += TriangleCount)
 					{
-						const FVertexInstanceID VertexInstanceID = MeshVertexInstanceIDs[IndexFace + Orientation[Index]];
-						VertexInstanceUVs.Set(VertexInstanceID, UVChannel, Tessellation.TexCoordArray[NewFaceIndex[IndexFace + Index]]);
-
-						VertexInstanceColors[VertexInstanceID] = FLinearColor::White;
-						VertexInstanceTangents[VertexInstanceID] = FVector(ForceInitToZero);
-						VertexInstanceBinormalSigns[VertexInstanceID] = 0.0f;
+						for (int32 Index = 0; Index < TriangleCount; Index++)
+						{
+							const FVertexInstanceID& VertexInstanceID = MeshVertexInstanceIDs[IndexFace + Orientation[Index]];
+							VertexInstanceUVs.Set(VertexInstanceID, UVChannel, Tessellation.TexCoordArray[NewFaceIndex[IndexFace + Index]]);
+						}
 					}
 				}
+				
+				for (const FVertexInstanceID& VertexInstanceID : MeshVertexInstanceIDs)
+				{
+					VertexInstanceColors[VertexInstanceID] = FLinearColor::White;
+					VertexInstanceTangents[VertexInstanceID] = FVector(ForceInitToZero);
+					VertexInstanceBinormalSigns[VertexInstanceID] = 0.0f;
+				}
 
 				if (!Step)
 				{
