Rename CAD Import cvar

diff --git a/Engine/Config/ConsoleVariables.ini b/Engine/Config/ConsoleVariables.ini
--- a/Engine/Config/ConsoleVariables.ini
+++ b/Engine/Config/ConsoleVariables.ini
@@ -99,18 +99,32 @@ fx.SkipVectorVMBackendOptimizations=1
  
 ; Comment this out to turn off audio submix effects.
 ;au.BypassAllSubmixEffects=0
+
 ; Datasmith CVar
 ; Uncomment to set our own values
 
-; Activate to save temporary CAD processing file. These file will be use in a next import to avoid CAD file processing
+; CAD file parallel processing
+; Default is MaxImportThreads = 0
+; 0: multi-processing , n: multi-processing limited to n process. CADCache is mandatory. 
+; 1: - if CADCache is enable, the scene is read in a sequential mode with cache i.e. cache is used for sub-file already read, 
+;    - if CADCache is disable, the scene is read all at once 
+;ds.CADTranslator.MaxImportThreads = 0
+
+; Uncomment to enable/disable temporary CAD processing file cache. These file will be use in a next import to avoid CAD file processing. 
+; If MaxImportThreads != 1, EnableCADCache value is ignored
+; Default is enable
 ;ds.CADTranslator.EnableCADCache=1
 
-; The timer that kill the worker if the import time is unusually long. With this time control, the load of the corrupted file is canceled but the rest of the scene is imported
-;ds.CADTranslator.EnableTimeControl=1
+; Uncomment to overwrite any existing cache associated with the file being imported. 
+; Default is false
+;ds.CADTranslator.OverwriteCache=0
 
-; Activate to parallelise CAD file processing
-;ds.CADTranslator.EnableThreadedImport=1
+; Uncomment to enable/disable the timer that kill the worker if the import time is unusually long. With this time control, the load of the corrupted file 
+; is canceled but the rest of the scene is imported. Default is enable
+;ds.CADTranslator.EnableTimeControl=1
 
-; Activate to use KernelIO Tessellation
-;ds.CADTranslator.EnableKernelIOTessellation=1
+; Uncomment to enable/disable CADKernel tessellation. If disable, CAD SDK tessellator is used.
+;ds.CADTranslator.DisableCADKernelTessellation=0
 
+; Uncomment to select the SDK to load CAD files: "KernelIO" or "TechSoft". Default is "KernelIO"
+;ds.CADTranslator.CADLibrary="KernelIO"
diff --git a/Engine/Plugins/Enterprise/DatasmithCADImporter/Source/CADInterfaces/Private/CoreTechFileParser.cpp b/Engine/Plugins/Enterprise/DatasmithCADImporter/Source/CADInterfaces/Private/CoreTechFileParser.cpp
--- a/Engine/Plugins/Enterprise/DatasmithCADImporter/Source/CADInterfaces/Private/CoreTechFileParser.cpp
+++ b/Engine/Plugins/Enterprise/DatasmithCADImporter/Source/CADInterfaces/Private/CoreTechFileParser.cpp
@@ -123,7 +123,7 @@ namespace CADLibrary
 
 		bool bNeedToProceed = true;
 
-		if (ImportParameters.bEnableCacheUsage && IFileManager::Get().FileExists(*CTFilePath))
+		if (!ImportParameters.bOverwriteCache && IFileManager::Get().FileExists(*CTFilePath))
 		{
 			if (IFileManager::Get().FileExists(*MeshArchiveFilePath)) // the file has been proceed with same meshing parameters
 			{
diff --git a/Engine/Plugins/Enterprise/DatasmithCADImporter/Source/CADInterfaces/Private/CoreTechFileReader.cpp b/Engine/Plugins/Enterprise/DatasmithCADImporter/Source/CADInterfaces/Private/CoreTechFileReader.cpp
--- a/Engine/Plugins/Enterprise/DatasmithCADImporter/Source/CADInterfaces/Private/CoreTechFileReader.cpp
+++ b/Engine/Plugins/Enterprise/DatasmithCADImporter/Source/CADInterfaces/Private/CoreTechFileReader.cpp
@@ -324,7 +324,6 @@ namespace CADLibrary
 			return ECoreTechParsingResult::ProcessFailed;
 		}
 
-#ifndef DATASMITH_CAD_IGNORE_CACHE
 		if (!Context.CachePath.IsEmpty())
 		{
 			uint32 FileHash = FileDescription.GetFileHash();
@@ -338,11 +337,10 @@ namespace CADLibrary
 				CT_IO_ERROR SaveResult = CT_KERNEL_IO::SaveFile(ObjectList, *CTFilePath, L"Ct");
 			}
 		}
-#endif
 
 		CoreTechFileReaderUtils::AddFaceIdAttribut(MainId);
 
-		if (Context.ImportParameters.StitchingTechnique != StitchingNone && Context.ImportParameters.bEnableKernelIOTessellation)
+		if (Context.ImportParameters.StitchingTechnique != StitchingNone && Context.ImportParameters.bDisableCADKernelTessellation)
 		{
 			CADLibrary::CTKIO_Repair(MainId, Context.ImportParameters.StitchingTechnique, 10.);
 		}
@@ -444,7 +442,7 @@ namespace CADLibrary
 		}
 
 		// 3dxml file is zipped files, it's full managed by Kernel_io. We cannot read it in sequential mode
-		if (MainFileExt != TEXT("3dxml") && Context.ImportParameters.bEnableCacheUsage)
+		if (MainFileExt != TEXT("3dxml") && Context.ImportParameters.bEnableSequentialImport)
 		{
 			Flags &= ~CT_LOAD_FLAGS_LOAD_EXTERNAL_REF;
 		}
@@ -529,7 +527,7 @@ namespace CADLibrary
 			CT_FLAGS BodyProperties;
 			CT_BODY_IO::AskProperties(BodyId, BodyProperties);
 
-			if (Context.ImportParameters.bEnableKernelIOTessellation || !(BodyProperties & CT_BODY_PROP_EXACT))
+			if (Context.ImportParameters.bDisableCADKernelTessellation || !(BodyProperties & CT_BODY_PROP_EXACT))
 			{
 				if (ReadKioBody(BodyId, ComponentId, DefaultMaterialHash, false))
 				{
@@ -570,8 +568,8 @@ namespace CADLibrary
 		if (CT_INSTANCE_IO::AskTransformation(InstanceNodeId, (double*) Instance.TransformMatrix.M) == IO_OK)
 		{
 			if (Instance.TransformMatrix.ContainsNaN())
-			{
-				Instance.TransformMatrix.SetIdentity();
+				{
+					Instance.TransformMatrix.SetIdentity();
 			}
 		}
 	
diff --git a/Engine/Plugins/Enterprise/DatasmithCADImporter/Source/CADInterfaces/Public/CTSession.h b/Engine/Plugins/Enterprise/DatasmithCADImporter/Source/CADInterfaces/Public/CTSession.h
--- a/Engine/Plugins/Enterprise/DatasmithCADImporter/Source/CADInterfaces/Public/CTSession.h
+++ b/Engine/Plugins/Enterprise/DatasmithCADImporter/Source/CADInterfaces/Public/CTSession.h
@@ -25,7 +25,7 @@ public:
 		: FCoreTechSessionBase(InOwner)
 		, ImportParams(InImportParameters)
 	{
-		ImportParams.bEnableKernelIOTessellation = true;
+		ImportParams.bDisableCADKernelTessellation = true;
 	}
 
 	void ClearData();
diff --git a/Engine/Plugins/Enterprise/DatasmithCADImporter/Source/CADKernelSurface/Public/CADModelToCADKernelConverterBase.h b/Engine/Plugins/Enterprise/DatasmithCADImporter/Source/CADKernelSurface/Public/CADModelToCADKernelConverterBase.h
--- a/Engine/Plugins/Enterprise/DatasmithCADImporter/Source/CADKernelSurface/Public/CADModelToCADKernelConverterBase.h
+++ b/Engine/Plugins/Enterprise/DatasmithCADImporter/Source/CADKernelSurface/Public/CADModelToCADKernelConverterBase.h
@@ -26,7 +26,7 @@ public:
 		, GeometricTolerance(0.00001 / InImportParameters.MetricUnit)
 		, SquareTolerance(GeometricTolerance* GeometricTolerance)
 	{
-		ImportParameters.bEnableKernelIOTessellation = false;
+		ImportParameters.bDisableCADKernelTessellation = false;
 	}
 
 	virtual void InitializeProcess(double InMetricUnit) override
diff --git a/Engine/Plugins/Enterprise/DatasmithCADImporter/Source/CADTools/Private/CADOption.cpp b/Engine/Plugins/Enterprise/DatasmithCADImporter/Source/CADTools/Private/CADOption.cpp
new file mode 100644
--- /dev/null
+++ b/Engine/Plugins/Enterprise/DatasmithCADImporter/Source/CADTools/Private/CADOption.cpp
@@ -0,0 +1,57 @@
+// Copyright Epic Games, Inc. All Rights Reserved.
+#include "CADOptions.h"
+
+namespace CADLibrary
+{
+
+FString GCADLibrary = TEXT("KernelIO");
+FAutoConsoleVariableRef GCADTranslatorLibrary(
+	TEXT("ds.CADTranslator.CADLibrary"),
+	GCADLibrary,
+	TEXT("Define the used library to import CAD file i.e. \"TechSoft\" or \"KernelIO\".\n"),
+	ECVF_Default);
+
+bool bGDisableCADKernelTessellation = true;
+FAutoConsoleVariableRef GCADTranslatorDisableCADKernelTessellation(
+	TEXT("ds.CADTranslator.DisableCADKernelTessellation"),
+	bGDisableCADKernelTessellation,
+	TEXT("Disable to use CAD import library tessellator.\n"),
+	ECVF_Default);
+
+int32 GMaxImportThreads = 0;
+FAutoConsoleVariableRef GCADTranslatorMaxImportThreads(
+	TEXT("ds.CADTranslator.MaxImportThreads"),
+	GMaxImportThreads,
+	TEXT("\
+CAD file parallel processing\n\
+Default is MaxImportThreads = 0\n\
+0: multi - processing, n : multi - processing limited to n process.CADCache is mandatory.\n\
+1: -if CADCache is enable, the scene is read in a sequential mode with cache i.e.cache is used for sub - file already read,\n\
+   -if CADCache is disable, the scene is read all at once\n"),
+	ECVF_Default);
+
+bool bGEnableCADCache = true;
+FAutoConsoleVariableRef GCADTranslatorEnableCADCache(
+	TEXT("ds.CADTranslator.EnableCADCache"),
+	bGEnableCADCache,
+	TEXT("\
+Enable/disable temporary CAD processing file cache. These file will be use in a next import to avoid CAD file processing.\n\
+If MaxImportThreads != 1, EnableCADCache value is ignored\n\
+Default is enable\n"),
+ECVF_Default);
+
+bool bGOverwriteCache = false;
+FAutoConsoleVariableRef GCADTranslatorOverwriteCache(
+	TEXT("ds.CADTranslator.OverwriteCache"),
+	bGOverwriteCache,
+	TEXT("Overwrite any existing cache associated with the file being imported.\n"),
+	ECVF_Default);
+
+bool bGEnableTimeControl = true;
+FAutoConsoleVariableRef GCADTranslatorEnableTimeControl(
+	TEXT("ds.CADTranslator.EnableTimeControl"),
+	bGEnableTimeControl,
+	TEXT("Enable the timer that kill the worker if the import time is unusually long. With this time control, the load of the corrupted file is canceled but the rest of the scene is imported.\n"),
+	ECVF_Default);
+
+}
diff --git a/Engine/Plugins/Enterprise/DatasmithCADImporter/Source/CADTools/Public/CADOptions.h b/Engine/Plugins/Enterprise/DatasmithCADImporter/Source/CADTools/Public/CADOptions.h
--- a/Engine/Plugins/Enterprise/DatasmithCADImporter/Source/CADTools/Public/CADOptions.h
+++ b/Engine/Plugins/Enterprise/DatasmithCADImporter/Source/CADTools/Public/CADOptions.h
@@ -3,12 +3,20 @@
 #pragma once
 
 #include "DatasmithUtils.h"
+#include "HAL/IConsoleManager.h"
 #include "Math/Vector.h"
 #include "Misc/Paths.h"
 #include "Templates/TypeHash.h"
 
 namespace CADLibrary
 {
+	CADTOOLS_API extern bool bGDisableCADKernelTessellation;
+	CADTOOLS_API extern bool bGEnableCADCache;
+	CADTOOLS_API extern bool bGEnableTimeControl;
+	CADTOOLS_API extern bool bGOverwriteCache;
+	CADTOOLS_API extern int32 GMaxImportThreads;
+	CADTOOLS_API extern FString GCADLibrary;
+
 	enum EStitchingTechnique
 	{
 		StitchingNone = 0,
@@ -43,8 +51,9 @@ namespace CADLibrary
 		EDisplayDataPropagationMode Propagation = EDisplayDataPropagationMode::TopDown;
 		EDisplayPreference DisplayPreference = EDisplayPreference::MaterialPrefered;
 		bool bScaleUVMap = true;
-		bool bEnableCacheUsage = true;
-		bool bEnableKernelIOTessellation = true;
+		bool bEnableSequentialImport = true;
+		bool bOverwriteCache = false;
+		bool bDisableCADKernelTessellation = false;
 		bool bEnableTimeControl = true;
 
 		void SetMetricUnit(double NewMetricUnit)
@@ -80,15 +89,16 @@ namespace CADLibrary
 			Ar << (uint8&) ImportParameters.Propagation;
 			Ar << (uint8&) ImportParameters.DisplayPreference;
 			Ar << ImportParameters.bScaleUVMap;
-			Ar << ImportParameters.bEnableCacheUsage;
-			Ar << ImportParameters.bEnableKernelIOTessellation;
+			Ar << ImportParameters.bEnableSequentialImport;
+			Ar << ImportParameters.bOverwriteCache;
+			Ar << ImportParameters.bDisableCADKernelTessellation;
 			return Ar;
 		}
 
 		FString DefineCADFilePath(const TCHAR* FolderPath, const TCHAR* InFileName) const
 		{
 			FString OutFileName = FPaths::Combine(FolderPath, InFileName);
-			if (bEnableKernelIOTessellation)
+			if (bDisableCADKernelTessellation)
 			{
 				OutFileName += TEXT(".ct");
 			}
diff --git a/Engine/Plugins/Enterprise/DatasmithCADImporter/Source/DatasmithCADTranslator/Private/DatasmithCADTranslator.cpp b/Engine/Plugins/Enterprise/DatasmithCADImporter/Source/DatasmithCADTranslator/Private/DatasmithCADTranslator.cpp
--- a/Engine/Plugins/Enterprise/DatasmithCADImporter/Source/DatasmithCADTranslator/Private/DatasmithCADTranslator.cpp
+++ b/Engine/Plugins/Enterprise/DatasmithCADImporter/Source/DatasmithCADTranslator/Private/DatasmithCADTranslator.cpp
@@ -12,18 +12,6 @@
 #include "DatasmithUtils.h"
 #include "IDatasmithSceneElements.h"
 
-static TAutoConsoleVariable<int32> CVarStaticCADTranslatorEnableThreadedImport(
-	TEXT("ds.CADTranslator.EnableThreadedImport"),
-	1,
-	TEXT("Activate to parallelise CAD file processing.\n"),
-	ECVF_Default);
-
-static TAutoConsoleVariable<int32> CVarStaticCADTranslatorEnableKernelIOTessellation(
-	TEXT("ds.CADTranslator.EnableKernelIOTessellation"),
-	1,
-	TEXT("Activate to use KernelIO Tessellation.\n"),
-	ECVF_Default);
-
 void FDatasmithCADTranslator::Initialize(FDatasmithTranslatorCapabilities& OutCapabilities)
 {
 	if (ICADInterfacesModule::GetAvailability() == ECADInterfaceAvailability::Unavailable)
@@ -85,7 +73,7 @@ void FDatasmithCADTranslator::Initialize(FDatasmithTranslatorCapabilities& OutCa
 
 bool FDatasmithCADTranslator::LoadScene(TSharedRef<IDatasmithScene> DatasmithScene)
 {
-	ImportParameters.bEnableKernelIOTessellation = (CVarStaticCADTranslatorEnableKernelIOTessellation.GetValueOnAnyThread() != 0);;
+	ImportParameters.bDisableCADKernelTessellation = CADLibrary::bGDisableCADKernelTessellation;
 
 	ImportParameters.MetricUnit = 0.001;
 	ImportParameters.ScaleFactor = 0.1;
@@ -130,26 +118,41 @@ bool FDatasmithCADTranslator::LoadScene(TSharedRef<IDatasmithScene> DatasmithSce
 		ImportParameters.Propagation = CADLibrary::EDisplayDataPropagationMode::BodyOnly;
 	}
 
-	ImportParameters.bEnableCacheUsage = false;
+	ImportParameters.bEnableTimeControl = CADLibrary::bGEnableTimeControl;
+	if (CADLibrary::GMaxImportThreads == 1)
+	{
+		ImportParameters.bEnableSequentialImport = CADLibrary::bGEnableCADCache;
+	}
+	else
+	{
+		ImportParameters.bEnableSequentialImport = true;
+	}
+
 	FString CachePath = FDatasmithCADTranslatorModule::Get().GetCacheDir();
 	if (!CachePath.IsEmpty())
 	{
 		CachePath = FPaths::ConvertRelativePathToFull(CachePath);
-		ImportParameters.bEnableCacheUsage = true;
 	}
-
-	// Only use multi-processed translation if it is required and cache's usage is enabled
-	if (ImportParameters.bEnableCacheUsage)
+	else
 	{
-		bool bWithProcessor = (CVarStaticCADTranslatorEnableThreadedImport.GetValueOnAnyThread() != 0);
+		// Sequential Import is disabling because it needs CADCache
+		ImportParameters.bEnableSequentialImport = false;
+	}
 
+	// Use sequential translation (multi-processed or not)
+	if (ImportParameters.bEnableSequentialImport)
+	{
 		TMap<uint32, FString> CADFileToUEFileMap;
-		int32 NumCores = FPlatformMisc::NumberOfCores();
 		{
+			int32 NumCores = FPlatformMisc::NumberOfCores();
+			if (CADLibrary::GMaxImportThreads > 1)
+			{
+				NumCores = FMath::Min(CADLibrary::GMaxImportThreads, NumCores);
+			}
 			DatasmithDispatcher::FDatasmithDispatcher Dispatcher(ImportParameters, CachePath, NumCores, CADFileToUEFileMap, CADFileToUEGeomMap);
 			Dispatcher.AddTask(FileDescription);
 
-			Dispatcher.Process(bWithProcessor);
+			Dispatcher.Process(CADLibrary::GMaxImportThreads != 1);
 		}
 
 		FDatasmithSceneGraphBuilder SceneGraphBuilder(CADFileToUEFileMap, CachePath, DatasmithScene, GetSource(), ImportParameters);
@@ -160,8 +163,6 @@ bool FDatasmithCADTranslator::LoadScene(TSharedRef<IDatasmithScene> DatasmithSce
 		return true;
 	}
 
-	ImportParameters.bEnableCacheUsage = false;
-
 	CADLibrary::FCoreTechFileParser FileParser(ImportParameters, *FPaths::EnginePluginsDir());
 	if (FileParser.ProcessFile(FileDescription) != CADLibrary::ECoreTechParsingResult::ProcessOk)
 	{
@@ -195,7 +196,7 @@ bool FDatasmithCADTranslator::LoadStaticMesh(const TSharedRef<IDatasmithMeshElem
 	if (TOptional< FMeshDescription > Mesh = MeshBuilderPtr->GetMeshDescription(MeshElement, MeshParameters))
 	{
 		OutMeshPayload.LodMeshes.Add(MoveTemp(Mesh.GetValue()));
-		if(ImportParameters.bEnableKernelIOTessellation)
+		if(ImportParameters.bDisableCADKernelTessellation)
 		{
 			CoreTechSurface::AddSurfaceDataForMesh(MeshElement->GetFile(), ImportParameters, MeshParameters, GetCommonTessellationOptions(), OutMeshPayload);
 		}
diff --git a/Engine/Plugins/Enterprise/DatasmithCADImporter/Source/DatasmithDispatcher/Private/DatasmithDispatcher.cpp b/Engine/Plugins/Enterprise/DatasmithCADImporter/Source/DatasmithDispatcher/Private/DatasmithDispatcher.cpp
--- a/Engine/Plugins/Enterprise/DatasmithCADImporter/Source/DatasmithDispatcher/Private/DatasmithDispatcher.cpp
+++ b/Engine/Plugins/Enterprise/DatasmithCADImporter/Source/DatasmithDispatcher/Private/DatasmithDispatcher.cpp
@@ -12,19 +12,6 @@
 #include "Misc/ScopeLock.h"
 #include "HAL/IConsoleManager.h"
 
-#ifndef DATASMITH_CAD_IGNORE_CACHE
-static TAutoConsoleVariable<int32> CVarStaticCADTranslatorEnableCADCache(
-	TEXT("ds.CADTranslator.EnableCADCache"),
-	1,
-	TEXT("Activate to save temporary CAD processing file. These file will be use in a next import to avoid CAD file processing\n"),
-	ECVF_Default);
-#endif
-
-static TAutoConsoleVariable<int32> CVarStaticCADTranslatorEnableTimeControl(
-	TEXT("ds.CADTranslator.EnableTimeControl"),
-	1,
-	TEXT("Enable the timer that kill the worker if the import time is unusually long. With this time control, the load of the corrupted file is canceled but the rest of the scene is imported.\n"),
-	ECVF_Default);
 
 namespace DatasmithDispatcher
 {
@@ -39,10 +26,7 @@ FDatasmithDispatcher::FDatasmithDispatcher(const CADLibrary::FImportParameters&
 	, NumberOfWorkers(InNumberOfWorkers)
 	, NextWorkerId(0)
 {
-	ImportParameters.bEnableTimeControl = (CVarStaticCADTranslatorEnableTimeControl.GetValueOnAnyThread() != 0);
-	ImportParameters.bEnableCacheUsage = (CVarStaticCADTranslatorEnableCADCache.GetValueOnAnyThread() != 0);
-
-	if (ImportParameters.bEnableCacheUsage)
+	if (ImportParameters.bEnableSequentialImport)
 	{
 		// init cache folders
 		IFileManager::Get().MakeDirectory(*FPaths::Combine(ProcessCacheFolder, TEXT("scene")), true);
diff --git a/Engine/Plugins/Enterprise/DatasmithCADImporter/Source/DatasmithOpenNurbsTranslator/DatasmithOpenNurbsTranslator.Build.cs b/Engine/Plugins/Enterprise/DatasmithCADImporter/Source/DatasmithOpenNurbsTranslator/DatasmithOpenNurbsTranslator.Build.cs
--- a/Engine/Plugins/Enterprise/DatasmithCADImporter/Source/DatasmithOpenNurbsTranslator/DatasmithOpenNurbsTranslator.Build.cs
+++ b/Engine/Plugins/Enterprise/DatasmithCADImporter/Source/DatasmithOpenNurbsTranslator/DatasmithOpenNurbsTranslator.Build.cs
@@ -24,6 +24,7 @@ namespace UnrealBuildTool.Rules
 					"CADKernel",
 					"CADKernelSurface",
 					"CADLibrary",
+					"CADTools",
 					"CoreTechSurface",
 					"DatasmithCore",
 					"DatasmithContent",
diff --git a/Engine/Plugins/Enterprise/DatasmithCADImporter/Source/DatasmithOpenNurbsTranslator/Private/DatasmithOpenNurbsTranslator.cpp b/Engine/Plugins/Enterprise/DatasmithCADImporter/Source/DatasmithOpenNurbsTranslator/Private/DatasmithOpenNurbsTranslator.cpp
--- a/Engine/Plugins/Enterprise/DatasmithCADImporter/Source/DatasmithOpenNurbsTranslator/Private/DatasmithOpenNurbsTranslator.cpp
+++ b/Engine/Plugins/Enterprise/DatasmithCADImporter/Source/DatasmithOpenNurbsTranslator/Private/DatasmithOpenNurbsTranslator.cpp
@@ -8,6 +8,7 @@
 #include "CADInterfacesModule.h"
 #include "CADKernelSurfaceExtension.h"
 #include "CADModelConverter.h"
+#include "CADOptions.h"
 #include "CoreTechSurfaceExtension.h"
 #include "CoreTechSurfaceHelper.h"
 #include "DatasmithImportOptions.h"
@@ -645,18 +646,12 @@ public:
 			TranslationCache = MakeShared<FTranslationCache>();
 		}
 
-		bool bEnableKernelIOTessellation = true;
-		if (IConsoleVariable* CVar = IConsoleManager::Get().FindConsoleVariable(TEXT("ds.CADTranslator.EnableKernelIOTessellation")))
-		{
-			bEnableKernelIOTessellation = CVar->GetInt() != 0;
-		}
-
 		CADLibrary::FImportParameters ImportParameters;
 		ImportParameters.MetricUnit = 0.001;
 		ImportParameters.ScaleFactor = 0.1;
 		ImportParameters.ModelCoordSys = FDatasmithUtils::EModelCoordSystem::ZUp_RightHanded_FBXLegacy;
 
-		if (bEnableKernelIOTessellation)
+		if (CADLibrary::bGDisableCADKernelTessellation)
 		{
 			TSharedRef<FOpenNurbsBRepToCoretechConverter> OpenNurbsBRepToCoretechConverter = MakeShared<FOpenNurbsBRepToCoretechConverter>(TEXT("Al2CTSharedSession"), ImportParameters);
 			CADModelConverter = OpenNurbsBRepToCoretechConverter;
@@ -685,7 +680,7 @@ public:
 	void SetBaseOptions(const FDatasmithImportBaseOptions& InBaseOptions);
 	void SetOpenNurbsOptions(const FDatasmithOpenNurbsOptions& Options);
 	void SetOutputPath(const FString& Path) { OutputPath = Path; }
-	
+
 	void ShowMessageLog(const FString& Filename);
 
 	bool LoadStaticMesh(const TSharedRef<IDatasmithMeshElement> MeshElement, FDatasmithMeshElementPayload& OutMeshPayload, const FDatasmithTessellationOptions& InTessellationOptions);
@@ -3285,7 +3280,7 @@ bool FDatasmithOpenNurbsTranslator::LoadStaticMesh(const TSharedRef<IDatasmithMe
 	return false;
 }
 
-bool FOpenNurbsTranslatorImpl::LoadStaticMesh(const TSharedRef<IDatasmithMeshElement> MeshElement, FDatasmithMeshElementPayload & OutMeshPayload, const FDatasmithTessellationOptions& InTessellationOptions)
+bool FOpenNurbsTranslatorImpl::LoadStaticMesh(const TSharedRef<IDatasmithMeshElement> MeshElement, FDatasmithMeshElementPayload& OutMeshPayload, const FDatasmithTessellationOptions& InTessellationOptions)
 {
 	if (TOptional<FMeshDescription> Mesh = GetMeshDescription(MeshElement))
 	{
diff --git a/Engine/Plugins/Enterprise/DatasmithCADImporter/Source/DatasmithPLMXMLTranslator/DatasmithPLMXMLTranslator.Build.cs b/Engine/Plugins/Enterprise/DatasmithCADImporter/Source/DatasmithPLMXMLTranslator/DatasmithPLMXMLTranslator.Build.cs
--- a/Engine/Plugins/Enterprise/DatasmithCADImporter/Source/DatasmithPLMXMLTranslator/DatasmithPLMXMLTranslator.Build.cs
+++ b/Engine/Plugins/Enterprise/DatasmithCADImporter/Source/DatasmithPLMXMLTranslator/DatasmithPLMXMLTranslator.Build.cs
@@ -6,39 +6,37 @@ using System.IO;
 
 namespace UnrealBuildTool.Rules
 {
-    public class DatasmithPLMXMLTranslator : ModuleRules
-    {
-        public DatasmithPLMXMLTranslator(ReadOnlyTargetRules Target) : base(Target)
-        {
-	        PublicDependencyModuleNames.AddRange(
-		        new string[]
-		        {
-			        "Core",
-			        "CoreUObject",
-		        }
-	        );
-
-            PrivateDependencyModuleNames.AddRange(
-                new string[]
-                {
-                    "Engine",
-                    "MeshDescription",
-					"XmlParser",
-
-					"DatasmithCore",
-
+	public class DatasmithPLMXMLTranslator : ModuleRules
+	{
+		public DatasmithPLMXMLTranslator(ReadOnlyTargetRules Target) : base(Target)
+		{
+			PublicDependencyModuleNames.AddRange(
+				new string[]
+				{
+					"Core",
+					"CoreUObject",
+				}
+			);
+
+			PrivateDependencyModuleNames.AddRange(
+				new string[]
+				{
 					"CADInterfaces",
+					"CADKernelSurface",
 					"CADLibrary",
 					"CADTools",
 					"CoreTechSurface",
 					"DatasmithCADTranslator", // for DatasmithMeshBuilder
-                    "DatasmithTranslator",
-                    "DatasmithDispatcher",
+					"DatasmithContent",
+					"DatasmithCore",
+					"DatasmithDispatcher",
+					"DatasmithTranslator",
+					"Engine",
+					"MeshDescription",
 					"ParametricSurface",
-
-                    "DatasmithContent",
-                }
-            );
+					"XmlParser",
+				}
+			);
 
 
 			if (Target.bBuildEditor == true)
@@ -53,9 +51,9 @@ namespace UnrealBuildTool.Rules
 			}
 
 			if (System.Type.GetType("CoreTech") != null)
-            {
-	            PrivateDependencyModuleNames.Add("CoreTech");
-            }
-        }
-    }
+			{
+				PrivateDependencyModuleNames.Add("CoreTech");
+			}
+		}
+	}
 }
diff --git a/Engine/Plugins/Enterprise/DatasmithCADImporter/Source/DatasmithPLMXMLTranslator/Private/DatasmithPlmXmlImporter.cpp b/Engine/Plugins/Enterprise/DatasmithCADImporter/Source/DatasmithPLMXMLTranslator/Private/DatasmithPlmXmlImporter.cpp
--- a/Engine/Plugins/Enterprise/DatasmithCADImporter/Source/DatasmithPLMXMLTranslator/Private/DatasmithPlmXmlImporter.cpp
+++ b/Engine/Plugins/Enterprise/DatasmithCADImporter/Source/DatasmithPLMXMLTranslator/Private/DatasmithPlmXmlImporter.cpp
@@ -19,6 +19,7 @@
 
 #include "CADData.h"
 #include "CADToolsModule.h"
+#include "CADKernelSurfaceExtension.h"
 #include "CoreTechSurfaceExtension.h"
 #include "CoreTechSurfaceHelper.h"
 #include "DatasmithAdditionalData.h"
@@ -304,13 +305,13 @@ namespace PlmXml
 			{
 				OutMeshPayload.LodMeshes.Add(MoveTemp(Mesh.GetValue()));
 
-				if (ImportParameters.bEnableKernelIOTessellation)
+				if (ImportParameters.bDisableCADKernelTessellation)
 				{
 					CoreTechSurface::AddSurfaceDataForMesh(MeshElement->GetFile(), ImportParameters, MeshParameters, TessellationOptions, OutMeshPayload);
 				}
 				else
 				{
-					//CADKernelSurface::AddSurfaceDataForMesh(MeshElement, ImportParameters, MeshParameters, TessellationOptions, OutMeshPayload);
+					CADKernelSurface::AddSurfaceDataForMesh(MeshElement->GetFile(), ImportParameters, MeshParameters, TessellationOptions, OutMeshPayload);
 				}
 
 				return true;
diff --git a/Engine/Plugins/Enterprise/DatasmithCADImporter/Source/DatasmithWireTranslator/Private/DatasmithWireTranslator.cpp b/Engine/Plugins/Enterprise/DatasmithCADImporter/Source/DatasmithWireTranslator/Private/DatasmithWireTranslator.cpp
--- a/Engine/Plugins/Enterprise/DatasmithCADImporter/Source/DatasmithWireTranslator/Private/DatasmithWireTranslator.cpp
+++ b/Engine/Plugins/Enterprise/DatasmithCADImporter/Source/DatasmithWireTranslator/Private/DatasmithWireTranslator.cpp
@@ -10,7 +10,6 @@
 #include "DatasmithUtils.h"
 #include "HAL/ConsoleManager.h"
 #include "DatasmithWireTranslatorModule.h"
-#include "HAL/ConsoleManager.h"
 #include "IDatasmithSceneElements.h"
 #include "Misc/FileHelper.h"
 #include "Misc/Paths.h"
@@ -144,17 +143,11 @@ public:
 		DatasmithScene->SetProductName(TEXT("Alias Tools"));
 		DatasmithScene->SetProductVersion(TEXT("Alias 2022"));
 
-		bool bEnableKernelIOTessellation = true;
-		if (IConsoleVariable* CVar = IConsoleManager::Get().FindConsoleVariable(TEXT("ds.CADTranslator.EnableKernelIOTessellation")))
-		{
-			bEnableKernelIOTessellation = CVar->GetInt() != 0;
-		}
-
 		CADLibrary::FImportParameters ImportParameters;
 		ImportParameters.MetricUnit = 0.01;
 		ImportParameters.ScaleFactor = 1;
 
-		if (bEnableKernelIOTessellation)
+		if(CADLibrary::bGDisableCADKernelTessellation)
 		{
 			TSharedRef<FAliasModelToCoretechConverter> AliasToCoretechConverter = MakeShared<FAliasModelToCoretechConverter>(TEXT("Al2CTSharedSession"), ImportParameters);
 			CADModelConverter = AliasToCoretechConverter;
@@ -162,9 +155,9 @@ public:
 		}
 		else
 		{
-			TSharedRef<FAliasModelToCADKernelConverter> AliasToCoretechConverter = MakeShared<FAliasModelToCADKernelConverter>(ImportParameters);
-			CADModelConverter = AliasToCoretechConverter;
-			AliasBRepConverter = AliasToCoretechConverter;
+			TSharedRef<FAliasModelToCADKernelConverter> AliasToCADKernelConverter = MakeShared<FAliasModelToCADKernelConverter>(ImportParameters);
+			CADModelConverter = AliasToCADKernelConverter;
+			AliasBRepConverter = AliasToCADKernelConverter;
 		}
 	}
 
diff --git a/Engine/Plugins/Enterprise/DatasmithCADImporter/Source/ParametricSurface/Private/ParametricSurfaceTranslator.cpp b/Engine/Plugins/Enterprise/DatasmithCADImporter/Source/ParametricSurface/Private/ParametricSurfaceTranslator.cpp
--- a/Engine/Plugins/Enterprise/DatasmithCADImporter/Source/ParametricSurface/Private/ParametricSurfaceTranslator.cpp
+++ b/Engine/Plugins/Enterprise/DatasmithCADImporter/Source/ParametricSurface/Private/ParametricSurfaceTranslator.cpp
@@ -7,7 +7,7 @@
 void FParametricSurfaceTranslator::GetSceneImportOptions(TArray<TStrongObjectPtr<UDatasmithOptionsBase>>& Options)
 {
 	FString Extension = GetSource().GetSourceFileExtension();
-	if (Extension == "cgr" || Extension == "3dxml")
+	if (Extension == TEXT("cgr") || Extension == TEXT("3dxml"))
 	{
 		return;
 	}
