Move IsSourceSupported form FParametricSurfaceTranslator to DatasmithCADTranslator

diff --git a/Engine/Plugins/Enterprise/DatasmithCADImporter/Source/DatasmithCADTranslator/Private/DatasmithCADTranslator.cpp b/Engine/Plugins/Enterprise/DatasmithCADImporter/Source/DatasmithCADTranslator/Private/DatasmithCADTranslator.cpp
--- a/Engine/Plugins/Enterprise/DatasmithCADImporter/Source/DatasmithCADTranslator/Private/DatasmithCADTranslator.cpp
+++ b/Engine/Plugins/Enterprise/DatasmithCADImporter/Source/DatasmithCADTranslator/Private/DatasmithCADTranslator.cpp
@@ -9,6 +9,7 @@
 #include "DatasmithDispatcher.h"
 #include "DatasmithMeshBuilder.h"
 #include "DatasmithSceneGraphBuilder.h"
+#include "DatasmithTranslator.h"
 #include "DatasmithUtils.h"
 #include "IDatasmithSceneElements.h"
 
@@ -71,6 +72,16 @@ void FDatasmithCADTranslator::Initialize(FDatasmithTranslatorCapabilities& OutCa
 	OutCapabilities.SupportedFileFormats.Add(FFileFormatInfo{ TEXT("ct"), TEXT("Kernel_IO files") });
 }
 
+bool FDatasmithCADTranslator::IsSourceSupported(const FDatasmithSceneSource& Source)
+{
+	if (Source.GetSourceFileExtension() != TEXT("xml"))
+	{
+		return true;
+	}
+
+	return Datasmith::CheckXMLFileSchema(Source.GetSourceFile(), TEXT("XPDMXML"), TEXT("ns3:Uos"));
+}
+
 bool FDatasmithCADTranslator::LoadScene(TSharedRef<IDatasmithScene> DatasmithScene)
 {
 	ImportParameters.bDisableCADKernelTessellation = CADLibrary::bGDisableCADKernelTessellation;
diff --git a/Engine/Plugins/Enterprise/DatasmithCADImporter/Source/DatasmithCADTranslator/Private/DatasmithCADTranslator.h b/Engine/Plugins/Enterprise/DatasmithCADImporter/Source/DatasmithCADTranslator/Private/DatasmithCADTranslator.h
--- a/Engine/Plugins/Enterprise/DatasmithCADImporter/Source/DatasmithCADTranslator/Private/DatasmithCADTranslator.h
+++ b/Engine/Plugins/Enterprise/DatasmithCADImporter/Source/DatasmithCADTranslator/Private/DatasmithCADTranslator.h
@@ -20,6 +20,7 @@ public:
 	virtual FName GetFName() const override { return "DatasmithCADTranslator"; };
 
 	virtual void Initialize(FDatasmithTranslatorCapabilities& OutCapabilities) override;
+	virtual bool IsSourceSupported(const FDatasmithSceneSource& Source) override;
 
 	virtual bool LoadScene(TSharedRef<IDatasmithScene> OutScene) override;
 
diff --git a/Engine/Plugins/Enterprise/DatasmithCADImporter/Source/ParametricSurface/Private/ParametricSurfaceTranslator.cpp b/Engine/Plugins/Enterprise/DatasmithCADImporter/Source/ParametricSurface/Private/ParametricSurfaceTranslator.cpp
--- a/Engine/Plugins/Enterprise/DatasmithCADImporter/Source/ParametricSurface/Private/ParametricSurfaceTranslator.cpp
+++ b/Engine/Plugins/Enterprise/DatasmithCADImporter/Source/ParametricSurface/Private/ParametricSurfaceTranslator.cpp
@@ -30,16 +30,3 @@ void FParametricSurfaceTranslator::SetSceneImportOptions(TArray<TStrongObjectPtr
 	}
 }
 
-bool FParametricSurfaceTranslator::IsSourceSupported(const FDatasmithSceneSource& Source)
-{
-	if (Source.GetSourceFileExtension() != TEXT("xml"))
-	{
-		return true;
-	}
-
-	return Datasmith::CheckXMLFileSchema(Source.GetSourceFile(), TEXT("XPDMXML"), TEXT("ns3:Uos"));
-}
-
-
-
-
diff --git a/Engine/Plugins/Enterprise/DatasmithCADImporter/Source/ParametricSurface/Public/ParametricSurfaceTranslator.h b/Engine/Plugins/Enterprise/DatasmithCADImporter/Source/ParametricSurface/Public/ParametricSurfaceTranslator.h
--- a/Engine/Plugins/Enterprise/DatasmithCADImporter/Source/ParametricSurface/Public/ParametricSurfaceTranslator.h
+++ b/Engine/Plugins/Enterprise/DatasmithCADImporter/Source/ParametricSurface/Public/ParametricSurfaceTranslator.h
@@ -19,7 +19,6 @@ public:
 	// Begin IDatasmithTranslator overrides
 	virtual void GetSceneImportOptions(TArray<TStrongObjectPtr<UDatasmithOptionsBase>>& Options) override;
 	virtual void SetSceneImportOptions(TArray<TStrongObjectPtr<UDatasmithOptionsBase>>& Options) override;
-	virtual bool IsSourceSupported(const FDatasmithSceneSource& Source) override;
 	// End IDatasmithTranslator overrides
 
 protected:
