Fix naming issues

diff --git a/Engine/Plugins/Enterprise/DatasmithCADImporter/Source/CADLibrary/Private/MeshDescriptionHelper.cpp b/Engine/Plugins/Enterprise/DatasmithCADImporter/Source/CADLibrary/Private/MeshDescriptionHelper.cpp
--- a/Engine/Plugins/Enterprise/DatasmithCADImporter/Source/CADLibrary/Private/MeshDescriptionHelper.cpp
+++ b/Engine/Plugins/Enterprise/DatasmithCADImporter/Source/CADLibrary/Private/MeshDescriptionHelper.cpp
@@ -361,12 +361,12 @@ namespace CADLibrary
 	bool FillMesh(const FMeshParameters& MeshParameters, const FImportParameters& ImportParams, FBodyMesh& BodyTessellation, FMeshDescription& MeshDescription)
 	{
 		const int32 UVChannel = 0;
-		const int32 TriangleCount = 3;
+		const int32 VertexCountPerFace = 3;
 		const TriangleIndex Clockwise = { 0, 1, 2 };
 		const TriangleIndex CounterClockwise = { 0, 2, 1 };
 
 		TArray<FVertexInstanceID> TriangleVertexInstanceIDs;
-		TriangleVertexInstanceIDs.SetNum(TriangleCount);
+		TriangleVertexInstanceIDs.SetNum(VertexCountPerFace);
 
 		TArray<FVertexInstanceID> MeshVertexInstanceIDs;
 		TArray<uint32> NewFaceIndex;  // new CT face index to remove degenerated face
@@ -406,7 +406,7 @@ namespace CADLibrary
 
 		if (Algo::AnyOf(BodyTessellation.Faces, [](const FTessellationData& FaceTessellation) { return !FaceTessellation.TexCoordArray.IsEmpty(); }))
 		{
-			VertexInstanceUVs.SetNumChannels(1);
+		VertexInstanceUVs.SetNumChannels(1);
 		}
 
 		int32 NbStep = 1;
@@ -441,8 +441,8 @@ namespace CADLibrary
 					continue;
 				}
 
-				int32 VertexIDs[3];
-				int32 VertexIndex[3];
+				int32 FaceVertexIDs[3];
+				int32 FaceVertexIndices[3];
 				FVector Temp3D = { 0, 0, 0 };
 				FVector2D TexCoord2D = { 0, 0 };
 
@@ -452,34 +452,34 @@ namespace CADLibrary
 				PatchIndex++;
 
 				// build each valid face i.e. 3 different indexes
-				for (int32 Index = 0; Index < Tessellation.VertexIndices.Num(); Index += 3)
+				for (int32 FaceIndex = 0; FaceIndex < Tessellation.VertexIndices.Num(); FaceIndex += VertexCountPerFace)
 				{
-					VertexIndex[0] = Tessellation.PositionIndices[Tessellation.VertexIndices[Index + Orientation[0]]];
-					VertexIndex[1] = Tessellation.PositionIndices[Tessellation.VertexIndices[Index + Orientation[1]]];
-					VertexIndex[2] = Tessellation.PositionIndices[Tessellation.VertexIndices[Index + Orientation[2]]];
+					FaceVertexIndices[0] = Tessellation.PositionIndices[Tessellation.VertexIndices[FaceIndex + Orientation[0]]];
+					FaceVertexIndices[1] = Tessellation.PositionIndices[Tessellation.VertexIndices[FaceIndex + Orientation[1]]];
+					FaceVertexIndices[2] = Tessellation.PositionIndices[Tessellation.VertexIndices[FaceIndex + Orientation[2]]];
 
-					if (VertexIndex[0] == INDEX_NONE || VertexIndex[1] == INDEX_NONE || VertexIndex[2] == INDEX_NONE)
+					if (FaceVertexIndices[0] == INDEX_NONE || FaceVertexIndices[1] == INDEX_NONE || FaceVertexIndices[2] == INDEX_NONE)
 					{
 						continue;
 					}
 
 					// Verify the 3 input indices are not defining a degenerated triangle
-					if (VertexIndex[0] == VertexIndex[1] || VertexIndex[0] == VertexIndex[2] || VertexIndex[1] == VertexIndex[2])
+					if (FaceVertexIndices[0] == FaceVertexIndices[1] || FaceVertexIndices[0] == FaceVertexIndices[2] || FaceVertexIndices[1] == FaceVertexIndices[2])
 					{
 						continue;
 					}
 
-					VertexIDs[0] = VertexIdSet[VertexIndex[0]];
-					VertexIDs[1] = VertexIdSet[VertexIndex[1]];
-					VertexIDs[2] = VertexIdSet[VertexIndex[2]];
+					FaceVertexIDs[0] = VertexIdSet[FaceVertexIndices[0]];
+					FaceVertexIDs[1] = VertexIdSet[FaceVertexIndices[1]];
+					FaceVertexIDs[2] = VertexIdSet[FaceVertexIndices[2]];
 
-					NewFaceIndex.Add(Tessellation.VertexIndices[Index + Orientation[0]]);
-					NewFaceIndex.Add(Tessellation.VertexIndices[Index + Orientation[1]]);
-					NewFaceIndex.Add(Tessellation.VertexIndices[Index + Orientation[2]]);
+					NewFaceIndex.Add(Tessellation.VertexIndices[FaceIndex + Orientation[0]]);
+					NewFaceIndex.Add(Tessellation.VertexIndices[FaceIndex + Orientation[1]]);
+					NewFaceIndex.Add(Tessellation.VertexIndices[FaceIndex + Orientation[2]]);
 
-					MeshVertexInstanceIDs.Add(TriangleVertexInstanceIDs[0] = MeshDescription.CreateVertexInstance((FVertexID)VertexIDs[0]));
-					MeshVertexInstanceIDs.Add(TriangleVertexInstanceIDs[1] = MeshDescription.CreateVertexInstance((FVertexID)VertexIDs[1]));
-					MeshVertexInstanceIDs.Add(TriangleVertexInstanceIDs[2] = MeshDescription.CreateVertexInstance((FVertexID)VertexIDs[2]));
+					MeshVertexInstanceIDs.Add(TriangleVertexInstanceIDs[0] = MeshDescription.CreateVertexInstance((FVertexID)FaceVertexIDs[0]));
+					MeshVertexInstanceIDs.Add(TriangleVertexInstanceIDs[1] = MeshDescription.CreateVertexInstance((FVertexID)FaceVertexIDs[1]));
+					MeshVertexInstanceIDs.Add(TriangleVertexInstanceIDs[2] = MeshDescription.CreateVertexInstance((FVertexID)FaceVertexIDs[2]));
 
 					// Add the triangle as a polygon to the mesh description
 					const FPolygonID PolygonID = MeshDescription.CreatePolygon(*PolygonGroupID, TriangleVertexInstanceIDs);
@@ -490,16 +490,16 @@ namespace CADLibrary
 
 				if (!Tessellation.TexCoordArray.IsEmpty())
 				{
-					for (int32 IndexFace = 0; IndexFace < MeshVertexInstanceIDs.Num(); IndexFace += TriangleCount)
+					for (int32 FaceIndex = 0; FaceIndex < MeshVertexInstanceIDs.Num(); FaceIndex += VertexCountPerFace)
 					{
-						for (int32 Index = 0; Index < TriangleCount; Index++)
+						for (int32 VertexIndex = 0; VertexIndex < VertexCountPerFace; VertexIndex++)
 						{
-							const FVertexInstanceID& VertexInstanceID = MeshVertexInstanceIDs[IndexFace + Orientation[Index]];
-							VertexInstanceUVs.Set(VertexInstanceID, UVChannel, Tessellation.TexCoordArray[NewFaceIndex[IndexFace + Index]]);
+							const FVertexInstanceID& VertexInstanceID = MeshVertexInstanceIDs[FaceIndex + Orientation[VertexIndex]];
+							VertexInstanceUVs.Set(VertexInstanceID, UVChannel, Tessellation.TexCoordArray[NewFaceIndex[FaceIndex + VertexIndex]]);
 						}
 					}
 				}
-				
+
 				for (const FVertexInstanceID& VertexInstanceID : MeshVertexInstanceIDs)
 				{
 					VertexInstanceColors[VertexInstanceID] = FLinearColor::White;
