TeshSoft PreImplementation:

The main idea is to replace CoreTechFileParser by a generic CADFileReader.
CADFileReader has an ICADFileParser that will process (parse) the file with the selected SDK (KernelIO or TechSoft or ...).
CoreTechFileReader become CoreTechFileParser and implement ICADFileParser.
The next step is to implement TechSoftFileParser.

To simplify the exchange, a FCADFileData replace FContext and extand the encapsulation of the data needed and generated during the file parsing

diff --git a/Engine/Plugins/Enterprise/DatasmithCADImporter/Source/CADInterfaces/Private/CADFileReader.cpp b/Engine/Plugins/Enterprise/DatasmithCADImporter/Source/CADInterfaces/Private/CADFileReader.cpp
new file mode 100644
--- /dev/null
+++ b/Engine/Plugins/Enterprise/DatasmithCADImporter/Source/CADInterfaces/Private/CADFileReader.cpp
@@ -0,0 +1,147 @@
+// Copyright Epic Games, Inc. All Rights Reserved.
+
+#include "CADFileReader.h"
+
+#include "CADFileParser.h"
+#include "CADOptions.h"
+#include "CoreTechFileParser.h"
+#include "HAL/FileManager.h"
+#include "TechSoftFileParser.h"
+#include "Templates/TypeHash.h"
+
+namespace CADLibrary
+{
+	FCADFileReader::FCADFileReader(const FImportParameters& ImportParams, FFileDescription& InFile, const FString& EnginePluginsPath, const FString& InCachePath)
+		: CADFileData(ImportParams, InFile, InCachePath)
+	{
+#ifdef USE_KERNEL_IO_SDK
+		if (GCADLibrary == TEXT("KernelIO"))
+		{
+			CADParser = MakeUnique<FCoreTechFileParser>(CADFileData, EnginePluginsPath);
+		}
+#endif
+#if defined(USE_KERNEL_IO_SDK) && defined(USE_TECHSOFT_SDK)
+		else
+#endif
+#ifdef USE_TECHSOFT_SDK
+		if (GCADLibrary == TEXT("TechSoft"))
+		{
+			CADParser = MakeUnique<FTechSoftFileParser>(CADFileData, EnginePluginsPath);
+		}
+#endif
+	}
+
+	bool FCADFileReader::FindFile(FFileDescription& File)
+	{
+		FString FileName = File.Name;
+
+		FString FilePath = FPaths::GetPath(File.Path);
+		FString RootFilePath = File.MainCadFilePath;
+
+		// Basic case: FilePath is, or is in a sub-folder of, RootFilePath
+		if (FilePath.StartsWith(RootFilePath))
+		{
+			return IFileManager::Get().FileExists(*File.Path);
+		}
+
+		// Advance case: end of FilePath is in a upper-folder of RootFilePath
+		// e.g.
+		// FilePath = D:\\data temp\\Unstructured project\\Folder2\\Added_Object.SLDPRT
+		//                                                 ----------------------------
+		// RootFilePath = D:\\data\\CAD Files\\SolidWorks\\p033 - Unstructured project\\Folder1
+		//                ------------------------------------------------------------
+		// NewPath = D:\\data\\CAD Files\\SolidWorks\\p033 - Unstructured project\\Folder2\\Added_Object.SLDPRT
+		TArray<FString> RootPaths;
+		RootPaths.Reserve(30);
+		do
+		{
+			RootFilePath = FPaths::GetPath(RootFilePath);
+			RootPaths.Emplace(RootFilePath);
+		} while (!FPaths::IsDrive(RootFilePath) && !RootFilePath.IsEmpty());
+
+		TArray<FString> FilePaths;
+		FilePaths.Reserve(30);
+		FilePaths.Emplace(FileName);
+		while (!FPaths::IsDrive(FilePath) && !FilePath.IsEmpty())
+		{
+			FString FolderName = FPaths::GetCleanFilename(FilePath);
+			FilePath = FPaths::GetPath(FilePath);
+			FilePaths.Emplace(FPaths::Combine(FolderName, FilePaths.Last()));
+		};
+
+		for (int32 IndexFolderPath = 0; IndexFolderPath < RootPaths.Num(); IndexFolderPath++)
+		{
+			for (int32 IndexFilePath = 0; IndexFilePath < FilePaths.Num(); IndexFilePath++)
+			{
+				FString NewFilePath = FPaths::Combine(RootPaths[IndexFolderPath], FilePaths[IndexFilePath]);
+				if (IFileManager::Get().FileExists(*NewFilePath))
+				{
+					File.Path = NewFilePath;
+					return true;
+				};
+			}
+		}
+
+		// Last case: the FilePath is elsewhere and the file exist
+		// A Warning is launch because the file could be expected to not be loaded
+		if (IFileManager::Get().FileExists(*File.Path))
+		{
+			return true;
+		}
+
+		CADFileData.AddWarningMessages(FString::Printf(TEXT("File %s cannot be found."), *File.Path));
+		return false;
+	}
+
+	ECADParsingResult FCADFileReader::ProcessFile()
+	{
+		if (!CADParser.IsValid())
+		{
+			return ECADParsingResult::ProcessFailed;
+		}
+
+		if (!FindFile(CADFileData.GetCADFileDescription()))
+		{
+			return ECADParsingResult::FileNotFound;
+		}
+
+		if (CADFileData.IsSequentialImport())
+		{
+			CADFileData.SetArchiveNames();
+
+			bool bNeedToProceed = true;
+
+			FString CADFileCachePath = CADFileData.GetCADCachePath();
+			if (!CADFileData.GetImportParameters().bOverwriteCache && IFileManager::Get().FileExists(*CADFileCachePath))
+			{
+				FString MeshArchiveFilePath = CADFileData.GetMeshArchiveFilePath();
+				if (IFileManager::Get().FileExists(*MeshArchiveFilePath)) // the file has been proceed with same meshing parameters
+				{
+					bNeedToProceed = false;
+				}
+				else // the file has been converted into CT file but meshed with different parameters
+				{
+					CADFileData.ReplaceFileByCacheBackup(CADFileCachePath);
+				}
+			}
+
+			if (!bNeedToProceed)
+			{
+				// The file has been yet proceed, get ExternalRef
+				CADFileData.LoadSceneGraphArchive();
+				return ECADParsingResult::ProcessOk;
+			}
+		}
+
+		// Process the file
+		ECADParsingResult Result = CADParser->Process();
+
+		if (ECADParsingResult::ProcessOk == Result)
+		{
+			CADFileData.ExportSceneGraphFile();
+			CADFileData.ExportMeshArchiveFile();
+		}
+
+		return Result;
+	}
+}
diff --git a/Engine/Plugins/Enterprise/DatasmithCADImporter/Source/CADInterfaces/Private/CoreTechFileParser.cpp b/Engine/Plugins/Enterprise/DatasmithCADImporter/Source/CADInterfaces/Private/CoreTechFileParser.cpp
deleted file mode 100644
diff --git a/Engine/Plugins/Enterprise/DatasmithCADImporter/Source/CADInterfaces/Private/CoreTechFileParser.h b/Engine/Plugins/Enterprise/DatasmithCADImporter/Source/CADInterfaces/Private/CoreTechFileParser.h
new file mode 100644
--- /dev/null
+++ b/Engine/Plugins/Enterprise/DatasmithCADImporter/Source/CADInterfaces/Private/CoreTechFileParser.h
@@ -0,0 +1,97 @@
+// Copyright Epic Games, Inc. All Rights Reserved.
+#pragma once
+
+#include "CoreMinimal.h"
+
+#ifdef USE_KERNEL_IO_SDK
+
+#include "CADData.h"
+#include "CADFileData.h"
+#include "CADFileParser.h"
+#include "CADOptions.h"
+#include "CADSceneGraph.h"
+#include "CoreTechTypes.h"
+
+#include "CoreTechBridge.h"
+
+#include "Containers/Map.h"
+#include "Containers/Queue.h"
+#include "Misc/Paths.h"
+
+struct FFileStatData;
+struct FFileDescription;
+
+#pragma warning(push)
+#pragma warning(disable:4996) // unsafe sprintf
+#pragma warning(disable:4828) // illegal character
+#include "kernel_io/kernel_io.h"
+#include "kernel_io/kernel_io_error.h"
+#include "kernel_io/kernel_io_type.h"
+#include "kernel_io/attribute_io/attribute_enum.h"
+#include "kernel_io/material_io/material_io.h"
+#pragma warning(pop)
+
+namespace CADKernel
+{
+	class FBody;
+	class FBodyMesh;
+	class FFaceMesh;
+	class FModelMesh;
+}
+
+namespace CADLibrary
+{
+	namespace CoreTechFileParserUtils
+	{
+		void ScaleUV(CT_OBJECT_ID FaceID, TArray<FVector2D>& TexCoordArray, float Scale);
+		uint32 GetFaceTessellation(CT_OBJECT_ID FaceID, FTessellationData& Tessellation);
+		void GetBodyTessellation(CT_OBJECT_ID BodyId, FBodyMesh& OutBodyMesh, TFunction<void(CT_OBJECT_ID, int32, FTessellationData&)> = TFunction<void(CT_OBJECT_ID, int32, FTessellationData&)>());
+	}
+
+	namespace CADKernelUtils
+	{
+		void GetBodyTessellation(const TSharedRef<CADKernel::FModelMesh>& ModelMesh, const TSharedRef<CADKernel::FBody>& Body, FBodyMesh& OutBodyMesh, uint32 DefaultMaterialHash, TFunction<void(FObjectDisplayDataId, FObjectDisplayDataId, int32)> SetFaceMainMaterial);
+
+		uint32 GetFaceTessellation(const TSharedRef<CADKernel::FFaceMesh>& FaceMesh, FBodyMesh& OutBodyMesh);
+	}
+
+
+	class CADINTERFACES_API FCoreTechFileParser : public ICADFileParser
+	{
+	public:
+
+		/**
+		 * @param 
+		 * @param EnginePluginsPath Full Path of EnginePlugins. Mandatory to set KernelIO to import DWG, or DGN files
+		 */
+		FCoreTechFileParser(FCADFileData& InCADData, const FString& EnginePluginsPath = TEXT(""));
+
+		virtual ECADParsingResult Process() override;
+
+	private:
+		bool ReadNode(CT_OBJECT_ID NodeId, uint32 ParentMaterialHash);
+		bool ReadInstance(CT_OBJECT_ID NodeId, uint32 ParentMaterialHash);
+		bool ReadComponent(CT_OBJECT_ID NodeId, uint32 ParentMaterialHash);
+
+		bool ReadKioBody(CT_OBJECT_ID NodeId, CT_OBJECT_ID ParentId, uint32 ParentMaterialHash, bool bNeedRepair);
+		bool ReadBody(CT_OBJECT_ID NodeId, CT_OBJECT_ID ParentId, uint32 ParentMaterialHash, bool bNeedRepair);
+
+		uint32 GetMaterialNum();
+		void ReadMaterials();
+
+		uint32 GetObjectMaterial(ICADArchiveObject& Object);
+		FArchiveColor& FindOrAddColor(uint32 ColorHId);
+		FArchiveMaterial& FindOrAddMaterial(CT_MATERIAL_ID MaterialId);
+		void SetFaceMainMaterial(FObjectDisplayDataId& InFaceMaterial, FObjectDisplayDataId& InBodyMaterial, FBodyMesh& BodyMesh, int32 FaceIndex);
+
+		void ReadNodeMetaData(CT_OBJECT_ID NodeId, TMap<FString, FString>& OutMetaData);
+		void GetStringMetaDataValue(CT_OBJECT_ID NodeId, const TCHAR* InMetaDataName, FString& OutMetaDataValue);
+
+		CT_FLAGS SetCoreTechImportOption();
+		void GetAttributeValue(CT_ATTRIB_TYPE attrib_type, int ith_field, FString& value);
+
+		FCADFileData& CADFileData;
+	};
+} // CADLibrary
+
+#endif // USE_KERNEL_IO_SDK
diff --git a/Engine/Plugins/Enterprise/DatasmithCADImporter/Source/CADInterfaces/Private/CoreTechFileParserImpl.cpp b/Engine/Plugins/Enterprise/DatasmithCADImporter/Source/CADInterfaces/Private/CoreTechFileParserImpl.cpp
new file mode 100644
--- /dev/null
+++ b/Engine/Plugins/Enterprise/DatasmithCADImporter/Source/CADInterfaces/Private/CoreTechFileParserImpl.cpp
@@ -0,0 +1,1906 @@
+// Copyright Epic Games, Inc. All Rights Reserved.
+
+#include "CoreTechFileParser.h"
+
+#include "CoreTechTypes.h"
+
+#ifdef USE_KERNEL_IO_SDK
+
+#pragma warning(push)
+#pragma warning(disable:4996) // unsafe sprintf
+#pragma warning(disable:4828) // illegal character
+#include "kernel_io/attribute_io/attribute_io.h"
+#include "kernel_io/list_io/list_io.h"
+#include "kernel_io/object_io/asm_io/component_io/component_io.h"
+#include "kernel_io/object_io/asm_io/instance_io/instance_io.h"
+#include "kernel_io/object_io/geom_io/surface_io/surface_io.h"
+#include "kernel_io/object_io/object_io.h"
+#include "kernel_io/object_io/topo_io/body_io/body_io.h"
+#include "kernel_io/object_io/topo_io/face_io/face_io.h"
+#pragma warning(pop)
+
+#include "CADData.h"
+#include "CADKernelTools.h"
+#include "CADOptions.h"
+#include "DatasmithUtils.h"
+
+#include "CADKernel/Core/Entity.h"
+#include "CADKernel/Core/Session.h"
+#include "CADKernel/Core/Types.h"
+
+#include "CADKernel/Mesh/Meshers/ParametricMesher.h"
+#include "CADKernel/Mesh/Structure/ModelMesh.h"
+
+#include "CADKernel/Topo/Model.h"
+#include "CADKernel/Topo/Body.h"
+
+#include "HAL/FileManager.h"
+#include "Internationalization/Text.h"
+#include "Templates/TypeHash.h"
+
+namespace CADLibrary
+{
+	namespace CoreTechFileParserUtils
+	{
+		FString AsFString(CT_STR CtName);
+		uint32 GetSceneFileHash(const uint32 InSGHash, const FImportParameters& ImportParam);
+		uint32 GetGeomFileHash(const uint32 InSGHash, const FImportParameters& ImportParam);
+		void AddFaceIdAttribut(CT_OBJECT_ID NodeId);
+		void GetInstancesAndBodies(CT_OBJECT_ID InComponentId, TArray<CT_OBJECT_ID>& OutInstances, TArray<CT_OBJECT_ID>& OutBodies);
+		void GetCTObjectDisplayDataIds(CT_OBJECT_ID ObjectID, FObjectDisplayDataId& Material);
+		uint32 GetStaticMeshUuid(const TCHAR* OutSgFile, const int32 BodyId);
+		bool GetColor(uint32 ColorHash, FColor& OutColor);
+		bool GetMaterial(uint32 MaterialID, FCADMaterial& OutMaterial);
+		int32 GetIntegerParameterDataValue(CT_OBJECT_ID NodeId, const TCHAR* InMetaDataName);
+	}
+
+	FArchiveMaterial& FCoreTechFileParser::FindOrAddMaterial(CT_MATERIAL_ID MaterialId)
+	{
+		if (FArchiveMaterial* NewMaterial = CADFileData.FindMaterial(MaterialId))
+		{
+			return *NewMaterial;
+		}
+
+		FArchiveMaterial& NewMaterial = CADFileData.AddMaterial(MaterialId);
+		CoreTechFileParserUtils::GetMaterial(MaterialId, NewMaterial.Material);
+		NewMaterial.UEMaterialName = BuildMaterialName(NewMaterial.Material);
+		return NewMaterial;
+	}
+
+	FArchiveColor& FCoreTechFileParser::FindOrAddColor(uint32 ColorHId)
+	{
+		if (FArchiveColor* Color = CADFileData.FindColor(ColorHId))
+		{
+			return *Color;
+		}
+
+		FArchiveColor& NewColor = CADFileData.AddColor(ColorHId);
+		CoreTechFileParserUtils::GetColor(ColorHId, NewColor.Color);
+		NewColor.UEMaterialName = BuildColorName(NewColor.Color);
+		return NewColor;
+	}
+
+	uint32 FCoreTechFileParser::GetObjectMaterial(ICADArchiveObject& Object)
+	{
+		if (FString* Material = Object.MetaData.Find(TEXT("MaterialName")))
+		{
+			return (uint32) FCString::Atoi64(**Material);
+		}
+		if (FString* Material = Object.MetaData.Find(TEXT("ColorName")))
+		{
+			return (uint32)FCString::Atoi64(**Material);
+		}
+		return 0;
+	}
+
+	void FCoreTechFileParser::SetFaceMainMaterial(FObjectDisplayDataId& InFaceMaterial, FObjectDisplayDataId& InBodyMaterial, FBodyMesh& BodyMesh, int32 FaceIndex)
+	{
+		uint32 FaceMaterialHash = 0;
+		uint32 BodyMaterialHash = 0;
+		uint32 FaceColorHash = 0;
+		uint32 BodyColorHash = 0;
+
+		FTessellationData& FaceTessellations = BodyMesh.Faces.Last();
+
+		if (InFaceMaterial.Material > 0)
+		{
+			FArchiveMaterial& Material = FindOrAddMaterial(InFaceMaterial.Material);
+			FaceTessellations.MaterialName = Material.UEMaterialName;
+			BodyMesh.MaterialSet.Add(Material.UEMaterialName);
+		}
+		else if (InBodyMaterial.Material > 0)
+		{
+			FArchiveMaterial& Material = FindOrAddMaterial(InBodyMaterial.Material);
+			FaceTessellations.MaterialName = Material.UEMaterialName;
+			BodyMesh.MaterialSet.Add(Material.UEMaterialName);
+		}
+
+		if (InFaceMaterial.Color > 0)
+		{
+			FArchiveColor& Color = FindOrAddColor(InFaceMaterial.Color);
+			FaceTessellations.ColorName = Color.UEMaterialName;
+			BodyMesh.ColorSet.Add(Color.UEMaterialName);
+		}
+		else if (InBodyMaterial.Color > 0)
+		{
+			FArchiveColor& Color = FindOrAddColor(InBodyMaterial.Color);
+			FaceTessellations.ColorName = Color.UEMaterialName;
+			BodyMesh.ColorSet.Add(Color.UEMaterialName);
+		}
+		else if (InBodyMaterial.DefaultMaterialName)
+		{
+			FaceTessellations.ColorName = InBodyMaterial.DefaultMaterialName;
+			BodyMesh.ColorSet.Add(InBodyMaterial.DefaultMaterialName);
+		}
+	}
+
+	uint32 FCoreTechFileParser::GetMaterialNum()
+	{
+		CT_UINT32 IColor = 1;
+		while (true)
+		{
+			CT_COLOR CtColor;
+			if (CT_MATERIAL_IO::AskIndexedColor(IColor, CtColor) != IO_OK)
+			{
+				break;
+			}
+			IColor++;
+		}
+
+		CT_UINT32 IMaterial = 1;
+		while (true)
+		{
+			CT_COLOR Diffuse, Ambient, Specular;
+			CT_FLOAT Shininess, Transparency, Reflexion;
+			CT_STR Name("");
+			CT_TEXTURE_ID TextureId;
+
+			if (CT_MATERIAL_IO::AskParameters(IMaterial, Name, Diffuse, Ambient, Specular, Shininess, Transparency, Reflexion, TextureId) != IO_OK)
+			{
+				break;
+			}
+			IMaterial++;
+		}
+
+		return IColor + IMaterial - 2;
+	}
+
+	void FCoreTechFileParser::ReadMaterials()
+	{
+		CT_UINT32 MaterialId = 1;
+		while (true)
+		{
+			FCADMaterial Material;
+			bool bReturn = CoreTechFileParserUtils::GetMaterial(MaterialId, Material);
+			if (!bReturn)
+			{
+				break;
+			}
+
+			FArchiveMaterial& MaterialObject = CADFileData.AddMaterial(MaterialId);
+			MaterialObject.UEMaterialName = BuildMaterialName(Material);
+			MaterialObject.Material = Material;
+
+			MaterialId++;
+		}
+	}
+
+	FCoreTechFileParser::FCoreTechFileParser(FCADFileData& InCADData, const FString& EnginePluginsPath)
+		: CADFileData(InCADData)
+	{
+	}
+
+	ECADParsingResult FCoreTechFileParser::Process()
+	{
+		CT_IO_ERROR Result = IO_OK;
+		CT_OBJECT_ID MainId = 0;
+
+		CT_KERNEL_IO::UnloadModel();
+
+		// the parallelization of monolithic JT file is set in SetCoreTechImportOption. Then it's processed as the other exploded formats
+		CT_FLAGS CTImportOption = SetCoreTechImportOption();
+
+		FString LoadOption;
+		CT_UINT32 NumberOfIds = 1;
+
+		// Set specific import option according to format
+		if (CADFileData.HasConfiguration())
+		{
+			if (CADFileData.FileExtension() == TEXT("jt"))
+			{
+				LoadOption = CADFileData.GetConfiguration();
+			}
+			else
+			{
+				NumberOfIds = CT_KERNEL_IO::AskFileNbOfIds(*CADFileData.FilePath());
+				if (NumberOfIds > 1)
+				{
+					CT_UINT32 ActiveConfig = CT_KERNEL_IO::AskFileActiveConfig(*CADFileData.FilePath());
+					for (CT_UINT32 i = 0; i < NumberOfIds; i++)
+					{
+						CT_STR ConfValue = CT_KERNEL_IO::AskFileIdIthName(*CADFileData.FilePath(), i);
+						if (CADFileData.GetConfiguration() == CoreTechFileParserUtils::AsFString(ConfValue))
+						{
+							ActiveConfig = i;
+							break;
+						}
+					}
+
+					CTImportOption |= CT_LOAD_FLAGS_READ_SPECIFIC_OBJECT;
+					LoadOption = FString::FromInt((int32)ActiveConfig);
+				}
+			}
+		}
+
+		CTKIO_ChangeUnit(CADFileData.MetricUnit());
+		Result = CT_KERNEL_IO::LoadFile(*CADFileData.FilePath(), MainId, CTImportOption, 0, *LoadOption);
+		if (Result == IO_ERROR_EMPTY_ASSEMBLY)
+		{
+			CT_KERNEL_IO::UnloadModel();
+			CT_FLAGS CTReImportOption = CTImportOption | CT_LOAD_FLAGS_LOAD_EXTERNAL_REF;
+			CTReImportOption &= ~CT_LOAD_FLAGS_READ_ASM_STRUCT_ONLY;  // BUG CT -> Ticket 11685
+			CTKIO_ChangeUnit(CADFileData.MetricUnit());
+			Result = CT_KERNEL_IO::LoadFile(*CADFileData.FilePath(), MainId, CTReImportOption, 0, *LoadOption);
+		}
+
+		// the file is loaded but it's empty, so no data is generate
+		if (Result == IO_ERROR_EMPTY_ASSEMBLY)
+		{
+			CT_KERNEL_IO::UnloadModel();
+			CADFileData.AddWarningMessages(FString::Printf(TEXT("File %s has been loaded but no assembly has been detected."), *CADFileData.FileName()));
+			return ECADParsingResult::ProcessOk;
+		}
+
+		if (Result != IO_OK && Result != IO_OK_MISSING_LICENSES)
+		{
+			CT_KERNEL_IO::UnloadModel();
+			return ECADParsingResult::ProcessFailed;
+		}
+
+		if (CADFileData.IsCacheDefined())
+		{
+			FString CacheFilePath = CADFileData.GetCADCachePath();
+			if (CacheFilePath != CADFileData.FilePath())
+			{
+				CT_LIST_IO ObjectList;
+				ObjectList.PushBack(MainId);
+
+				CT_IO_ERROR SaveResult = CT_KERNEL_IO::SaveFile(ObjectList, *CacheFilePath, L"Ct");
+			}
+		}
+
+		CoreTechFileParserUtils::AddFaceIdAttribut(MainId);
+
+		if (CADFileData.GetStitchingTechnique() != StitchingNone && !CADFileData.IsCADKernelTessellation())
+		{
+			CADLibrary::CTKIO_Repair(MainId, CADFileData.GetStitchingTechnique(), 10.);
+		}
+
+		CTKIO_SetCoreTechTessellationState(CADFileData.GetImportParameters());
+
+		const CT_OBJECT_TYPE TypeSet[] = { CT_INSTANCE_TYPE, CT_ASSEMBLY_TYPE, CT_PART_TYPE, CT_COMPONENT_TYPE, CT_BODY_TYPE, CT_UNLOADED_COMPONENT_TYPE, CT_UNLOADED_ASSEMBLY_TYPE, CT_UNLOADED_PART_TYPE };
+		enum EObjectTypeIndex : uint8	{ CT_INSTANCE_INDEX = 0, CT_ASSEMBLY_INDEX, CT_PART_INDEX, CT_COMPONENT_INDEX, CT_BODY_INDEX, CT_UNLOADED_COMPONENT_INDEX, CT_UNLOADED_ASSEMBLY_INDEX, CT_UNLOADED_PART_INDEX };
+
+		uint32 NbElements[8], NbTotal = 10;
+		for (int32 index = 0; index < 8; index++)
+		{
+			CT_KERNEL_IO::AskNbObjectsType(NbElements[index], TypeSet[index]);
+			NbTotal += NbElements[index];
+		}
+
+		CADFileData.ReserveBodyMeshes(NbElements[CT_BODY_INDEX]);
+
+		FArchiveSceneGraph& SceneGraphArchive = CADFileData.GetSceneGraphArchive();
+		SceneGraphArchive.BodySet.Reserve(NbElements[CT_BODY_INDEX]);
+		SceneGraphArchive.ComponentSet.Reserve(NbElements[CT_ASSEMBLY_INDEX] + NbElements[CT_PART_INDEX] + NbElements[CT_COMPONENT_INDEX]);
+		SceneGraphArchive.UnloadedComponentSet.Reserve(NbElements[CT_UNLOADED_COMPONENT_INDEX] + NbElements[CT_UNLOADED_ASSEMBLY_INDEX] + NbElements[CT_UNLOADED_PART_INDEX]);
+		SceneGraphArchive.ExternalRefSet.Reserve(NbElements[CT_UNLOADED_COMPONENT_INDEX] + NbElements[CT_UNLOADED_ASSEMBLY_INDEX] + NbElements[CT_UNLOADED_PART_INDEX]);
+		SceneGraphArchive.Instances.Reserve(NbElements[CT_INSTANCE_INDEX]);
+
+		SceneGraphArchive.CADIdToBodyIndex.Reserve(NbElements[CT_BODY_INDEX]);
+		SceneGraphArchive.CADIdToComponentIndex.Reserve(NbElements[CT_ASSEMBLY_INDEX] + NbElements[CT_PART_INDEX] + NbElements[CT_COMPONENT_INDEX]);
+		SceneGraphArchive.CADIdToUnloadedComponentIndex.Reserve(NbElements[CT_UNLOADED_COMPONENT_INDEX] + NbElements[CT_UNLOADED_ASSEMBLY_INDEX] + NbElements[CT_UNLOADED_PART_INDEX]);
+		SceneGraphArchive.CADIdToInstanceIndex.Reserve(NbElements[CT_INSTANCE_INDEX]);
+
+		uint32 MaterialNum = GetMaterialNum();
+		SceneGraphArchive.MaterialHIdToMaterial.Reserve(MaterialNum);
+
+		ReadMaterials();
+
+		// Parse the file
+		uint32 DefaultMaterialHash = 0;
+		bool bReadNodeSucceed = ReadNode(MainId, DefaultMaterialHash);
+		// End of parsing
+
+		CT_STR KernelIO_Version = CT_KERNEL_IO::AskVersion();
+		if (!KernelIO_Version.IsEmpty())
+		{
+			SceneGraphArchive.ComponentSet[0].MetaData.Add(TEXT("KernelIOVersion"), CoreTechFileParserUtils::AsFString(KernelIO_Version));
+		}
+
+		CT_KERNEL_IO::UnloadModel();
+
+		if (!bReadNodeSucceed)
+		{
+			return ECADParsingResult::ProcessFailed;
+		}
+
+		return ECADParsingResult::ProcessOk;
+	}
+
+	CT_FLAGS FCoreTechFileParser::SetCoreTechImportOption()
+	{
+		// Set import option
+		CT_FLAGS Flags = CT_LOAD_FLAGS_USE_DEFAULT;
+
+		const FString& MainFileExt = CADFileData.FileExtension();
+
+		// Parallelization of monolithic JT file,
+		// For JT file, first step the file is read with "Structure only option"
+		// For each body, the JT file is read with "READ_SPECIFIC_OBJECT", Configuration == BodyId
+		if (MainFileExt == TEXT("jt"))
+		{
+			if (!CADFileData.HasConfiguration())
+			{
+				FFileStatData FileStatData = IFileManager::Get().GetStatData(*CADFileData.FilePath());
+
+				if (FileStatData.FileSize > 2e6 /* 2 Mb */ && CADFileData.IsCacheDefined()) // First step 
+				{
+						Flags |= CT_LOAD_FLAGS_READ_ASM_STRUCT_ONLY;
+				}
+			}
+			else // Second step
+			{
+				Flags &= ~CT_LOAD_FLAGS_REMOVE_EMPTY_COMPONENTS;
+				Flags |= CT_LOAD_FLAGS_READ_SPECIFIC_OBJECT;
+			}
+		}
+
+		Flags |= CT_LOAD_FLAGS_READ_META_DATA;
+
+		if (MainFileExt == TEXT("catpart") || MainFileExt == TEXT("catproduct") || MainFileExt == TEXT("cgr"))
+		{
+			Flags |= CT_LOAD_FLAGS_V5_READ_GEOM_SET;
+		}
+
+		// All the BRep topology is not available in IGES import
+		// Ask Kernel IO to complete or create missing topology
+		if (MainFileExt == TEXT("igs") || MainFileExt == TEXT("iges"))
+		{
+			Flags |= CT_LOAD_FLAG_COMPLETE_TOPOLOGY;
+			Flags |= CT_LOAD_FLAG_SEARCH_NEW_TOPOLOGY;
+		}
+
+		// 3dxml file is zipped files, it's full managed by Kernel_io. We cannot read it in sequential mode
+		if (MainFileExt != TEXT("3dxml") && CADFileData.IsSequentialImport())
+		{
+			Flags &= ~CT_LOAD_FLAGS_LOAD_EXTERNAL_REF;
+		}
+
+		return Flags;
+	}
+
+	bool FCoreTechFileParser::ReadNode(CT_OBJECT_ID NodeId, uint32 DefaultMaterialHash)
+	{
+		CT_OBJECT_TYPE Type;
+		CT_OBJECT_IO::AskType(NodeId, Type);
+
+		switch (Type)
+		{
+		case CT_INSTANCE_TYPE:
+			return ReadInstance(NodeId, DefaultMaterialHash);
+
+		case CT_ASSEMBLY_TYPE:
+		case CT_PART_TYPE:
+		case CT_COMPONENT_TYPE:
+			return ReadComponent(NodeId, DefaultMaterialHash);
+
+		case CT_UNLOADED_ASSEMBLY_TYPE:
+		case CT_UNLOADED_COMPONENT_TYPE:
+		case CT_UNLOADED_PART_TYPE:
+			// should not append
+			ensure(false);
+			return false;
+
+		case CT_BODY_TYPE:
+			break;
+
+		//Treat all CT_CURVE_TYPE :
+		case CT_CURVE_TYPE:
+		case CT_C_NURBS_TYPE:
+		case CT_CONICAL_TYPE:
+		case CT_ELLIPSE_TYPE:
+		case CT_CIRCLE_TYPE:
+		case CT_PARABOLA_TYPE:
+		case CT_HYPERBOLA_TYPE:
+		case CT_LINE_TYPE:
+		case CT_C_COMPO_TYPE:
+		case CT_POLYLINE_TYPE:
+		case CT_EQUATION_CURVE_TYPE:
+		case CT_CURVE_ON_SURFACE_TYPE:
+		case CT_INTERSECTION_CURVE_TYPE:
+		default:
+			break;
+		}
+		return true;
+	}
+
+	bool FCoreTechFileParser::ReadComponent(CT_OBJECT_ID ComponentId, uint32 DefaultMaterialHash)
+	{
+		if (CADFileData.HasComponentOfId(ComponentId))
+		{
+			return true;
+		}
+
+		int32 Index = CADFileData.AddComponent(ComponentId);
+		FArchiveComponent& Component = CADFileData.GetComponentAt(Index);
+		ReadNodeMetaData(ComponentId, Component.MetaData);
+
+		if (uint32 MaterialHash = GetObjectMaterial(Component))
+		{
+			DefaultMaterialHash = MaterialHash;
+		}
+
+		TArray<CT_OBJECT_ID> Instances, Bodies;
+		CoreTechFileParserUtils::GetInstancesAndBodies(ComponentId, Instances, Bodies);
+
+		for (CT_OBJECT_ID InstanceId : Instances)
+		{
+			if (ReadInstance(InstanceId, DefaultMaterialHash))
+			{
+				Component.Children.Add(InstanceId);
+			}
+		}
+
+		for (CT_OBJECT_ID BodyId : Bodies)
+		{
+			CT_FLAGS BodyProperties;
+			CT_BODY_IO::AskProperties(BodyId, BodyProperties);
+
+			if (!CADFileData.IsCADKernelTessellation() || !(BodyProperties & CT_BODY_PROP_EXACT))
+			{
+				if (ReadKioBody(BodyId, ComponentId, DefaultMaterialHash, false))
+				{
+					Component.Children.Add(BodyId);
+				}
+			}
+			else
+			{
+				if (ReadBody(BodyId, ComponentId, DefaultMaterialHash, false))
+				{
+					Component.Children.Add(BodyId);
+				}
+			}
+		}
+
+		return true;
+	}
+
+	bool FCoreTechFileParser::ReadInstance(CT_OBJECT_ID InstanceNodeId, uint32 DefaultMaterialHash)
+	{
+		if (CADFileData.HasInstanceOfId(InstanceNodeId))
+		{
+			return true;
+		}
+
+		int32 InstanceIndex = CADFileData.AddInstance(InstanceNodeId);
+		FArchiveInstance& Instance = CADFileData.GetInstanceAt(InstanceIndex);
+		ReadNodeMetaData(InstanceNodeId, Instance.MetaData);
+
+		if (uint32 MaterialHash = GetObjectMaterial(Instance))
+		{
+			DefaultMaterialHash = MaterialHash;
+		}
+
+		// Ask the transformation of the instance
+		if (CT_INSTANCE_IO::AskTransformation(InstanceNodeId, (double*)Instance.TransformMatrix.M) == IO_OK)
+		{
+			if (Instance.TransformMatrix.ContainsNaN())
+			{
+				Instance.TransformMatrix.SetIdentity();
+			}
+		}
+
+		// Ask the reference
+		CT_OBJECT_ID ReferenceNodeId;
+		CT_IO_ERROR CTReturn = CT_INSTANCE_IO::AskChild(InstanceNodeId, ReferenceNodeId);
+		if (CTReturn != CT_IO_ERROR::IO_OK)
+		{
+			return false;
+		}
+		Instance.ReferenceNodeId = ReferenceNodeId;
+
+		CT_OBJECT_TYPE type;
+		CT_OBJECT_IO::AskType(ReferenceNodeId, type);
+		if (type == CT_UNLOADED_PART_TYPE || type == CT_UNLOADED_COMPONENT_TYPE || type == CT_UNLOADED_ASSEMBLY_TYPE)
+		{
+			Instance.bIsExternalRef = true;
+			if (int32* Index = CADFileData.FindUnloadedComponentOfId(ReferenceNodeId))
+			{
+				Instance.ExternalRef = CADFileData.GetExternalReferences(*Index);
+				return true;
+			}
+
+			const FString SupressedEntity = TEXT("Supressed Entity");
+			FString IsSupressedEntity = Instance.MetaData.FindRef(SupressedEntity);
+			if (IsSupressedEntity == TEXT("true"))
+			{
+				return false;
+			}
+
+			CT_STR ComponentFile, FileType;
+			CT_UINT3264 InternalId;
+			CT_COMPONENT_IO::AskExternalDefinition(ReferenceNodeId, ComponentFile, FileType, InternalId);
+			FString ExternalRefFullPath = CoreTechFileParserUtils::AsFString(ComponentFile);
+
+			if (ExternalRefFullPath.IsEmpty())
+			{
+				ExternalRefFullPath = CADFileData.FilePath();
+			}
+
+			FString Configuration;
+			if (CADFileData.FileExtension() == TEXT("jt"))
+			{
+				// Parallelization of monolithic JT file,
+				// is the external reference is the current file ? 
+				// Yes => this is an unloaded part that will be imported with CT_LOAD_FLAGS_READ_SPECIFIC_OBJECT Option
+				// No => the external reference is realy external... 
+				FString ExternalName = FPaths::GetCleanFilename(ExternalRefFullPath);
+				if (ExternalName == CADFileData.FileName())
+				{
+					Configuration = FString::Printf(TEXT("%d"), InternalId);
+				}
+			}
+			else
+			{
+				const FString ConfigName = TEXT("Configuration Name");
+				Configuration = Instance.MetaData.FindRef(ConfigName);
+			}
+
+			int32 UnloadedComponentIndex = CADFileData.AddUnloadedComponent(ReferenceNodeId);
+			FArchiveUnloadedComponent& UnloadedComponent = CADFileData.GetUnloadedComponentAt(UnloadedComponentIndex);
+			ReadNodeMetaData(ReferenceNodeId, UnloadedComponent.MetaData);
+
+			FFileDescription& NewFileDescription = CADFileData.AddExternalRef(*ExternalRefFullPath, *Configuration, *CADFileData.GetCADFileDescription().MainCadFilePath);
+			Instance.ExternalRef = NewFileDescription;
+
+			return true;
+		}
+
+		Instance.bIsExternalRef = false;
+
+		return ReadComponent(ReferenceNodeId, DefaultMaterialHash);
+	}
+
+	bool FCoreTechFileParser::ReadKioBody(CT_OBJECT_ID BodyId, CT_OBJECT_ID ParentId, uint32 DefaultMaterialHash, bool bNeedRepair)
+	{
+		if (CADFileData.HasBodyOfId(BodyId))
+		{
+			return true;
+		}
+
+		// Is this body a constructive geometry ?
+		CT_LIST_IO FaceList;
+		CT_BODY_IO::AskFaces(BodyId, FaceList);
+		if (1 == FaceList.Count())
+		{
+			FaceList.IteratorInitialize();
+			FString Value;
+			GetStringMetaDataValue(FaceList.IteratorIter(), TEXT("Constructive Plane"), Value);
+			if (Value == TEXT("true"))
+			{
+				return false;
+			}
+		}
+
+		int32 BodyIndex = CADFileData.AddBody(BodyId);
+		FArchiveBody& Body = CADFileData.GetBodyAt(BodyIndex);
+
+		ReadNodeMetaData(BodyId, Body.MetaData);
+
+		FBodyMesh& BodyMesh = CADFileData.AddBodyMesh(BodyId);
+
+		if (uint32 MaterialHash = GetObjectMaterial(Body))
+		{
+			DefaultMaterialHash = MaterialHash;
+		}
+
+		Body.MeshActorName = CoreTechFileParserUtils::GetStaticMeshUuid(*CADFileData.GetSceneGraphFileName(), BodyId);
+		BodyMesh.MeshActorName = Body.MeshActorName;
+
+		CT_FLAGS BodyProperties;
+		CT_BODY_IO::AskProperties(BodyId, BodyProperties);
+
+		// Save Body in CT file for re-tessellation before getBody because GetBody can call repair and modify the body (delete and build a new one with a new Id)
+		// CT file is saved only for Exact body i.e. not tessellated body
+		if (CADFileData.IsCacheDefined() && (BodyProperties & CT_BODY_PROP_EXACT))
+		{
+			CT_LIST_IO ObjectList;
+			ObjectList.PushBack(BodyId);
+			FString BodyFile = FString::Printf(TEXT("UEx%08x"), Body.MeshActorName);
+			CT_KERNEL_IO::SaveFile(ObjectList, *FPaths::Combine(CADFileData.GetCADCachePath(), TEXT("body"), BodyFile + TEXT(".ct")), L"Ct");
+		}
+
+		FObjectDisplayDataId BodyMaterial;
+		BodyMaterial.DefaultMaterialName = DefaultMaterialHash;
+		CoreTechFileParserUtils::GetCTObjectDisplayDataIds(BodyId, BodyMaterial);
+
+		TFunction<void(CT_OBJECT_ID, int32, FTessellationData&)> ProcessFace;
+
+		ProcessFace = [&](CT_OBJECT_ID FaceID, int32 Index, FTessellationData& Tessellation)
+		{
+			FObjectDisplayDataId FaceMaterial;
+			CoreTechFileParserUtils::GetCTObjectDisplayDataIds(FaceID, FaceMaterial);
+			SetFaceMainMaterial(FaceMaterial, BodyMaterial, BodyMesh, Index);
+
+			if (CADFileData.NeedUVMapScaling() && Tessellation.TexCoordArray.Num() > 0)
+			{
+				CoreTechFileParserUtils::ScaleUV(FaceID, Tessellation.TexCoordArray, (float) CADFileData.GetScaleFactor());
+			}
+		};
+
+		CoreTechFileParserUtils::GetBodyTessellation(BodyId, BodyMesh, ProcessFace);
+
+		Body.ColorFaceSet = BodyMesh.ColorSet;
+		Body.MaterialFaceSet = BodyMesh.MaterialSet;
+
+		return true;
+	}
+
+#ifdef CORETECHBRIDGE_DEBUG
+	static int32 BodyIndex = 0;
+#endif
+	bool FCoreTechFileParser::ReadBody(CT_OBJECT_ID BodyId, CT_OBJECT_ID ParentId, uint32 DefaultMaterialHash, bool bNeedRepair)
+	{
+		// Is this body a constructive geometry ?
+		CT_LIST_IO FaceList;
+		CT_BODY_IO::AskFaces(BodyId, FaceList);
+		if (1 == FaceList.Count())
+		{
+			FaceList.IteratorInitialize();
+			FString Value;
+			GetStringMetaDataValue(FaceList.IteratorIter(), TEXT("Constructive Plane"), Value);
+			if (Value == TEXT("true"))
+			{
+				return false;
+			}
+		}
+
+		int32 Index = CADFileData.AddBody(BodyId);
+		FArchiveBody& ArchiveBody = CADFileData.GetBodyAt(Index);
+		ReadNodeMetaData(BodyId, ArchiveBody.MetaData);
+
+		FBodyMesh& BodyMesh = CADFileData.AddBodyMesh(BodyId);
+
+		if (uint32 MaterialHash = GetObjectMaterial(ArchiveBody))
+		{
+			DefaultMaterialHash = MaterialHash;
+		}
+
+		ArchiveBody.MeshActorName = CoreTechFileParserUtils::GetStaticMeshUuid(*CADFileData.GetSceneGraphFileName(), BodyId);
+		BodyMesh.MeshActorName = ArchiveBody.MeshActorName;
+
+		{
+			TSharedRef<CADKernel::FSession> CADKernelSession = MakeShared<CADKernel::FSession>(0.00001 / CADFileData.GetMetricUnit());
+
+			TSharedRef<CADKernel::FModel> CADKernelModel = CADKernelSession->GetModel();
+
+			CADKernel::FCoreTechBridge CoreTechBridge(CADKernelSession);
+
+			TSharedRef<CADKernel::FBody> CADKernelBody = CoreTechBridge.AddBody(BodyId);
+			CADKernelModel->Add(CADKernelBody);
+
+#ifdef CORETECHBRIDGE_DEBUG
+			FString FolderName = FPaths::GetCleanFilename(FileDescription.MainCadFilePath);
+			CADKernelSession->SaveDatabase(*FPaths::Combine(CADFileData.CachePath, TEXT("CADKernel"), FolderName, FString::Printf(TEXT("%06d_"), BodyIndex++) + FileDescription.Name + TEXT(".ugeom")));
+#endif
+
+			// Save Body in CADKernelArchive file for re-tessellation
+			if (CADFileData.IsCacheDefined())
+			{
+				FString BodyFilePath = CADFileData.GetBodyCachePath(ArchiveBody.MeshActorName);
+				CADKernelSession->SaveDatabase(*BodyFilePath);
+			}
+
+			// Tessellate the body
+			TSharedRef<CADKernel::FModelMesh> CADKernelModelMesh = CADKernel::FEntity::MakeShared<CADKernel::FModelMesh>();
+
+			FCADKernelTools::DefineMeshCriteria(CADKernelModelMesh, CADFileData.GetImportParameters());
+
+			CADKernel::FParametricMesher Mesher(CADKernelModelMesh);
+			Mesher.MeshEntity(CADKernelModel);
+
+			TFunction<void(FObjectDisplayDataId, FObjectDisplayDataId, int32)> ProcessFace;
+			ProcessFace = [&](FObjectDisplayDataId FaceMaterial, FObjectDisplayDataId BodyMaterial, int32 Index)
+			{
+				SetFaceMainMaterial(FaceMaterial, BodyMaterial, BodyMesh, Index);
+			};
+
+			FCADKernelTools::GetBodyTessellation(CADKernelModelMesh, CADKernelBody, BodyMesh, DefaultMaterialHash, ProcessFace);
+
+			CADKernelSession->Clear();
+		}
+
+		ArchiveBody.ColorFaceSet = BodyMesh.ColorSet;
+		ArchiveBody.MaterialFaceSet = BodyMesh.MaterialSet;
+		return true;
+	}
+
+	void FCoreTechFileParser::GetAttributeValue(CT_ATTRIB_TYPE AttributType, int IthField, FString& Value)
+	{
+		CT_STR               FieldName;
+		CT_ATTRIB_FIELD_TYPE FieldType;
+
+		Value = "";
+
+		if (CT_ATTRIB_DEFINITION_IO::AskFieldDefinition(AttributType, IthField, FieldType, FieldName) != IO_OK)
+		{
+			return;
+		}
+
+		switch (FieldType) {
+			case CT_ATTRIB_FIELD_UNKNOWN:
+			{
+				break;
+			}
+			case CT_ATTRIB_FIELD_INTEGER:
+			{
+				int IValue;
+			if (CT_CURRENT_ATTRIB_IO::AskIntField(IthField, IValue) != IO_OK)
+				{
+					break;
+				}
+				Value = FString::FromInt(IValue);
+				break;
+			}
+			case CT_ATTRIB_FIELD_DOUBLE:
+			{
+				double DValue;
+				if (CT_CURRENT_ATTRIB_IO::AskDblField(IthField, DValue) != IO_OK)
+				{
+					break;
+				}
+				Value = FString::Printf(TEXT("%lf"), DValue);
+				break;
+			}
+			case CT_ATTRIB_FIELD_STRING:
+			{
+				CT_STR StrValue;
+				if (CT_CURRENT_ATTRIB_IO::AskStrField(IthField, StrValue) != IO_OK)
+				{
+					break;
+				}
+				Value = CoreTechFileParserUtils::AsFString(StrValue);
+				break;
+			}
+			case CT_ATTRIB_FIELD_POINTER:
+			{
+				break;
+			}
+		}
+	}
+
+	void FCoreTechFileParser::GetStringMetaDataValue(CT_OBJECT_ID NodeId, const TCHAR* InMetaDataName, FString& OutMetaDataValue)
+	{
+		CT_STR FieldName;
+		CT_UINT32 IthAttrib = 0;
+		while (CT_OBJECT_IO::SearchAttribute(NodeId, CT_ATTRIB_STRING_METADATA, IthAttrib++) == IO_OK)
+		{
+			if (CT_CURRENT_ATTRIB_IO::AskStrField(ITH_STRING_METADATA_NAME, FieldName) != IO_OK)
+			{
+				continue;
+			}
+			if (!FCString::Strcmp(InMetaDataName, *CoreTechFileParserUtils::AsFString(FieldName)))
+			{
+				CT_STR FieldStrValue;
+				CT_CURRENT_ATTRIB_IO::AskStrField(ITH_STRING_METADATA_VALUE, FieldStrValue);
+				OutMetaDataValue = CoreTechFileParserUtils::AsFString(FieldStrValue);
+				return;
+			}
+		}
+	}
+
+	void FCoreTechFileParser::ReadNodeMetaData(CT_OBJECT_ID NodeId, TMap<FString, FString>& OutMetaData)
+	{
+		if (CT_COMPONENT_IO::IsA(NodeId, CT_COMPONENT_TYPE))
+		{
+			CT_STR FileName, FileType;
+			CT_COMPONENT_IO::AskExternalDefinition(NodeId, FileName, FileType);
+			OutMetaData.Add(TEXT("ExternalDefinition"), CoreTechFileParserUtils::AsFString(FileName));
+		}
+
+		CT_SHOW_ATTRIBUTE IsShow = CT_UNKNOWN;
+		if (CT_OBJECT_IO::AskShowAttribute(NodeId, IsShow) == IO_OK)
+		{
+			switch (IsShow)
+			{
+			case CT_SHOW:
+				OutMetaData.Add(TEXT("ShowAttribute"), TEXT("show"));
+				break;
+			case CT_NOSHOW:
+				OutMetaData.Add(TEXT("ShowAttribute"), TEXT("noShow"));
+				break;
+			case CT_UNKNOWN:
+				OutMetaData.Add(TEXT("ShowAttribute"), TEXT("unknown"));
+				break;
+			}
+		}
+
+		CT_UINT32 IthAttrib = 0;
+		while (CT_OBJECT_IO::SearchAttribute(NodeId, CT_ATTRIB_ALL, IthAttrib++) == IO_OK)
+		{
+			// Get the current attribute type
+			CT_ATTRIB_TYPE       AttributeType;
+			CT_STR               TypeName;
+
+			CT_STR               FieldName;
+			CT_STR               FieldStrValue;
+			CT_INT32             FieldIntValue;
+			CT_DOUBLE            FieldDoubleValue0, FieldDoubleValue1, FieldDoubleValue2;
+			FString              FieldValue;
+
+
+			if (CT_CURRENT_ATTRIB_IO::AskAttributeType(AttributeType) != IO_OK)
+			{
+				continue;
+			}
+
+			switch (AttributeType) {
+
+			case CT_ATTRIB_SPLT:
+				break;
+
+			case CT_ATTRIB_NAME:
+				if (CT_CURRENT_ATTRIB_IO::AskStrField(ITH_NAME_VALUE, FieldStrValue) == IO_OK)
+				{
+					OutMetaData.Add(TEXT("CTName"), CoreTechFileParserUtils::AsFString(FieldStrValue));
+				}
+				break;
+
+			case CT_ATTRIB_ORIGINAL_NAME:
+				if (CT_CURRENT_ATTRIB_IO::AskStrField(ITH_NAME_VALUE, FieldStrValue) == IO_OK)
+				{
+					OutMetaData.Add(TEXT("Name"), CoreTechFileParserUtils::AsFString(FieldStrValue));
+				}
+				break;
+
+			case CT_ATTRIB_ORIGINAL_FILENAME:
+				if (CT_CURRENT_ATTRIB_IO::AskStrField(ITH_FILENAME_VALUE, FieldStrValue) == IO_OK)
+				{
+					OutMetaData.Add(TEXT("FileName"), CoreTechFileParserUtils::AsFString(FieldStrValue));
+				}
+				break;
+
+			case CT_ATTRIB_UUID:
+				if (CT_CURRENT_ATTRIB_IO::AskStrField(ITH_UUID_VALUE, FieldStrValue) == IO_OK)
+				{
+					OutMetaData.Add(TEXT("UUID"), CoreTechFileParserUtils::AsFString(FieldStrValue));
+				}
+				break;
+
+			case CT_ATTRIB_INPUT_FORMAT_AND_EMETTOR:
+				if (CT_CURRENT_ATTRIB_IO::AskStrField(ITH_INPUT_FORMAT_AND_EMETTOR, FieldStrValue) == IO_OK)
+				{
+					OutMetaData.Add(TEXT("Input_Format_and_Emitter"), CoreTechFileParserUtils::AsFString(FieldStrValue));
+				}
+				break;
+
+			case CT_ATTRIB_CONFIGURATION_NAME:
+				if (CT_CURRENT_ATTRIB_IO::AskStrField(ITH_NAME_VALUE, FieldStrValue) == IO_OK)
+				{
+					OutMetaData.Add(TEXT("ConfigurationName"), CoreTechFileParserUtils::AsFString(FieldStrValue));
+				}
+				break;
+
+			case CT_ATTRIB_LAYERID:
+				GetAttributeValue(AttributeType, ITH_LAYERID_VALUE, FieldValue);
+				OutMetaData.Add(TEXT("LayerId"), FieldValue);
+				GetAttributeValue(AttributeType, ITH_LAYERID_NAME, FieldValue);
+				OutMetaData.Add(TEXT("LayerName"), FieldValue);
+				GetAttributeValue(AttributeType, ITH_LAYERID_FLAG, FieldValue);
+				OutMetaData.Add(TEXT("LayerFlag"), FieldValue);
+				break;
+
+			case CT_ATTRIB_COLORID:
+				{
+					if (CT_CURRENT_ATTRIB_IO::AskIntField(ITH_COLORID_VALUE, FieldIntValue) != IO_OK)
+					{
+						break;
+					}
+					uint32 ColorId = FieldIntValue;
+
+					uint8 Alpha = 255;
+					if (CT_OBJECT_IO::SearchAttribute(NodeId, CT_ATTRIB_TRANSPARENCY) == IO_OK)
+					{
+						if (CT_CURRENT_ATTRIB_IO::AskDblField(0, FieldDoubleValue0) == IO_OK)
+						{
+							Alpha = FMath::Max((1. - FieldDoubleValue0), FieldDoubleValue0) * 255.;
+						}
+					}
+
+					uint32 ColorHId = BuildColorId(ColorId, Alpha);
+					FArchiveColor& ColorArchive = FindOrAddColor(ColorHId);
+					OutMetaData.Add(TEXT("ColorName"), FString::FromInt(ColorArchive.UEMaterialName));
+
+					FString colorHexa = FString::Printf(TEXT("%02x%02x%02x%02x"), ColorArchive.Color.R, ColorArchive.Color.G, ColorArchive.Color.B, ColorArchive.Color.A);
+					OutMetaData.Add(TEXT("ColorValue"), colorHexa);
+				}
+				break;
+
+			case CT_ATTRIB_MATERIALID:
+			{
+				if (CT_CURRENT_ATTRIB_IO::AskIntField(ITH_MATERIALID_VALUE, FieldIntValue) != IO_OK)
+				{
+					break;
+				}
+				if (FArchiveMaterial* Material = CADFileData.FindMaterial(FieldIntValue))
+				{
+					OutMetaData.Add(TEXT("MaterialName"), FString::FromInt(Material->UEMaterialName));
+				}
+				break;
+			}
+
+			case CT_ATTRIB_TRANSPARENCY:
+				if (CT_CURRENT_ATTRIB_IO::AskDblField(ITH_TRANSPARENCY_VALUE, FieldDoubleValue0) != IO_OK)
+				{
+					break;
+				}
+				FieldIntValue = FMath::Max((1. - FieldDoubleValue0), FieldDoubleValue0) * 255.;
+				OutMetaData.Add(TEXT("Transparency"), FString::FromInt(FieldIntValue));
+				break;
+
+			case CT_ATTRIB_COMMENT:
+				//ITH_COMMENT_POSX, ITH_COMMENT_POSY, ITH_COMMENT_POSZ, ITH_COMMENT_TEXT
+				break;
+
+			case CT_ATTRIB_REFCOUNT:
+				if (CT_CURRENT_ATTRIB_IO::AskIntField(ITH_REFCOUNT_VALUE, FieldIntValue) != IO_OK)
+				{
+					break;
+				}
+				//OutMetaData.Add(TEXT("RefCount"), FString::FromInt(FieldIntValue));
+				break;
+
+			case CT_ATTRIB_TESS_PARAMS:
+			case CT_ATTRIB_COMPARE_RESULT:
+				break;
+
+			case CT_ATTRIB_DENSITY:
+				//ITH_VOLUME_DENSITY_VALUE
+				break;
+
+			case CT_ATTRIB_MASS_PROPERTIES:
+				if (CT_CURRENT_ATTRIB_IO::AskDblField(ITH_MASS_PROPERTIES_AREA, FieldDoubleValue0) != IO_OK)
+				{
+					break;
+				}
+				OutMetaData.Add(TEXT("Area"), FString::Printf(TEXT("%lf"), FieldDoubleValue0));
+				if (CT_CURRENT_ATTRIB_IO::AskDblField(ITH_MASS_PROPERTIES_VOLUME, FieldDoubleValue0) != IO_OK)
+				{
+					break;
+				}
+				OutMetaData.Add(TEXT("Volume"), FString::Printf(TEXT("%lf"), FieldDoubleValue0));
+				if (CT_CURRENT_ATTRIB_IO::AskDblField(ITH_MASS_PROPERTIES_MASS, FieldDoubleValue0) != IO_OK)
+				{
+					break;
+				}
+				OutMetaData.Add(TEXT("Mass"), FString::Printf(TEXT("%lf"), FieldDoubleValue0));
+				if (CT_CURRENT_ATTRIB_IO::AskDblField(ITH_MASS_PROPERTIES_LENGTH, FieldDoubleValue0) != IO_OK)
+				{
+					break;
+				}
+				OutMetaData.Add(TEXT("Length"), FString::Printf(TEXT("%lf"), FieldDoubleValue0));
+				//ITH_MASS_PROPERTIES_COGX, ITH_MASS_PROPERTIES_COGY, ITH_MASS_PROPERTIES_COGZ
+				//ITH_MASS_PROPERTIES_M1, ITH_MASS_PROPERTIES_M2, ITH_MASS_PROPERTIES_M3
+				//ITH_MASS_PROPERTIES_IXXG,ITH_MASS_PROPERTIES_IYYG, ITH_MASS_PROPERTIES_IZZG, ITH_MASS_PROPERTIES_IXYG, ITH_MASS_PROPERTIES_IYZG, ITH_MASS_PROPERTIES_IZXG
+				//ITH_MASS_PROPERTIES_AXIS1X, ITH_MASS_PROPERTIES_AXIS1Y, ITH_MASS_PROPERTIES_AXIS1Z, ITH_MASS_PROPERTIES_AXIS2X, ITH_MASS_PROPERTIES_AXIS2Y, ITH_MASS_PROPERTIES_AXIS2Z, ITH_MASS_PROPERTIES_AXIS3X, ITH_MASS_PROPERTIES_AXIS3Y, ITH_MASS_PROPERTIES_AXIS3Z
+				//ITH_MASS_PROPERTIES_XMIN, ITH_MASS_PROPERTIES_YMIN, ITH_MASS_PROPERTIES_ZMIN, ITH_MASS_PROPERTIES_XMAX, ITH_MASS_PROPERTIES_YMAX, ITH_MASS_PROPERTIES_ZMAX
+				break;
+
+			case CT_ATTRIB_THICKNESS:
+				//ITH_THICKNESS_VALUE
+				break;
+
+			case CT_ATTRIB_INTEGER_METADATA:
+				if (CT_CURRENT_ATTRIB_IO::AskStrField(ITH_INTEGER_METADATA_NAME, FieldName) != IO_OK)
+				{
+					break;
+				}
+				if (CT_CURRENT_ATTRIB_IO::AskIntField(ITH_INTEGER_METADATA_VALUE, FieldIntValue) != IO_OK)
+				{
+					break;
+				}
+				OutMetaData.Add(CoreTechFileParserUtils::AsFString(FieldName), FString::FromInt(FieldIntValue));
+				break;
+
+			case CT_ATTRIB_DOUBLE_METADATA:
+				if (CT_CURRENT_ATTRIB_IO::AskStrField(ITH_DOUBLE_METADATA_NAME, FieldName) != IO_OK)
+				{
+					break;
+				}
+				if (CT_CURRENT_ATTRIB_IO::AskDblField(ITH_DOUBLE_METADATA_VALUE, FieldDoubleValue0) != IO_OK)
+				{
+					break;
+				}
+				OutMetaData.Add(CoreTechFileParserUtils::AsFString(FieldName), FString::Printf(TEXT("%lf"), FieldDoubleValue0));
+				break;
+
+			case CT_ATTRIB_STRING_METADATA:
+				if (CT_CURRENT_ATTRIB_IO::AskStrField(ITH_STRING_METADATA_NAME, FieldName) != IO_OK)
+				{
+					break;
+				}
+				if (CT_CURRENT_ATTRIB_IO::AskStrField(ITH_STRING_METADATA_VALUE, FieldStrValue) != IO_OK)
+				{
+					break;
+				}
+				OutMetaData.Add(CoreTechFileParserUtils::AsFString(FieldName), CoreTechFileParserUtils::AsFString(FieldStrValue));
+				break;
+
+			case CT_ATTRIB_ORIGINAL_UNITS:
+				if (CT_CURRENT_ATTRIB_IO::AskDblField(ITH_ORIGINAL_UNITS_MASS, FieldDoubleValue0) != IO_OK)
+				{
+					break;
+				}
+				if (CT_CURRENT_ATTRIB_IO::AskDblField(ITH_ORIGINAL_UNITS_LENGTH, FieldDoubleValue1) != IO_OK)
+				{
+					break;
+				}
+				if (CT_CURRENT_ATTRIB_IO::AskDblField(ITH_ORIGINAL_UNITS_DURATION, FieldDoubleValue2) != IO_OK)
+				{
+					break;
+				}
+				OutMetaData.Add(TEXT("OriginalUnitsMass"), FString::Printf(TEXT("%lf"), FieldDoubleValue0));
+				OutMetaData.Add(TEXT("OriginalUnitsLength"), FString::Printf(TEXT("%lf"), FieldDoubleValue1));
+				OutMetaData.Add(TEXT("OriginalUnitsDuration"), FString::Printf(TEXT("%lf"), FieldDoubleValue2));
+				break;
+
+			case CT_ATTRIB_ORIGINAL_TOLERANCE:
+			case CT_ATTRIB_IGES_PARAMETERS:
+			case CT_ATTRIB_READ_V4_MARKER:
+				break;
+
+			case CT_ATTRIB_PRODUCT:
+				if (CT_CURRENT_ATTRIB_IO::AskStrField(ITH_PRODUCT_REVISION, FieldStrValue) == IO_OK)
+				{
+					OutMetaData.Add(TEXT("ProductRevision"), CoreTechFileParserUtils::AsFString(FieldStrValue));
+				}
+
+				if (CT_CURRENT_ATTRIB_IO::AskStrField(ITH_PRODUCT_DEFINITION, FieldStrValue) == IO_OK)
+				{
+					OutMetaData.Add(TEXT("ProductDefinition"), CoreTechFileParserUtils::AsFString(FieldStrValue));
+				}
+
+				if (CT_CURRENT_ATTRIB_IO::AskStrField(ITH_PRODUCT_NOMENCLATURE, FieldStrValue) == IO_OK)
+				{
+					OutMetaData.Add(TEXT("ProductNomenclature"), CoreTechFileParserUtils::AsFString(FieldStrValue));
+				}
+
+				if (CT_CURRENT_ATTRIB_IO::AskStrField(ITH_PRODUCT_SOURCE, FieldStrValue) == IO_OK)
+				{
+					OutMetaData.Add(TEXT("ProductSource"), CoreTechFileParserUtils::AsFString(FieldStrValue));
+				}
+
+				if (CT_CURRENT_ATTRIB_IO::AskStrField(ITH_PRODUCT_DESCRIPTION, FieldStrValue) != IO_OK)
+				{
+					OutMetaData.Add(TEXT("ProductDescription"), CoreTechFileParserUtils::AsFString(FieldStrValue));
+				}
+				break;
+
+			case CT_ATTRIB_SIMPLIFY:
+			case CT_ATTRIB_MIDFACE:
+			case CT_ATTRIB_DEBUG_STRING:
+			case CT_ATTRIB_DEFEATURING:
+			case CT_ATTRIB_BREPLINKID:
+			case CT_ATTRIB_MARKUPS_REF:
+			case CT_ATTRIB_COLLISION:
+				break;
+
+			case CT_ATTRIB_EXTERNAL_ID:
+				//ITH_EXTERNAL_ID_VALUE
+				break;
+
+			case CT_ATTRIB_MODIFIER:
+			case CT_ATTRIB_ORIGINAL_SURF_OLD:
+			case CT_ATTRIB_RESULT_BREPLINKID:
+				break;
+
+			case CT_ATTRIB_AREA:
+				//ITH_AREA_VALUE
+				break;
+
+			case CT_ATTRIB_ACIS_SG_PIDNAME:
+			case CT_ATTRIB_CURVE_ORIGINAL_BOUNDARY_PARAMS:
+				break;
+
+			case CT_ATTRIB_INTEGER_PARAMETER:
+				if (CT_CURRENT_ATTRIB_IO::AskStrField(ITH_INTEGER_PARAMETER_NAME, FieldName) != IO_OK)
+				{
+					break;
+				}
+				if (CT_CURRENT_ATTRIB_IO::AskIntField(ITH_INTEGER_PARAMETER_VALUE, FieldIntValue) != IO_OK)
+				{
+					break;
+				}
+				OutMetaData.Add(CoreTechFileParserUtils::AsFString(FieldName), FString::FromInt(FieldIntValue));
+				break;
+
+			case CT_ATTRIB_DOUBLE_PARAMETER:
+				if (CT_CURRENT_ATTRIB_IO::AskStrField(ITH_DOUBLE_PARAMETER_NAME, FieldName) != IO_OK)
+				{
+					break;
+				}
+				if (CT_CURRENT_ATTRIB_IO::AskDblField(ITH_DOUBLE_PARAMETER_VALUE, FieldDoubleValue0) != IO_OK)
+				{
+					break;
+				}
+				OutMetaData.Add(CoreTechFileParserUtils::AsFString(FieldName), FString::Printf(TEXT("%lf"), FieldDoubleValue0));
+				break;
+
+			case CT_ATTRIB_STRING_PARAMETER:
+				if (CT_CURRENT_ATTRIB_IO::AskStrField(ITH_STRING_PARAMETER_NAME, FieldName) != IO_OK)
+				{
+					break;
+				}
+				if (CT_CURRENT_ATTRIB_IO::AskStrField(ITH_STRING_PARAMETER_VALUE, FieldStrValue) != IO_OK)
+				{
+					break;
+				}
+				OutMetaData.Add(CoreTechFileParserUtils::AsFString(FieldName), CoreTechFileParserUtils::AsFString(FieldStrValue));
+				break;
+
+			case CT_ATTRIB_PARAMETER_ARRAY:
+				//ITH_PARAMETER_ARRAY_NAME
+				//ITH_PARAMETER_ARRAY_NUMBER
+				//ITH_PARAMETER_ARRAY_VALUES
+				break;
+
+			case CT_ATTRIB_SAVE_OPTION:
+				if (CT_CURRENT_ATTRIB_IO::AskStrField(ITH_SAVE_OPTION_AUTHOR, FieldStrValue) == IO_OK)
+				{
+					OutMetaData.Add(TEXT("SaveOptionAuthor"), CoreTechFileParserUtils::AsFString(FieldStrValue));
+				}
+
+				if (CT_CURRENT_ATTRIB_IO::AskStrField(ITH_SAVE_OPTION_ORGANIZATION, FieldStrValue) == IO_OK)
+				{
+					OutMetaData.Add(TEXT("SaveOptionOrganization"), CoreTechFileParserUtils::AsFString(FieldStrValue));
+				}
+
+				if (CT_CURRENT_ATTRIB_IO::AskStrField(ITH_SAVE_OPTION_FILE_DESCRIPTION, FieldStrValue) == IO_OK)
+				{
+					OutMetaData.Add(TEXT("SaveOptionFileDescription"), CoreTechFileParserUtils::AsFString(FieldStrValue));
+				}
+
+				if (CT_CURRENT_ATTRIB_IO::AskStrField(ITH_SAVE_OPTION_AUTHORISATION, FieldStrValue) == IO_OK)
+				{
+					OutMetaData.Add(TEXT("SaveOptionAuthorisation"), CoreTechFileParserUtils::AsFString(FieldStrValue));
+				}
+
+				if (CT_CURRENT_ATTRIB_IO::AskStrField(ITH_SAVE_OPTION_PREPROCESSOR, FieldStrValue) == IO_OK)
+				{
+					OutMetaData.Add(TEXT("SaveOptionPreprocessor"), CoreTechFileParserUtils::AsFString(FieldStrValue));
+				}
+				break;
+
+			case CT_ATTRIB_ORIGINAL_ID:
+				GetAttributeValue(AttributeType, ITH_ORIGINAL_ID_VALUE, FieldValue);
+				OutMetaData.Add(TEXT("OriginalId"), FieldValue);
+				break;
+
+			case CT_ATTRIB_ORIGINAL_ID_STRING:
+				if (CT_CURRENT_ATTRIB_IO::AskStrField(ITH_ORIGINAL_ID_VALUE_STRING, FieldStrValue) != IO_OK)
+				{
+					break;
+				}
+				OutMetaData.Add(TEXT("OriginalIdStr"), CoreTechFileParserUtils::AsFString(FieldStrValue));
+				break;
+
+			case CT_ATTRIB_COLOR_RGB_DOUBLE:
+				if (CT_CURRENT_ATTRIB_IO::AskDblField(ITH_ATTRIB_COLOR_R_DOUBLE, FieldDoubleValue0) != IO_OK)
+				{
+					break;
+				}
+				if (CT_CURRENT_ATTRIB_IO::AskDblField(ITH_ATTRIB_COLOR_G_DOUBLE, FieldDoubleValue1) != IO_OK)
+				{
+					break;
+				}
+				if (CT_CURRENT_ATTRIB_IO::AskDblField(ITH_ATTRIB_COLOR_B_DOUBLE, FieldDoubleValue2) != IO_OK)
+				{
+					break;
+				}
+				FieldValue = FString::Printf(TEXT("%lf"), FieldDoubleValue0) + TEXT(", ") + FString::Printf(TEXT("%lf"), FieldDoubleValue1) + TEXT(", ") + FString::Printf(TEXT("%lf"), FieldDoubleValue2);
+				//OutMetaData.Add(TEXT("ColorRGBDouble"), FieldValue);
+				break;
+
+			case CT_ATTRIB_REVERSE_COLORID:
+			case CT_ATTRIB_INITIAL_FILTER:
+			case CT_ATTRIB_ORIGINAL_SURF:
+			case CT_ATTRIB_LINKMANAGER_BRANCH_FACE:
+			case CT_ATTRIB_LINKMANAGER_PMI:
+			case CT_ATTRIB_NULL:
+			case CT_ATTRIB_MEASURE_VALIDATION_ATTRIBUTE:
+				break;
+
+			case CT_ATTRIB_INTEGER_VALIDATION_ATTRIBUTE:
+				if (CT_CURRENT_ATTRIB_IO::AskStrField(ITH_INTEGER_VALIDATION_NAME, FieldName) != IO_OK)
+				{
+					break;
+				}
+				if (CT_CURRENT_ATTRIB_IO::AskIntField(ITH_INTEGER_VALIDATION_VALUE, FieldIntValue) != IO_OK)
+				{
+					break;
+				}
+				OutMetaData.Add(CoreTechFileParserUtils::AsFString(FieldName), FString::FromInt(FieldIntValue));
+				break;
+
+			case CT_ATTRIB_DOUBLE_VALIDATION_ATTRIBUTE:
+				if (CT_CURRENT_ATTRIB_IO::AskStrField(ITH_DOUBLE_VALIDATION_NAME, FieldName) != IO_OK)
+				{
+					break;
+				}
+				if (CT_CURRENT_ATTRIB_IO::AskDblField(ITH_DOUBLE_VALIDATION_VALUE, FieldDoubleValue0) != IO_OK)
+				{
+					break;
+				}
+				OutMetaData.Add(CoreTechFileParserUtils::AsFString(FieldName), FString::Printf(TEXT("%lf"), FieldDoubleValue0));
+				break;
+
+			case CT_ATTRIB_STRING_VALIDATION_ATTRIBUTE:
+				if (CT_CURRENT_ATTRIB_IO::AskStrField(ITH_STRING_VALIDATION_NAME, FieldName) != IO_OK)
+				{
+					break;
+				}
+				if (CT_CURRENT_ATTRIB_IO::AskStrField(ITH_STRING_VALIDATION_VALUE, FieldStrValue) != IO_OK)
+				{
+					break;
+				}
+				OutMetaData.Add(CoreTechFileParserUtils::AsFString(FieldName), CoreTechFileParserUtils::AsFString(FieldStrValue));
+				break;
+
+			case CT_ATTRIB_BOUNDING_BOX:
+				//ITH_BOUNDING_BOX_XMIN, ITH_BOUNDING_BOX_YMIN, ITH_BOUNDING_BOX_ZMIN, ITH_BOUNDING_BOX_XMAX, ITH_BOUNDING_BOX_YMAX, ITH_BOUNDING_BOX_ZMAX
+				break;
+
+			case CT_ATTRIB_DATABASE:
+			case CT_ATTRIB_CURVE_FONT:
+			case CT_ATTRIB_CURVE_WEIGHT:
+			case CT_ATTRIB_COMPARE_TOPO:
+			case CT_ATTRIB_MONIKER_GUID_TABLE:
+			case CT_ATTRIB_MONIKER_DATA:
+			case CT_ATTRIB_MONIKER_BODY_ID:
+			case CT_ATTRIB_NO_INSTANCE:
+				break;
+
+			case CT_ATTRIB_GROUPNAME:
+				if (CT_CURRENT_ATTRIB_IO::AskStrField(ITH_GROUPNAME_VALUE, FieldStrValue) != IO_OK)
+				{
+					break;
+				}
+				OutMetaData.Add(TEXT("GroupName"), CoreTechFileParserUtils::AsFString(FieldStrValue));
+				break;
+
+			case CT_ATTRIB_ANALYZE_ID:
+			case CT_ATTRIB_ANALYZER_DISPLAY_MODE:
+			case CT_ATTRIB_ANIMATION_ID:
+			case CT_ATTRIB_PROJECTED_SURFACE_ID:
+			case CT_ATTRIB_ANALYZE_LINK:
+			case CT_ATTRIB_TOPO_EVENT_ID:
+			case CT_ATTRIB_ADDITIVE_MANUFACTURING:
+			case CT_ATTRIB_MOLDING_RESULT:
+			case CT_ATTRIB_AMF_ID:
+			case CT_ATTRIB_PARAMETER_LINK:
+				break;
+
+			default:
+				break;
+			}
+		}
+
+		// Clean metadata value i.e. remove all unprintable characters
+		for (auto& MetaPair : OutMetaData)
+		{
+			FDatasmithUtils::SanitizeStringInplace(MetaPair.Value);
+		}
+	}
+
+	namespace CoreTechFileParserUtils
+	{
+		template<typename ValueType>
+		void FillArrayOfVector(int32 ElementCount, void* InCTValueArray, FVector* OutValueArray)
+		{
+			ValueType* Values = (ValueType*)InCTValueArray;
+			for (int Indice = 0; Indice < ElementCount; ++Indice)
+			{
+				OutValueArray[Indice].Set((float)Values[Indice * 3], (float)Values[Indice * 3 + 1], (float)Values[Indice * 3 + 2]);
+			}
+		}
+
+		template<typename ValueType>
+		void FillArrayOfVector2D(int32 ElementCount, void* InCTValueArray, FVector2D* OutValueArray)
+		{
+			ValueType* Values = (ValueType*)InCTValueArray;
+			for (int Indice = 0; Indice < ElementCount; ++Indice)
+			{
+				OutValueArray[Indice].Set((float)Values[Indice * 2], (float)Values[Indice * 2 + 1]);
+			}
+		}
+
+		template<typename ValueType>
+		void FillArrayOfInt(int32 ElementCount, void* InCTValueArray, int32* OutValueArray)
+		{
+			ValueType* Values = (ValueType*)InCTValueArray;
+			for (int Indice = 0; Indice < ElementCount; ++Indice)
+			{
+				OutValueArray[Indice] = (int32)Values[Indice];
+			}
+		}
+
+		double Distance(const CT_COORDINATE& Point1, const CT_COORDINATE& Point2)
+		{
+			return sqrt((Point2.xyz[0] - Point1.xyz[0]) * (Point2.xyz[0] - Point1.xyz[0]) + (Point2.xyz[1] - Point1.xyz[1]) * (Point2.xyz[1] - Point1.xyz[1]) + (Point2.xyz[2] - Point1.xyz[2]) * (Point2.xyz[2] - Point1.xyz[2]));
+		};
+
+		void ScaleUV(CT_OBJECT_ID FaceID, TArray<FVector2D>& TexCoordArray, float Scale)
+		{
+			float VMin, VMax, UMin, UMax;
+			VMin = UMin = HUGE_VALF;
+			VMax = UMax =  -HUGE_VALF;
+
+			for (const FVector2D& TexCoord : TexCoordArray)
+			{
+				UMin = FMath::Min(TexCoord[0], UMin);
+				UMax = FMath::Max(TexCoord[0], UMax);
+				VMin = FMath::Min(TexCoord[1], VMin);
+				VMax = FMath::Max(TexCoord[1], VMax);
+			}
+
+			double PuMin, PuMax, PvMin, PvMax;
+			PuMin = PvMin = HUGE_VALF;
+			PuMax = PvMax = -HUGE_VALF;
+
+			// fast UV min max 
+			CT_FACE_IO::AskUVminmax(FaceID, PuMin, PuMax, PvMin, PvMax);
+
+			const uint32 NbIsoCurves = 7;
+
+			// Compute Point grid on the restricted surface defined by [PuMin, PuMax], [PvMin, PvMax]
+			CT_OBJECT_ID SurfaceID;
+			CT_ORIENTATION Orientation;
+			CT_FACE_IO::AskSurface(FaceID, SurfaceID, Orientation);
+
+			CT_OBJECT_TYPE SurfaceType;
+			CT_SURFACE_IO::AskType(SurfaceID, SurfaceType);
+
+			float DeltaU = (PuMax - PuMin) / (NbIsoCurves - 1);
+			float DeltaV = (PvMax - PvMin) / (NbIsoCurves - 1);
+			float U = PuMin, V = PvMin;
+
+			CT_COORDINATE NodeMatrix[121];
+
+			for (int32 IndexI = 0; IndexI < NbIsoCurves; IndexI++)
+			{
+				for (int32 IndexJ = 0; IndexJ < NbIsoCurves; IndexJ++)
+				{
+					CT_SURFACE_IO::Evaluate(SurfaceID, U, V, NodeMatrix[IndexI * NbIsoCurves + IndexJ]);
+					V += DeltaV;
+				}
+				U += DeltaU;
+				V = PvMin;
+			}
+
+			// Compute length of 7 iso V line
+			float LengthU[NbIsoCurves];
+			float LengthUMin = HUGE_VAL;
+			float LengthUMax = 0;
+			float LengthUMed = 0;
+
+			for (int32 IndexJ = 0; IndexJ < NbIsoCurves; IndexJ++)
+			{
+				LengthU[IndexJ] = 0;
+				for (int32 IndexI = 0; IndexI < (NbIsoCurves - 1); IndexI++)
+				{
+					LengthU[IndexJ] += Distance(NodeMatrix[IndexI * NbIsoCurves + IndexJ], NodeMatrix[(IndexI + 1) * NbIsoCurves + IndexJ]);
+				}
+				LengthUMed += LengthU[IndexJ];
+				LengthUMin = FMath::Min(LengthU[IndexJ], LengthUMin);
+				LengthUMax = FMath::Max(LengthU[IndexJ], LengthUMax);
+			}
+			LengthUMed /= NbIsoCurves;
+			LengthUMed = LengthUMed * 2 / 3 + LengthUMax / 3;
+
+			// Compute length of 7 iso U line
+			float LengthV[NbIsoCurves];
+			float LengthVMin = HUGE_VAL;
+			float LengthVMax = 0;
+			float LengthVMed = 0;
+
+			for (int32 IndexI = 0; IndexI < NbIsoCurves; IndexI++)
+			{
+				LengthV[IndexI] = 0;
+				for (int32 IndexJ = 0; IndexJ < (NbIsoCurves - 1); IndexJ++)
+				{
+					LengthV[IndexI] += Distance(NodeMatrix[IndexI * NbIsoCurves + IndexJ], NodeMatrix[IndexI * NbIsoCurves + IndexJ + 1]);
+				}
+				LengthVMed += LengthV[IndexI];
+				LengthVMin = FMath::Min(LengthV[IndexI], LengthVMin);
+				LengthVMax = FMath::Max(LengthV[IndexI], LengthVMax);
+			}
+			LengthVMed /= NbIsoCurves;
+			LengthVMed = LengthVMed * 2 / 3 + LengthVMax / 3;
+
+			switch (SurfaceType)
+			{
+			case CT_CONE_TYPE:
+			case CT_CYLINDER_TYPE:
+			case CT_SPHERE_TYPE:
+			case CT_TORUS_TYPE:
+				Swap(LengthUMed, LengthVMed);
+				break;
+			case CT_S_REVOL_TYPE:
+				// Need swap ?
+				// Swap(LengthUMed, LengthVMed);
+				break;
+			case CT_S_NURBS_TYPE:
+			case CT_PLANE_TYPE:
+			case CT_S_OFFSET_TYPE:
+			case CT_S_RULED_TYPE:
+			case CT_TABULATED_RULED_TYPE:
+			case CT_S_LINEARTRANSFO_TYPE:
+			case CT_S_NONLINEARTRANSFO_TYPE:
+			case CT_S_BLEND_TYPE:
+			default:
+				break;
+			}
+
+			// scale the UV map
+			// 0.1 define UV in cm and not in mm
+			float VScale = Scale * LengthVMed * 1 / (VMax - VMin) / 100;
+			float UScale = Scale * LengthUMed * 1 / (UMax - UMin) / 100;
+
+			for (FVector2D& TexCoord : TexCoordArray)
+			{
+				TexCoord[0] *= UScale;
+				TexCoord[1] *= VScale;
+			}
+		}
+
+		uint32 GetFaceTessellation(CT_OBJECT_ID FaceID, FTessellationData& Tessellation)
+		{
+			CT_IO_ERROR Error = IO_OK;
+
+			CT_UINT32         VertexCount;
+			CT_UINT32         NormalCount;
+			CT_UINT32         IndexCount;
+			CT_TESS_DATA_TYPE VertexType;
+			CT_TESS_DATA_TYPE TexCoordType;
+			CT_TESS_DATA_TYPE NormalType;
+			CT_LOGICAL        HasRGBColor;
+			CT_UINT16         UserSize;
+			CT_TESS_DATA_TYPE IndexType;
+			void*             VertexArray;
+			void*             TexCoordArray;
+			void*             NormalArray;
+			void*             ColorArray;
+			void*             UserArray;
+			void*             IndexArray;
+
+			Error = CT_FACE_IO::AskTesselation(FaceID, VertexCount, NormalCount, IndexCount,
+				VertexType, TexCoordType, NormalType, HasRGBColor, UserSize, IndexType,
+				VertexArray, TexCoordArray, NormalArray, ColorArray, UserArray, IndexArray);
+
+			// Something wrong happened, either an error or no data to collect
+			if (Error != IO_OK || VertexArray == nullptr || IndexArray == nullptr || IndexCount == 0)
+			{
+				VertexCount = NormalCount = IndexCount = 0;
+				return 0;
+			}
+
+			Tessellation.PatchId = CoreTechFileParserUtils::GetIntegerParameterDataValue(FaceID, TEXT("DatasmithFaceId"));
+
+			Tessellation.VertexIndices.SetNum(IndexCount);
+
+			switch (IndexType)
+			{
+			case CT_TESS_UBYTE:
+				FillArrayOfInt<uint8>(IndexCount, IndexArray, Tessellation.VertexIndices.GetData());
+				break;
+			case CT_TESS_USHORT:
+				FillArrayOfInt<uint16>(IndexCount, IndexArray, Tessellation.VertexIndices.GetData());
+				break;
+			case CT_TESS_UINT:
+				FillArrayOfInt<uint32>(IndexCount, IndexArray, Tessellation.VertexIndices.GetData());
+				break;
+			}
+
+			Tessellation.PositionArray.SetNum(VertexCount);
+			switch (VertexType)
+			{
+			case CT_TESS_FLOAT:
+				FillArrayOfVector<float>(VertexCount, VertexArray, Tessellation.PositionArray.GetData());
+				break;
+			case CT_TESS_DOUBLE:
+				FillArrayOfVector<double>(VertexCount, VertexArray, Tessellation.PositionArray.GetData());
+				break;
+			}
+
+			Tessellation.NormalArray.SetNum(NormalCount);
+			switch (NormalType)
+			{
+			case CT_TESS_BYTE:
+				Tessellation.NormalArray.SetNumZeroed(NormalCount);
+				break;
+			case CT_TESS_SHORT:
+			{
+				int8* InCTValueArray = (int8*)NormalArray;
+				for (CT_UINT32 Indice = 0; Indice < NormalCount; ++Indice)
+				{
+					Tessellation.NormalArray[Indice].Set(((float)InCTValueArray[Indice]) / 255.f, ((float)InCTValueArray[Indice + 1]) / 255.f, ((float)InCTValueArray[Indice + 2]) / 255.f);
+				}
+				break;
+			}
+			case CT_TESS_FLOAT:
+				FillArrayOfVector<float>(NormalCount, NormalArray, Tessellation.NormalArray.GetData());
+				break;
+			}
+
+			if (TexCoordArray)
+			{
+				Tessellation.TexCoordArray.SetNum(VertexCount);
+				switch (TexCoordType)
+				{
+				case CT_TESS_SHORT:
+				{
+					int8* InCTValueArray = (int8*)TexCoordArray;
+					for (CT_UINT32 Indice = 0; Indice < VertexCount; ++Indice)
+					{
+						Tessellation.TexCoordArray[Indice].Set(((float)InCTValueArray[Indice]) / 255.f, ((float)InCTValueArray[Indice + 1]) / 255.f);
+					}
+					break;
+				}
+				case CT_TESS_FLOAT:
+					FillArrayOfVector2D<float>(VertexCount, TexCoordArray, Tessellation.TexCoordArray.GetData());
+					break;
+				case CT_TESS_DOUBLE:
+					FillArrayOfVector2D<double>(VertexCount, TexCoordArray, Tessellation.TexCoordArray.GetData());
+					break;
+				}
+			}
+
+			return Tessellation.VertexIndices.Num() / 3;
+		}
+
+
+		void GetCTObjectDisplayDataIds(CT_OBJECT_ID ObjectID, FObjectDisplayDataId& Material)
+		{
+			if (CT_OBJECT_IO::SearchAttribute(ObjectID, CT_ATTRIB_MATERIALID) == IO_OK)
+			{
+				CT_UINT32 MaterialId = 0;
+				if (CT_CURRENT_ATTRIB_IO::AskIntField(ITH_MATERIALID_VALUE, MaterialId) == IO_OK && MaterialId > 0)
+				{
+					Material.Material = (uint32)MaterialId;
+				}
+			}
+
+			if (CT_OBJECT_IO::SearchAttribute(ObjectID, CT_ATTRIB_COLORID) == IO_OK)
+			{
+				CT_UINT32 ColorId = 0;
+				if (CT_CURRENT_ATTRIB_IO::AskIntField(ITH_COLORID_VALUE, ColorId) == IO_OK && ColorId > 0)
+				{
+					uint8 alpha = 255;
+					if (CT_OBJECT_IO::SearchAttribute(ObjectID, CT_ATTRIB_TRANSPARENCY) == IO_OK)
+					{
+						CT_DOUBLE dbl_value = 0.;
+						if (CT_CURRENT_ATTRIB_IO::AskDblField(0, dbl_value) == IO_OK && dbl_value >= 0.0 && dbl_value <= 1.0)
+						{
+							alpha = (uint8)int((1. - dbl_value) * 255.);
+						}
+					}
+					Material.Color = BuildColorId(ColorId, alpha);
+				}
+			}
+		}
+
+		FString AsFString(CT_STR CtName)
+		{
+			return CtName.IsEmpty() ? FString() : CtName.toUnicode();
+		};
+
+		uint32 GetSceneFileHash(const uint32 InSGHash, const FImportParameters& ImportParam)
+		{
+			uint32 FileHash = HashCombine(InSGHash, ::GetTypeHash(ImportParam.StitchingTechnique));
+			return FileHash;
+		}
+
+		uint32 GetGeomFileHash(const uint32 InSGHash, const FImportParameters& ImportParam)
+		{
+			uint32 FileHash = InSGHash;
+			FileHash = HashCombine(FileHash, ::GetTypeHash(ImportParam.ChordTolerance));
+			FileHash = HashCombine(FileHash, ::GetTypeHash(ImportParam.MaxEdgeLength));
+			FileHash = HashCombine(FileHash, ::GetTypeHash(ImportParam.MaxNormalAngle));
+			FileHash = HashCombine(FileHash, ::GetTypeHash(ImportParam.MetricUnit));
+			FileHash = HashCombine(FileHash, ::GetTypeHash(ImportParam.ScaleFactor));
+			FileHash = HashCombine(FileHash, ::GetTypeHash(ImportParam.StitchingTechnique));
+			return FileHash;
+		}
+
+		bool GetColor(uint32 ColorUuid, FColor& OutColor)
+		{
+			uint32 ColorId;
+			uint8 Alpha;
+			GetCTColorIdAlpha(ColorUuid, ColorId, Alpha);
+
+			CT_COLOR CtColor = { 200, 200, 200 };
+			if (ColorId > 0)
+			{
+				if (CT_MATERIAL_IO::AskIndexedColor((CT_OBJECT_ID)ColorId, CtColor) != IO_OK)
+				{
+					return false;
+				}
+			}
+
+			OutColor.R = CtColor[0];
+			OutColor.G = CtColor[1];
+			OutColor.B = CtColor[2];
+			OutColor.A = Alpha;
+			return true;
+		}
+
+		bool GetMaterial(uint32 MaterialId, FCADMaterial& OutMaterial)
+		{
+			//	// Ref. BaseHelper.cpp
+			CT_STR CtName;
+			CT_COLOR CtDiffuse = { 200, 200, 200 }, CtAmbient = { 200, 200, 200 }, CtSpecular = { 200, 200, 200 };
+			CT_FLOAT CtShininess = 0.f, CtTransparency = 0.f, CtReflexion = 0.f;
+			CT_TEXTURE_ID CtTextureId = 0;
+			if (MaterialId)
+			{
+				CT_IO_ERROR bReturn = CT_MATERIAL_IO::AskParameters(MaterialId, CtName, CtDiffuse, CtAmbient, CtSpecular, CtShininess, CtTransparency, CtReflexion, CtTextureId);
+				if (bReturn != IO_OK)
+				{
+					return false;
+				}
+			}
+
+			CT_STR CtTextureName = "";
+			if (CtTextureId)
+			{
+				CT_INT32  Width = 0, Height = 0;
+				if (!(CT_TEXTURE_IO::AskParameters(CtTextureId, CtTextureName, Width, Height) == IO_OK && Width && Height))
+				{
+					CtTextureName = "";
+				}
+			}
+
+			OutMaterial.MaterialName = AsFString(CtName);
+			OutMaterial.Diffuse = FColor(CtDiffuse[0], CtDiffuse[1], CtDiffuse[2], 255);
+			OutMaterial.Ambient = FColor(CtAmbient[0], CtAmbient[1], CtAmbient[2], 255);
+			OutMaterial.Specular = FColor(CtSpecular[0], CtSpecular[1], CtSpecular[2], 255);
+			OutMaterial.Shininess = CtShininess;
+			OutMaterial.Transparency = CtTransparency;
+			OutMaterial.Reflexion = CtReflexion;
+			OutMaterial.TextureName = AsFString(CtTextureName);
+
+			return true;
+		}
+
+		uint32 GetStaticMeshUuid(const TCHAR* OutSgFile, const int32 BodyId)
+		{
+			uint32 BodyUUID = ::GetTypeHash(OutSgFile);
+			BodyUUID = HashCombine(BodyUUID, ::GetTypeHash(BodyId));
+			return BodyUUID;
+		}
+
+		// For each face, adds an integer parameter representing the id of the face to avoid re-indentation of faces in sub CT file
+		// Used by Re-tessellation Rule to Skip Deleted Surfaces
+		void AddFaceIdAttribut(CT_OBJECT_ID NodeId)
+		{
+			CT_OBJECT_TYPE Type;
+			CT_OBJECT_IO::AskType(NodeId, Type);
+
+			switch (Type)
+			{
+			case CT_INSTANCE_TYPE:
+			{
+				CT_OBJECT_ID ReferenceNodeId;
+				if (CT_INSTANCE_IO::AskChild(NodeId, ReferenceNodeId) == IO_OK)
+				{
+					AddFaceIdAttribut(ReferenceNodeId);
+				}
+				break;
+			}
+
+			case CT_ASSEMBLY_TYPE:
+			case CT_PART_TYPE:
+			case CT_COMPONENT_TYPE:
+			{
+				CT_LIST_IO Children;
+				if (CT_COMPONENT_IO::AskChildren(NodeId, Children) == IO_OK)
+				{
+					Children.IteratorInitialize();
+					CT_OBJECT_ID ChildId;
+					while ((ChildId = Children.IteratorIter()) != 0)
+					{
+						AddFaceIdAttribut(ChildId);
+					}
+				}
+
+				break;
+			}
+
+			case CT_BODY_TYPE:
+			{
+				CT_LIST_IO FaceList;
+				CT_BODY_IO::AskFaces(NodeId, FaceList);
+
+				CT_OBJECT_ID FaceID;
+				FaceList.IteratorInitialize();
+				while ((FaceID = FaceList.IteratorIter()) != 0)
+				{
+					CT_OBJECT_IO::AddAttribute(FaceID, CT_ATTRIB_INTEGER_PARAMETER);
+
+					ensure(CT_CURRENT_ATTRIB_IO::SetStrField(ITH_INTEGER_PARAMETER_NAME, "DatasmithFaceId") == IO_OK);
+					ensure(CT_CURRENT_ATTRIB_IO::SetIntField(ITH_INTEGER_PARAMETER_VALUE, FaceID) == IO_OK);
+				}
+				break;
+			}
+
+			default:
+				break;
+			}
+		}
+
+		void GetInstancesAndBodies(CT_OBJECT_ID InComponentId, TArray<CT_OBJECT_ID>& OutInstances, TArray<CT_OBJECT_ID>& OutBodies)
+		{
+			CT_LIST_IO Children;
+			CT_COMPONENT_IO::AskChildren(InComponentId, Children);
+
+			int32 NbChildren = Children.Count();
+			OutInstances.Empty(NbChildren);
+			OutBodies.Empty(NbChildren);
+
+			Children.IteratorInitialize();
+			CT_OBJECT_ID ChildId;
+			while ((ChildId = Children.IteratorIter()) != 0)
+			{
+				CT_OBJECT_TYPE Type;
+				CT_OBJECT_IO::AskType(ChildId, Type);
+
+				switch (Type)
+				{
+				case CT_INSTANCE_TYPE:
+				{
+					OutInstances.Add(ChildId);
+					break;
+				}
+
+				case CT_BODY_TYPE:
+				{
+					OutBodies.Add(ChildId);
+					break;
+				}
+
+				// we don't manage CURVE, POINT, and COORDSYSTEM (the other kind of child of the component).
+				default:
+					break;
+				}
+			}
+		}
+
+		uint32 GetBodiesFaceSetNum(TArray<CT_OBJECT_ID>& BodySet)
+		{
+			uint32 size = 0;
+			for (int Index = 0; Index < BodySet.Num(); Index++)
+			{
+				// Loop through the face of the first body and collect material data
+				CT_LIST_IO FaceList;
+				CT_BODY_IO::AskFaces(BodySet[Index], FaceList);
+				size += FaceList.Count();
+			}
+			return size;
+		}
+
+		int32 GetIntegerParameterDataValue(CT_OBJECT_ID NodeId, const TCHAR* InMetaDataName)
+		{
+			CT_STR FieldName;
+			CT_UINT32 IthAttrib = 0;
+			int32 IntegerParameterValue = 0;
+			while (CT_OBJECT_IO::SearchAttribute(NodeId, CT_ATTRIB_INTEGER_PARAMETER, IthAttrib++) == IO_OK)
+			{
+				if (CT_CURRENT_ATTRIB_IO::AskStrField(ITH_INTEGER_PARAMETER_NAME, FieldName) != IO_OK)
+				{
+					continue;
+				}
+				if (!FCString::Strcmp(InMetaDataName, *AsFString(FieldName)))
+				{
+					CT_CURRENT_ATTRIB_IO::AskIntField(ITH_INTEGER_PARAMETER_VALUE, IntegerParameterValue);
+					break;
+				}
+			}
+			return IntegerParameterValue;
+		}
+
+		uint32 GetSize(CT_TESS_DATA_TYPE type)
+		{
+			switch (type)
+			{
+			case CT_TESS_USE_DEFAULT:
+				return sizeof(uint32);
+			case CT_TESS_UBYTE:
+				return sizeof(uint8_t);
+			case CT_TESS_BYTE:
+				return sizeof(int8_t);
+			case CT_TESS_USHORT:
+				return sizeof(int16_t);
+			case CT_TESS_SHORT:
+				return sizeof(uint16_t);
+			case CT_TESS_UINT:
+				return sizeof(uint32);
+			case CT_TESS_INT:
+				return sizeof(int32);
+			case CT_TESS_ULONG:
+				return sizeof(uint64);
+			case CT_TESS_LONG:
+				return sizeof(int64);
+			case CT_TESS_FLOAT:
+				return sizeof(float);
+			case CT_TESS_DOUBLE:
+				return sizeof(double);
+			}
+			return 0;
+		}
+
+		void GetBodyTessellation(CT_OBJECT_ID BodyId, FBodyMesh& OutBodyMesh, TFunction<void(CT_OBJECT_ID, int32, FTessellationData&)> ProcessFace)
+		{
+			FBox& BBox = OutBodyMesh.BBox;
+
+			// Compute Body BBox based on CAD data
+			uint32 VerticesSize;
+			CT_BODY_IO::AskVerticesSizeArray(BodyId, VerticesSize);
+
+			TArray<CT_COORDINATE> VerticesArray;
+			VerticesArray.SetNum(VerticesSize);
+			CT_BODY_IO::AskVerticesArray(BodyId, VerticesArray.GetData());
+
+			for (const CT_COORDINATE& Point : VerticesArray)
+			{
+				BBox += FVector((float)Point.xyz[0], (float)Point.xyz[1], (float)Point.xyz[2]);
+			}
+
+			CT_LIST_IO FaceList;
+			CT_BODY_IO::AskFaces(BodyId, FaceList);
+			uint32 FaceSize = FaceList.Count();
+
+			// Allocate memory space for tessellation data
+			OutBodyMesh.Faces.Reserve(FaceSize);
+			OutBodyMesh.ColorSet.Reserve(FaceSize);
+			OutBodyMesh.MaterialSet.Reserve(FaceSize);
+
+			// Loop through the face of bodies and collect all tessellation data
+			int32 FaceIndex = 0;
+			CT_OBJECT_ID FaceID;
+			FaceList.IteratorInitialize();
+			while ((FaceID = FaceList.IteratorIter()) != 0)
+			{
+				FTessellationData& Tessellation = OutBodyMesh.Faces.Emplace_GetRef();
+				uint32 TriangleNum = GetFaceTessellation(FaceID, Tessellation);
+
+				if (TriangleNum == 0)
+				{
+					continue;
+				}
+
+				OutBodyMesh.TriangleCount += TriangleNum;
+
+				if (ProcessFace)
+				{
+					ProcessFace(FaceID, FaceIndex, Tessellation);
+				}
+				FaceIndex++;
+			}
+		}
+
+	}
+
+}
+
+#endif // USE_KERNEL_IO_SDK
diff --git a/Engine/Plugins/Enterprise/DatasmithCADImporter/Source/CADInterfaces/Private/CoreTechFileReader.cpp b/Engine/Plugins/Enterprise/DatasmithCADImporter/Source/CADInterfaces/Private/CoreTechFileReader.cpp
deleted file mode 100644
diff --git a/Engine/Plugins/Enterprise/DatasmithCADImporter/Source/CADInterfaces/Private/CoreTechFileReader.h b/Engine/Plugins/Enterprise/DatasmithCADImporter/Source/CADInterfaces/Private/CoreTechFileReader.h
deleted file mode 100644
diff --git a/Engine/Plugins/Enterprise/DatasmithCADImporter/Source/CADInterfaces/Private/CoreTechInterfaceImpl.cpp b/Engine/Plugins/Enterprise/DatasmithCADImporter/Source/CADInterfaces/Private/CoreTechInterfaceImpl.cpp
--- a/Engine/Plugins/Enterprise/DatasmithCADImporter/Source/CADInterfaces/Private/CoreTechInterfaceImpl.cpp
+++ b/Engine/Plugins/Enterprise/DatasmithCADImporter/Source/CADInterfaces/Private/CoreTechInterfaceImpl.cpp
@@ -2,7 +2,7 @@
 #include "CoreTechInterfaceImpl.h"
 
 #include "CADInterfacesModule.h"
-#include "CoretechFileReader.h"
+#include "CoretechFileParser.h"
 
 #include "Misc/Paths.h"
 
@@ -458,22 +458,22 @@ namespace CADLibrary
 		return ChangeTesselationParameters(ImportParams.ChordTolerance / ImportParams.ScaleFactor, ModelUnit, ImportParams.MaxNormalAngle);
 	}
 
-	ECoreTechParsingResult FCoreTechInterfaceImpl::LoadFile(const FFileDescription& InFileDescription, const FImportParameters& InImportParameters, const FString& InCachePath, FArchiveSceneGraph& OutSceneGraphArchive, TArray<FString>& OutWarningMessages, TArray<FBodyMesh>& OutBodyMeshes)
-	{
-		FCoreTechFileReader::FContext Context(InImportParameters, InCachePath, OutSceneGraphArchive, OutWarningMessages, OutBodyMeshes);
+	//ECADParsingResult FCoreTechInterfaceImpl::LoadFile(const FFileDescription& InFileDescription, const FImportParameters& InImportParameters, const FString& InCachePath, FArchiveSceneGraph& OutSceneGraphArchive, TArray<FString>& OutWarningMessages, TArray<FBodyMesh>& OutBodyMeshes)
+	//{
+	//	FCoreTechFileReader::FContext Context(InImportParameters, InCachePath, OutSceneGraphArchive, OutWarningMessages, OutBodyMeshes);
 
-		FCoreTechFileReader FileReader(Context);
+	//	FCoreTechFileReader FileReader(Context);
 
-		ECoreTechParsingResult Result = FileReader.ProcessFile(InFileDescription);
+	//	ECoreTechParsingResult Result = FileReader.ProcessFile(InFileDescription);
 
-		return Result;
-	}
+	//	return Result;
+	//}
 
 	void FCoreTechInterfaceImpl::GetTessellation(uint64 ObjectId, FBodyMesh& OutBodyMesh, bool bIsBody)
 	{
 		if (bIsBody)
 		{
-			CoreTechFileReaderUtils::GetBodyTessellation(ObjectId, OutBodyMesh);
+			CoreTechFileParserUtils::GetBodyTessellation(ObjectId, OutBodyMesh);
 		}
 		else
 		{
@@ -487,7 +487,7 @@ namespace CADLibrary
 
 			while (CT_OBJECT_ID BodyId = ObjectList.IteratorIter())
 			{
-				CoreTechFileReaderUtils::GetBodyTessellation(BodyId, OutBodyMesh);
+				CoreTechFileParserUtils::GetBodyTessellation(BodyId, OutBodyMesh);
 			}
 		}
 	}
diff --git a/Engine/Plugins/Enterprise/DatasmithCADImporter/Source/CADInterfaces/Private/CoreTechInterfaceImpl.h b/Engine/Plugins/Enterprise/DatasmithCADImporter/Source/CADInterfaces/Private/CoreTechInterfaceImpl.h
--- a/Engine/Plugins/Enterprise/DatasmithCADImporter/Source/CADInterfaces/Private/CoreTechInterfaceImpl.h
+++ b/Engine/Plugins/Enterprise/DatasmithCADImporter/Source/CADInterfaces/Private/CoreTechInterfaceImpl.h
@@ -71,22 +71,21 @@ namespace CADLibrary
 		virtual bool Repair(uint64 MainObjectID, EStitchingTechnique StitchingTechnique, double SewingToleranceFactor) override;
 		virtual bool SetCoreTechTessellationState(const FImportParameters& ImportParams) override;
 
-		virtual ECoreTechParsingResult LoadFile(
-			const FFileDescription& InFileDescription,
-			const FImportParameters& InImportParameters,
-			const FString& InCachePath,
-			FArchiveSceneGraph& OutSceneGraphArchive,
-			TArray<FString>& OutWarningMessages,
-			TArray<FBodyMesh>& OutBodyMeshes) override;
-
-		virtual ECoreTechParsingResult LoadFile(const FFileDescription& InFileDescription, FLoadingContext& Context) override
-		{
-			Context.SceneGraphArchive = MakeShared<FArchiveSceneGraph>();
-			Context.WarningMessages = MakeShared<TArray<FString>>();
-			Context.BodyMeshes = MakeShared<TArray<FBodyMesh>>();
-
-			return LoadFile(InFileDescription, Context.ImportParameters, Context.CachePath, *Context.SceneGraphArchive, *Context.WarningMessages, *Context.BodyMeshes);
-		}
+		//virtual ECADParsingResult LoadFile(
+		//	const FFileDescription& InFileDescription,
+		//	const FImportParameters& InImportParameters,
+		//	const FString& InCachePath,
+		//	FArchiveSceneGraph& OutSceneGraphArchive,
+		//	TArray<FString>& OutWarningMessages,
+		//	TArray<FBodyMesh>& OutBodyMeshes) override;
+
+		//virtual ECADParsingResult LoadFile(const FFileDescription& InFileDescription, FLoadingContext& Context) override
+		//{
+		//	Context.SceneGraphArchive = MakeShared<FArchiveSceneGraph>();
+		//	Context.WarningMessages = MakeShared<TArray<FString>>();
+		//	Context.BodyMeshes = MakeShared<TArray<FBodyMesh>>();
+		//	return LoadFile(InFileDescription, Context.ImportParameters, Context.CachePath, *Context.SceneGraphArchive, *Context.WarningMessages, *Context.BodyMeshes);
+		//}
 
 		virtual void GetTessellation(uint64 BodyId, FBodyMesh& OutBodyMesh, bool bIsBody) override;
 
diff --git a/Engine/Plugins/Enterprise/DatasmithCADImporter/Source/CADInterfaces/Private/CoreTechTypes.cpp b/Engine/Plugins/Enterprise/DatasmithCADImporter/Source/CADInterfaces/Private/CoreTechTypes.cpp
--- a/Engine/Plugins/Enterprise/DatasmithCADImporter/Source/CADInterfaces/Private/CoreTechTypes.cpp
+++ b/Engine/Plugins/Enterprise/DatasmithCADImporter/Source/CADInterfaces/Private/CoreTechTypes.cpp
@@ -101,31 +101,32 @@ namespace CADLibrary
 		}
 	}
 
-	ECoreTechParsingResult CTKIO_LoadFile(const FFileDescription& InFileDescription, const FImportParameters& InImportParameters, const FString& InCachePath, FArchiveSceneGraph& OutSceneGraphArchive, TArray<FString>& OutWarningMessages, TArray<FBodyMesh>& OutBodyMeshes)
-	{
-		if ( CoreTechInterface.IsValid())
-		{
-			if (!CoreTechInterface->IsExternal())
-			{
-				return CoreTechInterface->LoadFile(InFileDescription, InImportParameters, InCachePath, OutSceneGraphArchive, OutWarningMessages, OutBodyMeshes);
-			}
-
-			// Allocated data is crossing dll boundaries. The data must be copied.
-			FLoadingContext Context(InImportParameters, InCachePath);
-
-			const ECoreTechParsingResult Result = CoreTechInterface->LoadFile(InFileDescription, Context);
-			if (Result == ECoreTechParsingResult::ProcessOk)
-			{
-				OutSceneGraphArchive = *Context.SceneGraphArchive;
-				OutWarningMessages = *Context.WarningMessages;
-				OutBodyMeshes = *Context.BodyMeshes;
-			}
-
-			return Result;
-		}
-
-		return  ECoreTechParsingResult::ProcessFailed;
-	}
+	//ECADParsingResult CTKIO_LoadFile(const FFileDescription& InFileDescription, const FImportParameters& InImportParameters, const FString& InCachePath, FArchiveSceneGraph& OutSceneGraphArchive, TArray<FString>& OutWarningMessages, TArray<FBodyMesh>& OutBodyMeshes)
+	//{
+	//	if ( CoreTechInterface.IsValid())
+	//	{
+	//		if (!CoreTechInterface->IsExternal())
+	//		{
+	//			//return CoreTechInterface->LoadFile(InFileDescription, InImportParameters, InCachePath, OutSceneGraphArchive, OutWarningMessages, OutBodyMeshes);
+	//			return ECADParsingResult::ProcessFailed;
+	//		}
+
+	//		// Allocated data is crossing dll boundaries. The data must be copied.
+	//		FLoadingContext Context(InImportParameters, InCachePath);
+
+	//		const ECADParsingResult Result = ECADParsingResult::ProcessFailed;  //CoreTechInterface->LoadFile(InFileDescription, Context);
+	//		if (Result == ECADParsingResult::ProcessOk)
+	//		{
+	//			OutSceneGraphArchive = *Context.SceneGraphArchive;
+	//			OutWarningMessages = *Context.WarningMessages;
+	//			OutBodyMeshes = *Context.BodyMeshes;
+	//		}
+
+	//		return Result;
+	//	}
+
+	//	return ECADParsingResult::ProcessFailed;
+	//}
 
 	bool CTKIO_CreateNurbsSurface(const FNurbsSurface& Surface, uint64& ObjectID)
 	{
diff --git a/Engine/Plugins/Enterprise/DatasmithCADImporter/Source/CADInterfaces/Public/CADFileData.h b/Engine/Plugins/Enterprise/DatasmithCADImporter/Source/CADInterfaces/Public/CADFileData.h
new file mode 100644
--- /dev/null
+++ b/Engine/Plugins/Enterprise/DatasmithCADImporter/Source/CADInterfaces/Public/CADFileData.h
@@ -0,0 +1,393 @@
+// Copyright Epic Games, Inc. All Rights Reserved.
+#pragma once
+
+#include "CoreMinimal.h"
+
+#include "CADData.h"
+#include "CADOptions.h"
+#include "CADSceneGraph.h"
+
+namespace CADLibrary
+{
+	enum class ECADImportSdk
+	{
+		None = 0,
+		KernelIO, 
+		TechSoft
+	};
+
+	class FCADFileData
+	{
+	public:
+		FCADFileData(const FImportParameters& InImportParameters, const FFileDescription& InFileDescription, const FString& InCachePath)
+			: CADImportSdk(ECADImportSdk::KernelIO)
+			, ImportParameters(InImportParameters)
+			, CachePath(InCachePath)
+			, BodyCacheExt(InImportParameters.bDisableCADKernelTessellation ? TEXT(".ct") : TEXT(".ugeom"))
+			, bIsCacheDefined(!InCachePath.IsEmpty())
+			, FileDescription(InFileDescription)
+		{
+			SceneGraphArchive.FullPath = FileDescription.OriginalPath;
+			SceneGraphArchive.CADFileName = FileDescription.Name;
+		}
+
+		uint32 GetSceneFileHash()
+		{
+			if (!SceneFileHash)
+			{
+				SceneFileHash = HashCombine(FileDescription.GetFileHash(), GetTypeHash(ImportParameters.StitchingTechnique));
+				SceneFileHash = HashCombine(SceneFileHash, GetTypeHash(CADImportSdk));
+			}
+			return SceneFileHash;
+		}
+
+		uint32 GetGeomFileHash()
+		{
+			if (!GeomFileHash)
+			{
+				GeomFileHash = GetSceneFileHash();
+				GeomFileHash = HashCombine(GeomFileHash, GetTypeHash(ImportParameters.ChordTolerance));
+				GeomFileHash = HashCombine(GeomFileHash, GetTypeHash(ImportParameters.MaxEdgeLength));
+				GeomFileHash = HashCombine(GeomFileHash, GetTypeHash(ImportParameters.MaxNormalAngle));
+				GeomFileHash = HashCombine(GeomFileHash, GetTypeHash(ImportParameters.MetricUnit));
+				GeomFileHash = HashCombine(GeomFileHash, GetTypeHash(ImportParameters.ScaleFactor));
+				GeomFileHash = HashCombine(GeomFileHash, GetTypeHash(ImportParameters.StitchingTechnique));
+			}
+			return GeomFileHash;
+		}
+
+		void SetArchiveNames()
+		{
+			SceneGraphArchive.ArchiveFileName = FString::Printf(TEXT("UEx%08x"), GetSceneFileHash());
+			MeshArchiveFile = FString::Printf(TEXT("UEx%08x"), GetGeomFileHash());
+		}
+
+		const FString GetSceneGraphFilePath() const
+		{
+			if (IsCacheDefined())
+			{
+				return FPaths::Combine(CachePath, TEXT("scene"), SceneGraphArchive.ArchiveFileName + TEXT(".sg"));
+			}
+			return FString();
+		}
+
+		const FString GetMeshArchiveFilePath() const
+		{
+			if (IsCacheDefined())
+			{
+				return FPaths::Combine(CachePath, TEXT("mesh"), MeshArchiveFile + TEXT(".gm"));
+			}
+			return FString();
+		}
+
+		const FString GetBodyCachePath(uint32 BodyHash) const
+		{
+			if (IsCacheDefined())
+			{
+				FString BodyFileName = FString::Printf(TEXT("UEx%08x"), BodyHash);
+				return FPaths::Combine(CachePath, TEXT("body"), BodyFileName + BodyCacheExt);
+			}
+			return FString();
+		}
+
+		const FString GetCADCachePath() const
+		{
+			if (IsCacheDefined())
+			{
+				FString CacheFileName = FString::Printf(TEXT("UEx%08x"), FileDescription.GetFileHash());
+				return FPaths::Combine(CachePath, TEXT("cad"), CacheFileName + TEXT(".ct"));
+			}
+			return FString();
+		}
+
+		bool IsCacheDefined() const
+		{
+			return bIsCacheDefined;
+		}
+
+		void AddWarningMessages(const FString& Message)
+		{
+			WarningMessages.Add(Message);
+		}
+
+		void LoadSceneGraphArchive()
+		{
+			FString SceneGraphFilePath = GetSceneGraphFilePath();
+			SceneGraphArchive.DeserializeMockUpFile(*SceneGraphFilePath);
+		}
+
+		void ExportSceneGraphFile()
+		{
+			FString SceneGraphFilePath = GetSceneGraphFilePath();
+			SceneGraphArchive.SerializeMockUp(*SceneGraphFilePath);
+		}
+
+		bool HasComponentOfId(FCadId ComponentId) const
+		{
+			return SceneGraphArchive.CADIdToComponentIndex.Find(ComponentId) != nullptr;
+		}
+
+
+		int32 AddComponent(FCadId ComponentId)
+		{
+			int32 Index = SceneGraphArchive.ComponentSet.Emplace(ComponentId);
+			SceneGraphArchive.CADIdToComponentIndex.Add(ComponentId, Index);
+			return Index;
+		}
+
+		FArchiveComponent& GetComponentAt(int32 Index)
+		{
+			return SceneGraphArchive.ComponentSet[Index];
+		}
+
+
+		bool HasInstanceOfId(FCadId InstanceId) const
+		{
+			return SceneGraphArchive.CADIdToInstanceIndex.Find(InstanceId) != nullptr;
+		}
+
+		int32 AddInstance(FCadId InstanceId)
+		{
+			int32 Index = SceneGraphArchive.Instances.Emplace(InstanceId);
+			SceneGraphArchive.CADIdToInstanceIndex.Add(InstanceId, Index);
+			return Index;
+		}
+
+		FArchiveInstance& GetInstanceAt(int32 Index)
+		{
+			return SceneGraphArchive.Instances[Index];
+		}
+
+
+		bool HasBodyOfId(FCadId BodyId) const
+		{
+			return SceneGraphArchive.CADIdToBodyIndex.Find(BodyId) != nullptr;
+		}
+
+		int32 AddBody(FCadId BodyId)
+		{
+			int32 Index = SceneGraphArchive.BodySet.Emplace(BodyId);
+			SceneGraphArchive.CADIdToBodyIndex.Add(BodyId, Index);
+			return Index;
+		}
+
+		FArchiveBody& GetBodyAt(int32 Index)
+		{
+			return SceneGraphArchive.BodySet[Index];
+		}
+
+
+		int32* FindUnloadedComponentOfId(FCadId ComponentId)
+		{
+			return SceneGraphArchive.CADIdToUnloadedComponentIndex.Find(ComponentId);
+		}
+
+		bool HasUnloadedComponentOfId(FCadId ComponentId)
+		{
+			return FindUnloadedComponentOfId(ComponentId) != nullptr;
+		}
+
+		int32 AddUnloadedComponent(FCadId ComponentId)
+		{
+			int32 Index = SceneGraphArchive.UnloadedComponentSet.Emplace(ComponentId);
+			SceneGraphArchive.CADIdToUnloadedComponentIndex.Add(ComponentId, Index);
+			return Index;
+		}
+
+		FArchiveUnloadedComponent& GetUnloadedComponentAt(int32 Index)
+		{
+			return SceneGraphArchive.UnloadedComponentSet[Index];
+		}
+
+		FFileDescription& GetExternalReferences(int32 Index)
+		{
+			return SceneGraphArchive.ExternalRefSet[Index];
+		}
+
+		FFileDescription& AddExternalRef(const TCHAR* InFilePath, const TCHAR* InConfiguration, const TCHAR* InRootFilePath)
+		{
+			return SceneGraphArchive.ExternalRefSet.Emplace_GetRef(InFilePath, InConfiguration, InRootFilePath);
+		}
+
+		FFileDescription& AddExternalRef(const FFileDescription& InFileDescription)
+		{
+			return SceneGraphArchive.ExternalRefSet.Emplace_GetRef(InFileDescription);
+		}
+
+		FBodyMesh& AddBodyMesh(FCadId BodyId)
+		{
+			return BodyMeshes.Emplace_GetRef(BodyId);
+		}
+
+
+		void ExportMeshArchiveFile()
+		{
+			FString MeshArchiveFilePath = GetMeshArchiveFilePath();
+			SerializeBodyMeshSet(*MeshArchiveFilePath, BodyMeshes);
+		}
+
+		const TArray<FFileDescription>& GetExternalRefSet() const 
+		{
+			return SceneGraphArchive.ExternalRefSet;
+		}
+
+		const FString& GetSceneGraphFileName() const
+		{
+			return SceneGraphArchive.ArchiveFileName;
+		}
+
+		const FString& GetMeshFileName() const
+		{
+			return MeshArchiveFile;
+		}
+
+		const TArray<FString>& GetWarningMessages() const
+		{
+			return WarningMessages;
+		}
+
+		const FArchiveSceneGraph& GetSceneGraphArchive() const
+		{
+			return SceneGraphArchive;
+		}
+
+		FArchiveSceneGraph& GetSceneGraphArchive()
+		{
+			return SceneGraphArchive;
+		}
+
+		FArchiveMaterial* FindMaterial(uint32 MaterialId)
+		{
+			return SceneGraphArchive.MaterialHIdToMaterial.Find(MaterialId);
+		}
+
+		FArchiveMaterial& AddMaterial(uint32 MaterialId)
+		{
+			return SceneGraphArchive.MaterialHIdToMaterial.Emplace(MaterialId, MaterialId);
+		}
+
+		FArchiveColor* FindColor(uint32 ColorId)
+		{
+			return SceneGraphArchive.ColorHIdToColor.Find(ColorId);
+		}
+
+		FArchiveColor& AddColor(uint32 ColorId)
+		{
+			return SceneGraphArchive.ColorHIdToColor.Emplace(ColorId, ColorId);
+		}
+
+		const TArray<FBodyMesh>& GetBodyMeshes() const
+		{
+			return BodyMeshes;
+		}
+
+		TArray<FBodyMesh>& GetBodyMeshes()
+		{
+			return BodyMeshes;
+		}
+
+		const FFileDescription& GetCADFileDescription() const
+		{
+			return FileDescription;
+		}
+
+		FFileDescription& GetCADFileDescription()
+		{
+			return FileDescription;
+		}
+
+		bool HasConfiguration() const
+		{
+			return !FileDescription.Configuration.IsEmpty();
+		}
+
+		const FString& GetConfiguration() const
+		{
+			return FileDescription.Configuration;
+		}
+
+		const FString& FileExtension() const
+		{
+			return FileDescription.Extension;
+		}
+
+		const FString& FilePath() const
+		{
+			return FileDescription.Path;
+		}
+
+		const FString& FileName() const
+		{
+			return FileDescription.Name;
+		}
+
+		void ReplaceFileByCacheBackup(const FString& CacheFilePath)
+		{
+			FileDescription.ReplaceByKernelIOBackup(CacheFilePath);
+		}
+
+		void ReserveBodyMeshes(int32 MaxBodyCount)
+		{
+			BodyMeshes.Reserve(MaxBodyCount);
+		}
+
+		const FImportParameters& GetImportParameters() const
+		{
+			return ImportParameters;
+		}
+
+		double MetricUnit() const
+		{
+			return ImportParameters.MetricUnit;
+		}
+
+		const EStitchingTechnique& GetStitchingTechnique() const
+		{
+			return ImportParameters.StitchingTechnique;
+		}
+
+		bool IsCADKernelTessellation() const
+		{
+			return ImportParameters.bDisableCADKernelTessellation;
+		}
+
+		bool IsSequentialImport() const
+		{
+			return ImportParameters.bEnableSequentialImport;
+		}
+
+		bool NeedUVMapScaling() const
+		{
+			return ImportParameters.bScaleUVMap;
+		}
+
+		double GetScaleFactor() const
+		{
+			return ImportParameters.ScaleFactor;
+		}
+
+		double GetMetricUnit() const
+		{
+			return ImportParameters.MetricUnit;
+		}
+
+	private:
+		ECADImportSdk CADImportSdk = ECADImportSdk::KernelIO;
+		const FImportParameters ImportParameters;
+		const FString CachePath;
+		const FString BodyCacheExt;
+		const bool bIsCacheDefined;
+
+		FFileDescription FileDescription;
+
+		FString MeshArchiveFile;
+
+		FArchiveSceneGraph SceneGraphArchive;
+		TArray<FBodyMesh> BodyMeshes;
+
+		TArray<FString> WarningMessages;
+
+		uint32 SceneFileHash = 0;
+		uint32 GeomFileHash = 0;
+	};
+}
\ No newline at end of file
diff --git a/Engine/Plugins/Enterprise/DatasmithCADImporter/Source/CADInterfaces/Public/CADFileParser.h b/Engine/Plugins/Enterprise/DatasmithCADImporter/Source/CADInterfaces/Public/CADFileParser.h
new file mode 100644
--- /dev/null
+++ b/Engine/Plugins/Enterprise/DatasmithCADImporter/Source/CADInterfaces/Public/CADFileParser.h
@@ -0,0 +1,16 @@
+// Copyright Epic Games, Inc. All Rights Reserved.
+#pragma once
+
+#include "CoreMinimal.h"
+
+namespace CADLibrary
+{
+
+class ICADFileParser
+{
+public:
+	virtual ~ICADFileParser() = default;
+	virtual ECADParsingResult Process() = 0;
+};
+
+}
\ No newline at end of file
diff --git a/Engine/Plugins/Enterprise/DatasmithCADImporter/Source/CADInterfaces/Public/CADFileReader.h b/Engine/Plugins/Enterprise/DatasmithCADImporter/Source/CADInterfaces/Public/CADFileReader.h
new file mode 100644
--- /dev/null
+++ b/Engine/Plugins/Enterprise/DatasmithCADImporter/Source/CADInterfaces/Public/CADFileReader.h
@@ -0,0 +1,50 @@
+// Copyright Epic Games, Inc. All Rights Reserved.
+#pragma once
+
+#include "CoreMinimal.h"
+
+#include "CADData.h"
+#include "CADFileData.h"
+#include "CADFileParser.h"
+#include "CADOptions.h"
+#include "CADSceneGraph.h"
+
+namespace CADLibrary
+{
+
+	class CADINTERFACES_API FCADFileReader
+	{
+	public:
+		/**
+		 * @param ImportParams Parameters that setting import data like mesh SAG...
+		 * @param EnginePluginsPath Full Path of EnginePlugins. Mandatory to set KernelIO to import DWG, or DGN files
+		 * @param InCachePath Full path of the cache in which the data will be saved
+		 */
+		FCADFileReader(const FImportParameters& ImportParams, FFileDescription& InCTFileDescription, const FString& EnginePluginsPath = TEXT(""), const FString& InCachePath = TEXT(""));
+		ECADParsingResult ProcessFile();
+
+		const FCADFileData& GetCADFileData() const
+		{
+			return CADFileData;
+		}
+
+		FCADFileData& GetCADFileData()
+		{
+			return CADFileData;
+		}
+
+	private:
+		bool FindFile(FFileDescription& File);
+
+		uint32 GetSceneFileHash(const uint32 InSGHash, const FImportParameters& ImportParam);
+		uint32 GetGeomFileHash(const uint32 InSGHash, const FImportParameters& ImportParam);
+
+		void LoadSceneGraphArchive(const FString& SceneGraphFilePath);
+
+		void ExportSceneGraphFile();
+		void ExportMeshArchiveFile();
+
+		FCADFileData CADFileData;
+		TUniquePtr<ICADFileParser> CADParser;
+	};
+} // ns CADLibrary
diff --git a/Engine/Plugins/Enterprise/DatasmithCADImporter/Source/CADInterfaces/Public/CoreTechFileParser.h b/Engine/Plugins/Enterprise/DatasmithCADImporter/Source/CADInterfaces/Public/CoreTechFileParser.h
deleted file mode 100644
diff --git a/Engine/Plugins/Enterprise/DatasmithCADImporter/Source/CADInterfaces/Public/CoreTechTypes.h b/Engine/Plugins/Enterprise/DatasmithCADImporter/Source/CADInterfaces/Public/CoreTechTypes.h
--- a/Engine/Plugins/Enterprise/DatasmithCADImporter/Source/CADInterfaces/Public/CoreTechTypes.h
+++ b/Engine/Plugins/Enterprise/DatasmithCADImporter/Source/CADInterfaces/Public/CoreTechTypes.h
@@ -3,26 +3,11 @@
 
 
 #include "CADOptions.h"
-#include "Containers/Array.h"
-#include "Math/Vector.h"
-#include "Math/Vector2D.h"
 
 namespace CADLibrary
 {
-	enum class ECoreTechParsingResult : uint8
-	{
-		Unknown,
-		Running,
-		UnTreated,
-		ProcessOk,
-		ProcessFailed,
-		FileNotFound,
-	};
-
-	class FBodyMesh;
-	struct FFileDescription;
-	class FArchiveSceneGraph;
 	class ICoreTechInterface;
+	class FBodyMesh;
 
 	// Helper struct used to pass Nurbs surface definition to CoreTech
 	struct FNurbsSurface
@@ -58,29 +43,6 @@ namespace CADLibrary
 		TArray<double> ControlPoints;
 	};
 
-	// Helper struct used to pass the result of ICoreTechInterface::LoadFile across dll boundaries
-	// when the ICoreTechInterface object has been created through the DatasmithCADRuntime dll 
-	struct FLoadingContext
-	{
-		FLoadingContext(const FImportParameters& InImportParameters, const FString& InCachePath)
-			: ImportParameters(InImportParameters)
-			, CachePath(InCachePath)
-		{
-		}
-
-		FLoadingContext(const FLoadingContext& Other)
-			: ImportParameters(Other.ImportParameters)
-			, CachePath(Other.CachePath)
-		{
-		}
-
-		const FImportParameters& ImportParameters;
-		const FString& CachePath;
-		TSharedPtr<FArchiveSceneGraph> SceneGraphArchive;
-		TSharedPtr<TArray<FString>> WarningMessages;
-		TSharedPtr<TArray<FBodyMesh>> BodyMeshes;
-	};
-
 	class ICoreTechInterface
 	{
 	public:
@@ -138,16 +100,6 @@ namespace CADLibrary
 
 		virtual void GetTessellation(uint64 BodyId, TSharedPtr<FBodyMesh>& OutBodyMesh, bool bIsBody) = 0;
 
-		virtual ECoreTechParsingResult LoadFile(
-			const FFileDescription& InFileDescription,
-			const FImportParameters& InImportParameters,
-			const FString& InCachePath,
-			FArchiveSceneGraph& OutSceneGraphArchive,
-			TArray<FString>& OutWarningMessages,
-			TArray<FBodyMesh>& OutBodyMeshes) = 0;
-
-		virtual ECoreTechParsingResult LoadFile(const FFileDescription& InFileDescription, FLoadingContext& LoadingContext) = 0;
-
 		virtual bool CreateNurbsSurface(const FNurbsSurface& Surface, uint64& ObjectID) = 0;
 
 		virtual bool CreateNurbsCurve(const FNurbsCurve& Curve, uint64& ObjectID) = 0;
@@ -217,8 +169,6 @@ namespace CADLibrary
 
 	CADINTERFACES_API void CTKIO_GetTessellation(uint64 ObjectId, FBodyMesh& OutBodyMesh, bool bIsBody = true);
 
-	CADINTERFACES_API ECoreTechParsingResult CTKIO_LoadFile(const FFileDescription& InFileDescription, const FImportParameters& InImportParameters, const FString& InCachePath, FArchiveSceneGraph& OutSceneGraphArchive, TArray<FString>& OutWarningMessages, TArray<FBodyMesh>& OutBodyMeshes);
-
 	CADINTERFACES_API bool CTKIO_CreateNurbsSurface(const FNurbsSurface& NurbsDefinition, uint64& ObjectID);
 
 	CADINTERFACES_API bool CTKIO_CreateNurbsCurve(const FNurbsCurve& Curve, uint64& ObjectID);
diff --git a/Engine/Plugins/Enterprise/DatasmithCADImporter/Source/CADTools/Private/CADData.cpp b/Engine/Plugins/Enterprise/DatasmithCADImporter/Source/CADTools/Private/CADData.cpp
--- a/Engine/Plugins/Enterprise/DatasmithCADImporter/Source/CADTools/Private/CADData.cpp
+++ b/Engine/Plugins/Enterprise/DatasmithCADImporter/Source/CADTools/Private/CADData.cpp
@@ -157,7 +157,7 @@ void GetCleanFilenameAndExtension(const FString& InFilePath, FString& OutFilenam
 	}
 }
 
-uint32 FFileDescription::GetFileHash()
+uint32 FFileDescription::GetFileHash() const
 {
 	FFileStatData FileStatData = IFileManager::Get().GetStatData(*OriginalPath);
 
@@ -171,4 +171,68 @@ uint32 FFileDescription::GetFileHash()
 	return FileHash;
 }
 
+ECADFormat FileFormat(const FString& Extension)
+{
+	if (Extension == TEXT("sldprt") || Extension == TEXT("sldasm"))
+	{
+		return ECADFormat::SOLIDWORKS;
+	}
+	else if (Extension == TEXT("catpart") || Extension == TEXT("catproduct"))
+	{
+		return ECADFormat::CATIA;
+	}
+	else if (Extension == TEXT("cgr"))
+	{
+		return ECADFormat::CATIA_CGR;
+	}
+	else if (Extension == TEXT("x_t") || Extension == TEXT("x_b"))
+	{
+		return ECADFormat::PARASOLID;
+	}
+	else if (Extension == TEXT("iges") || Extension == TEXT("igs"))
+	{
+		return ECADFormat::IGES;
+	}
+	else if (Extension == TEXT("step") || Extension == TEXT("stp"))
+	{
+		return ECADFormat::STEP;
+	}
+	else if (Extension == TEXT("jt"))
+	{
+		return ECADFormat::JT;
+	}
+	else if (Extension == TEXT("prt.*") || Extension == TEXT("asm.*"))
+	{
+		return ECADFormat::CREO;
+	}
+	else if (Extension == TEXT("prt"))
+	{
+		return ECADFormat::NX;
+	}
+	else if (Extension == TEXT("step") || Extension == TEXT("stp"))
+	{
+		return ECADFormat::STEP;
+	}
+	else if (Extension == TEXT("sat"))
+	{
+		return ECADFormat::ACIS;
+	}
+	else if (Extension == TEXT("model"))
+	{
+		return ECADFormat::CATIAV4;
+	}
+	else if (Extension == TEXT("3dxml"))
+	{
+		return ECADFormat::CATIA_3DXML;
+	}
+	else if (Extension == TEXT("ipt") || Extension == TEXT("iam"))
+	{
+		return ECADFormat::INVENTOR;
+	}
+	else
+	{
+		return ECADFormat::OTHER;
+	}
+}
+
 }
diff --git a/Engine/Plugins/Enterprise/DatasmithCADImporter/Source/CADTools/Public/CADData.h b/Engine/Plugins/Enterprise/DatasmithCADImporter/Source/CADTools/Public/CADData.h
--- a/Engine/Plugins/Enterprise/DatasmithCADImporter/Source/CADTools/Public/CADData.h
+++ b/Engine/Plugins/Enterprise/DatasmithCADImporter/Source/CADTools/Public/CADData.h
@@ -18,6 +18,38 @@ using FCADUUID = uint32;  // Universal unique identifier that be used for the un
 namespace CADLibrary
 {
 
+enum class ECADFormat
+{
+	ACIS,
+	CATIA,
+	CATIA_CGR,
+	CATIA_3DXML,
+	CATIAV4,
+	CREO,
+	DWG,
+	DGN,
+	IGES,
+	INVENTOR,
+	JT,
+	NX,
+	PARASOLID,
+	SOLIDWORKS,
+	STEP,
+	OTHER
+};
+
+CADTOOLS_API ECADFormat FileFormat(const FString& Extension);
+
+enum class ECADParsingResult : uint8
+{
+	Unknown,
+	Running,
+	UnTreated,
+	ProcessOk,
+	ProcessFailed,
+	FileNotFound,
+};
+
 // TODO: Remove from hear and replace by DatasmithUtils::GetCleanFilenameAndExtension... But need to remove DatasmithCore dependancies 
 CADTOOLS_API void GetCleanFilenameAndExtension(const FString& InFilePath, FString& OutFilename, FString& OutExtension);
 
@@ -76,7 +108,12 @@ struct CADTOOLS_API FFileDescription
 
 	friend CADTOOLS_API FArchive& operator<<(FArchive& Ar, FFileDescription& File);
 
-	uint32 GetFileHash();
+	uint32 GetFileHash() const;
+
+	ECADFormat GetFormat()
+	{
+		return FileFormat(Extension);
+	}
 
 	FString Path;
 	FString OriginalPath;
diff --git a/Engine/Plugins/Enterprise/DatasmithCADImporter/Source/DatasmithCADTranslator/Private/DatasmithCADTranslator.cpp b/Engine/Plugins/Enterprise/DatasmithCADImporter/Source/DatasmithCADTranslator/Private/DatasmithCADTranslator.cpp
--- a/Engine/Plugins/Enterprise/DatasmithCADImporter/Source/DatasmithCADTranslator/Private/DatasmithCADTranslator.cpp
+++ b/Engine/Plugins/Enterprise/DatasmithCADImporter/Source/DatasmithCADTranslator/Private/DatasmithCADTranslator.cpp
@@ -3,6 +3,7 @@
 #include "DatasmithCADTranslator.h"
 
 #include "CADInterfacesModule.h"
+#include "CADFileReader.h"
 #include "CADKernelSurfaceExtension.h"
 #include "CoreTechSurfaceHelper.h"
 #include "DatasmithCADTranslatorModule.h"
@@ -92,11 +93,11 @@ bool FDatasmithCADTranslator::LoadScene(TSharedRef<IDatasmithScene> DatasmithSce
 	ImportParameters.ChordTolerance = TesselationOptions.ChordTolerance;
 	ImportParameters.MaxEdgeLength = TesselationOptions.MaxEdgeLength;
 	ImportParameters.MaxNormalAngle = TesselationOptions.NormalTolerance;
-	ImportParameters.StitchingTechnique = (CADLibrary::EStitchingTechnique) TesselationOptions.StitchingTechnique;
+	ImportParameters.StitchingTechnique = (CADLibrary::EStitchingTechnique)TesselationOptions.StitchingTechnique;
 
 	CADLibrary::FFileDescription FileDescription(*FPaths::ConvertRelativePathToFull(GetSource().GetSourceFile()),
 		TEXT(""),
-		*FPaths::GetPath(FPaths::ConvertRelativePathToFull(GetSource().GetSourceFile())) );
+		*FPaths::GetPath(FPaths::ConvertRelativePathToFull(GetSource().GetSourceFile())));
 
 	// Do not change the model unit when translator is called by the Datasmith runtime plugin.
 #if WITH_EDITOR
@@ -174,16 +175,17 @@ bool FDatasmithCADTranslator::LoadScene(TSharedRef<IDatasmithScene> DatasmithSce
 		return true;
 	}
 
-	CADLibrary::FCoreTechFileParser FileParser(ImportParameters, *FPaths::EnginePluginsDir());
-	if (FileParser.ProcessFile(FileDescription) != CADLibrary::ECoreTechParsingResult::ProcessOk)
+	CADLibrary::FCADFileReader FileReader(ImportParameters, FileDescription, *FPaths::EnginePluginsDir());
+	if (FileReader.ProcessFile() != CADLibrary::ECADParsingResult::ProcessOk)
 	{
 		return false;
 	}
 
-	FDatasmithSceneBaseGraphBuilder SceneGraphBuilder(&FileParser.GetSceneGraphArchive(), DatasmithScene, GetSource(), ImportParameters);
+	CADLibrary::FCADFileData& CADFileData = FileReader.GetCADFileData();
+	FDatasmithSceneBaseGraphBuilder SceneGraphBuilder(&CADFileData.GetSceneGraphArchive(), DatasmithScene, GetSource(), ImportParameters);
 	SceneGraphBuilder.Build();
 
-	MeshBuilderPtr = MakeUnique<FDatasmithMeshBuilder>(FileParser.GetBodyMeshes(), ImportParameters);
+	MeshBuilderPtr = MakeUnique<FDatasmithMeshBuilder>(CADFileData.GetBodyMeshes(), ImportParameters);
 
 	return true;
 }
@@ -207,7 +209,7 @@ bool FDatasmithCADTranslator::LoadStaticMesh(const TSharedRef<IDatasmithMeshElem
 	if (TOptional< FMeshDescription > Mesh = MeshBuilderPtr->GetMeshDescription(MeshElement, MeshParameters))
 	{
 		OutMeshPayload.LodMeshes.Add(MoveTemp(Mesh.GetValue()));
-		if(ImportParameters.bDisableCADKernelTessellation)
+		if (ImportParameters.bDisableCADKernelTessellation)
 		{
 			CoreTechSurface::AddSurfaceDataForMesh(MeshElement->GetFile(), ImportParameters, MeshParameters, GetCommonTessellationOptions(), OutMeshPayload);
 		}
diff --git a/Engine/Plugins/Enterprise/DatasmithCADImporter/Source/DatasmithCADTranslator/Private/DatasmithSceneGraphBuilder.cpp b/Engine/Plugins/Enterprise/DatasmithCADImporter/Source/DatasmithCADTranslator/Private/DatasmithSceneGraphBuilder.cpp
--- a/Engine/Plugins/Enterprise/DatasmithCADImporter/Source/DatasmithCADTranslator/Private/DatasmithSceneGraphBuilder.cpp
+++ b/Engine/Plugins/Enterprise/DatasmithCADImporter/Source/DatasmithCADTranslator/Private/DatasmithSceneGraphBuilder.cpp
@@ -3,7 +3,6 @@
 
 #include "CADData.h"
 #include "CADSceneGraph.h"
-#include "CoreTechFileParser.h"
 #include "DatasmithSceneFactory.h"
 #include "DatasmithSceneSource.h"
 #include "DatasmithUtils.h"
@@ -11,7 +10,7 @@
 #include "MeshDescriptionHelper.h"
 #include "Misc/FileHelper.h"
 
-namespace
+namespace DatasmithSceneGraphBuilderImpl
 {
 	void GetMainMaterial(const TMap<FString, FString>& InNodeMetaDataMap, ActorData& OutNodeData, bool bMaterialPropagationIsTopDown)
 	{
@@ -32,7 +31,7 @@ namespace
 		}
 	}
 
-	void AddTransformToActor(CADLibrary::FArchiveInstance& Instance, TSharedPtr< IDatasmithActorElement > Actor, const CADLibrary::FImportParameters& ImportParameters)
+	void AddTransformToActor(const CADLibrary::FArchiveInstance& Instance, TSharedPtr< IDatasmithActorElement > Actor, const CADLibrary::FImportParameters& ImportParameters)
 	{
 		if (!Actor.IsValid())
 		{
@@ -137,7 +136,7 @@ void FDatasmithSceneGraphBuilder::FillAnchorActor(const TSharedRef<IDatasmithAct
 	}
 
 	FCadId RootId = 1;
-	int32* Index = SceneGraph->CADIdToComponentIndex.Find(RootId);
+	const int32* Index = SceneGraph->CADIdToComponentIndex.Find(RootId);
 	if (!Index)
 	{
 		return;
@@ -157,7 +156,7 @@ void FDatasmithSceneGraphBuilder::FillAnchorActor(const TSharedRef<IDatasmithAct
 	AddMetaData(ActorElement, InstanceNodeMetaDataMap, Component.MetaData);
 
 	ActorData ComponentData(*ActorUUID, ParentData);
-	GetMainMaterial(Component.MetaData, ComponentData, bMaterialPropagationIsTopDown);
+	DatasmithSceneGraphBuilderImpl::GetMainMaterial(Component.MetaData, ComponentData, bMaterialPropagationIsTopDown);
 
 	AddChildren(ActorElement, Component, ComponentData);
 
@@ -202,7 +201,7 @@ FDatasmithSceneBaseGraphBuilder::FDatasmithSceneBaseGraphBuilder(CADLibrary::FAr
 bool FDatasmithSceneBaseGraphBuilder::Build()
 {
 	FCadId RootId = 1;
-	int32* Index = SceneGraph->CADIdToComponentIndex.Find(RootId);
+	const int32* Index = SceneGraph->CADIdToComponentIndex.Find(RootId);
 	if (!Index)
 	{
 		return false;
@@ -217,7 +216,7 @@ bool FDatasmithSceneBaseGraphBuilder::Build()
 	// application_name is something like "Catia V5"
 	DatasmithScene->SetVendor(TEXT("CoreTechnologie"));
 
-	if (FString* ProductVersion = Component.MetaData.Find(TEXT("KernelIOVersion")))
+	if (const FString* ProductVersion = Component.MetaData.Find(TEXT("KernelIOVersion")))
 	{
 		DatasmithScene->SetProductVersion(**ProductVersion);
 	}
@@ -320,7 +319,7 @@ TSharedPtr< IDatasmithActorElement >  FDatasmithSceneBaseGraphBuilder::BuildInst
 					AncestorSceneGraphHash.Add(InstanceSceneGraphHash);
 
 					FCadId RootId = 1;
-					int32* Index = SceneGraph->CADIdToComponentIndex.Find(RootId);
+					const int32* Index = SceneGraph->CADIdToComponentIndex.Find(RootId);
 					if (Index)
 					{
 						Reference = &(SceneGraph->ComponentSet[*Index]);
@@ -332,15 +331,16 @@ TSharedPtr< IDatasmithActorElement >  FDatasmithSceneBaseGraphBuilder::BuildInst
 		if(!Reference)
 		{
 			SceneGraph = InstanceSceneGraph;
-			int32* Index = SceneGraph->CADIdToUnloadedComponentIndex.Find(Instance.ReferenceNodeId);
+			const int32* Index = SceneGraph->CADIdToUnloadedComponentIndex.Find(Instance.ReferenceNodeId);
 			if (Index)
 			{
 				Reference = &(SceneGraph->UnloadedComponentSet[*Index]);
 			}
 		}
 	}
-	else {
-		int32* Index = SceneGraph->CADIdToComponentIndex.Find(Instance.ReferenceNodeId);
+	else 
+	{
+		const int32* Index = SceneGraph->CADIdToComponentIndex.Find(Instance.ReferenceNodeId);
 		if (Index)
 		{
 			Reference = &(SceneGraph->ComponentSet[*Index]);
@@ -356,7 +356,7 @@ TSharedPtr< IDatasmithActorElement >  FDatasmithSceneBaseGraphBuilder::BuildInst
 	FString ActorLabel;
 	GetNodeUUIDAndName(Instance.MetaData, Reference->MetaData, InstanceIndex, ParentData.Uuid, ActorUUID, ActorLabel);
 
-	TSharedPtr< IDatasmithActorElement > Actor = CreateActor(*ActorUUID, *ActorLabel);
+	TSharedPtr<IDatasmithActorElement> Actor = CreateActor(*ActorUUID, *ActorLabel);
 	if (!Actor.IsValid())
 	{
 		return TSharedPtr< IDatasmithActorElement >();
@@ -365,12 +365,12 @@ TSharedPtr< IDatasmithActorElement >  FDatasmithSceneBaseGraphBuilder::BuildInst
 
 	ActorData InstanceData(*ActorUUID, ParentData);
 
-	GetMainMaterial(Instance.MetaData, InstanceData, bMaterialPropagationIsTopDown);
-	GetMainMaterial(Reference->MetaData, InstanceData, bMaterialPropagationIsTopDown);
+	DatasmithSceneGraphBuilderImpl::GetMainMaterial(Instance.MetaData, InstanceData, bMaterialPropagationIsTopDown);
+	DatasmithSceneGraphBuilderImpl::GetMainMaterial(Reference->MetaData, InstanceData, bMaterialPropagationIsTopDown);
 
 	AddChildren(Actor, *Reference, InstanceData);
 
-	AddTransformToActor(Instance, Actor, ImportParameters);
+	DatasmithSceneGraphBuilderImpl::AddTransformToActor(Instance, Actor, ImportParameters);
 
 	if (SceneGraph != InstanceSceneGraph)
 	{
@@ -388,25 +388,25 @@ TSharedPtr< IDatasmithActorElement >  FDatasmithSceneBaseGraphBuilder::CreateAct
 		Actor->SetLabel(InLabel);
 		return Actor;
 	}
-	return TSharedPtr< IDatasmithActorElement >();
+	return TSharedPtr<IDatasmithActorElement>();
 }
 
 void FDatasmithSceneBaseGraphBuilder::GetNodeUUIDAndName(
-	TMap<FString, FString>& InInstanceNodeMetaDataMap,
-	TMap<FString, FString>& InReferenceNodeMetaDataMap,
+	const TMap<FString, FString>& InInstanceNodeMetaDataMap,
+	const TMap<FString, FString>& InReferenceNodeMetaDataMap,
 	int32 InComponentIndex,
 	const TCHAR* InParentUEUUID,
 	FString& OutUEUUID,
 	FString& OutName
 )
 {
-	FString* InstanceKernelIOName = InInstanceNodeMetaDataMap.Find(TEXT("CTName"));
-	FString* InstanceCADName = InInstanceNodeMetaDataMap.Find(TEXT("Name"));
-	FString* InstanceUUID = InInstanceNodeMetaDataMap.Find(TEXT("UUID"));
+	const FString* InstanceKernelIOName = InInstanceNodeMetaDataMap.Find(TEXT("CTName"));
+	const FString* InstanceCADName = InInstanceNodeMetaDataMap.Find(TEXT("Name"));
+	const FString* InstanceUUID = InInstanceNodeMetaDataMap.Find(TEXT("UUID"));
 
-	FString* ReferenceKernelIOName = InReferenceNodeMetaDataMap.Find(TEXT("CTName"));
-	FString* ReferenceCADName = InReferenceNodeMetaDataMap.Find(TEXT("Name"));
-	FString* ReferenceUUID = InReferenceNodeMetaDataMap.Find(TEXT("UUID"));
+	const FString* ReferenceKernelIOName = InReferenceNodeMetaDataMap.Find(TEXT("CTName"));
+	const FString* ReferenceCADName = InReferenceNodeMetaDataMap.Find(TEXT("Name"));
+	const FString* ReferenceUUID = InReferenceNodeMetaDataMap.Find(TEXT("UUID"));
 
 	// Outname Name
 	// IName and RName are KernelIO build Name. Original names (CAD system name) are preferred
@@ -430,7 +430,7 @@ void FDatasmithSceneBaseGraphBuilder::GetNodeUUIDAndName(
 	{
 		OutName = "NoName";
 	}
-	CleanName(OutName);
+	DatasmithSceneGraphBuilderImpl::CleanName(OutName);
 
 	FCADUUID UEUUID = HashCombine(GetTypeHash(InParentUEUUID), GetTypeHash(InComponentIndex));
 
@@ -463,7 +463,7 @@ void FDatasmithSceneBaseGraphBuilder::GetNodeUUIDAndName(
 	OutUEUUID = FString::Printf(TEXT("0x%08x"), UEUUID);
 }
 
-TSharedPtr< IDatasmithActorElement > FDatasmithSceneBaseGraphBuilder::BuildComponent(CADLibrary::FArchiveComponent& Component, const ActorData& ParentData)
+TSharedPtr<IDatasmithActorElement> FDatasmithSceneBaseGraphBuilder::BuildComponent(CADLibrary::FArchiveComponent& Component, const ActorData& ParentData)
 {
 	TMap<FString, FString> InstanceNodeMetaDataMap;
 
@@ -481,14 +481,14 @@ TSharedPtr< IDatasmithActorElement > FDatasmithSceneBaseGraphBuilder::BuildCompo
 	AddMetaData(Actor, InstanceNodeMetaDataMap, Component.MetaData);
 
 	ActorData ComponentData(*ActorUUID, ParentData);
-	GetMainMaterial(Component.MetaData, ComponentData, bMaterialPropagationIsTopDown);
+	DatasmithSceneGraphBuilderImpl::GetMainMaterial(Component.MetaData, ComponentData, bMaterialPropagationIsTopDown);
 
 	AddChildren(Actor, Component, ComponentData);
 
 	return Actor;
 }
 
-TSharedPtr< IDatasmithActorElement > FDatasmithSceneBaseGraphBuilder::BuildBody(int32 BodyIndex, const ActorData& ParentData)
+TSharedPtr<IDatasmithActorElement> FDatasmithSceneBaseGraphBuilder::BuildBody(int32 BodyIndex, const ActorData& ParentData)
 {
 	TMap<FString, FString> InstanceNodeMetaDataMap;
 
@@ -507,14 +507,14 @@ TSharedPtr< IDatasmithActorElement > FDatasmithSceneBaseGraphBuilder::BuildBody(
 		Body.ColorFaceSet.Add(MaterialUuid);
 	}
 
-	TSharedPtr< IDatasmithMeshElement > MeshElement = FindOrAddMeshElement(Body, BodyLabel);
+	TSharedPtr<IDatasmithMeshElement> MeshElement = FindOrAddMeshElement(Body, BodyLabel);
 	if (!MeshElement.IsValid())
 	{
 		return TSharedPtr< IDatasmithActorElement >();
 	}
 
 
-	TSharedPtr< IDatasmithMeshActorElement > ActorElement = FDatasmithSceneFactory::CreateMeshActor(*BodyUUID);
+	TSharedPtr<IDatasmithMeshActorElement> ActorElement = FDatasmithSceneFactory::CreateMeshActor(*BodyUUID);
 	if (!ActorElement.IsValid())
 	{
 		return TSharedPtr< IDatasmithActorElement >();
@@ -756,7 +756,7 @@ bool FDatasmithSceneBaseGraphBuilder::DoesActorHaveChildrenOrIsAStaticMesh(const
 }
 
 
-void FDatasmithSceneBaseGraphBuilder::AddChildren(TSharedPtr< IDatasmithActorElement > Actor, CADLibrary::FArchiveComponent& Component, const ActorData& ParentData)
+void FDatasmithSceneBaseGraphBuilder::AddChildren(TSharedPtr< IDatasmithActorElement > Actor, const CADLibrary::FArchiveComponent& Component, const ActorData& ParentData)
 {
 	for (const int32 ChildId : Component.Children)
 	{
diff --git a/Engine/Plugins/Enterprise/DatasmithCADImporter/Source/DatasmithCADTranslator/Public/DatasmithSceneGraphBuilder.h b/Engine/Plugins/Enterprise/DatasmithCADImporter/Source/DatasmithCADTranslator/Public/DatasmithSceneGraphBuilder.h
--- a/Engine/Plugins/Enterprise/DatasmithCADImporter/Source/DatasmithCADTranslator/Public/DatasmithSceneGraphBuilder.h
+++ b/Engine/Plugins/Enterprise/DatasmithCADImporter/Source/DatasmithCADTranslator/Public/DatasmithSceneGraphBuilder.h
@@ -69,8 +69,8 @@ protected:
 	TSharedPtr< IDatasmithActorElement > BuildComponent(CADLibrary::FArchiveComponent& Component, const ActorData& ParentData);
 	TSharedPtr< IDatasmithActorElement > BuildBody(int32 BodyIndex, const ActorData& ParentData);
 
-	void AddMetaData(TSharedPtr< IDatasmithActorElement > ActorElement, TMap<FString, FString>& InstanceNodeAttributeSetMap, TMap<FString, FString>& ReferenceNodeAttributeSetMap);
-	void AddChildren(TSharedPtr< IDatasmithActorElement > Actor, CADLibrary::FArchiveComponent& Component, const ActorData& ParentData);
+	void AddMetaData(TSharedPtr<IDatasmithActorElement> ActorElement, TMap<FString, FString>& InstanceNodeAttributeSetMap, TMap<FString, FString>& ReferenceNodeAttributeSetMap);
+	void AddChildren(TSharedPtr<IDatasmithActorElement> Actor, const CADLibrary::FArchiveComponent& Component, const ActorData& ParentData);
 	bool DoesActorHaveChildrenOrIsAStaticMesh(const TSharedPtr< IDatasmithActorElement >& ActorElement);
 
 	TSharedPtr< IDatasmithUEPbrMaterialElement > GetDefaultMaterial();
@@ -79,7 +79,7 @@ protected:
 	TSharedPtr< IDatasmithActorElement > CreateActor(const TCHAR* ActorUUID, const TCHAR* ActorLabel);
 	virtual TSharedPtr< IDatasmithMeshElement > FindOrAddMeshElement(CADLibrary::FArchiveBody& Body, FString& InLabel);
 
-	void GetNodeUUIDAndName(TMap<FString, FString>& InInstanceNodeMetaDataMap, TMap<FString, FString>& InReferenceNodeMetaDataMap, int32 InComponentIndex, const TCHAR* InParentUEUUID, FString& OutUEUUID, FString& OutName);
+	void GetNodeUUIDAndName(const TMap<FString, FString>& InInstanceNodeMetaDataMap, const TMap<FString, FString>& InReferenceNodeMetaDataMap, int32 InComponentIndex, const TCHAR* InParentUEUUID, FString& OutUEUUID, FString& OutName);
 
 protected:
 	CADLibrary::FArchiveSceneGraph* SceneGraph;
diff --git a/Engine/Plugins/Enterprise/DatasmithCADImporter/Source/DatasmithDispatcher/Private/DatasmithDispatcher.cpp b/Engine/Plugins/Enterprise/DatasmithCADImporter/Source/DatasmithDispatcher/Private/DatasmithDispatcher.cpp
--- a/Engine/Plugins/Enterprise/DatasmithCADImporter/Source/DatasmithDispatcher/Private/DatasmithDispatcher.cpp
+++ b/Engine/Plugins/Enterprise/DatasmithCADImporter/Source/DatasmithDispatcher/Private/DatasmithDispatcher.cpp
@@ -2,6 +2,8 @@
 
 #include "DatasmithDispatcher.h"
 
+#include "CADFileData.h"
+#include "CADFileReader.h"
 #include "DatasmithDispatcherConfig.h"
 #include "DatasmithDispatcherLog.h"
 #include "DatasmithDispatcherTask.h"
@@ -93,8 +95,8 @@ void FDatasmithDispatcher::SetTaskState(int32 TaskIndex, ETaskState TaskState)
 		FileDescription = Task.FileDescription;
 
 		if (TaskState == ETaskState::ProcessOk
-		 || TaskState == ETaskState::ProcessFailed
-		 || TaskState == ETaskState::FileNotFound)
+			|| TaskState == ETaskState::ProcessFailed
+			|| TaskState == ETaskState::FileNotFound)
 		{
 			CompletedTaskCount++;
 		}
@@ -246,24 +248,24 @@ void FDatasmithDispatcher::ProcessLocal()
 	{
 		CADLibrary::FFileDescription& FileDescription = Task->FileDescription;
 
-		CADLibrary::FCoreTechFileParser FileParser(ImportParameters, *FPaths::EnginePluginsDir(), ProcessCacheFolder);
-		ETaskState ProcessResult = FileParser.ProcessFile(FileDescription);
+		CADLibrary::FCADFileReader FileReader(ImportParameters, FileDescription, *FPaths::EnginePluginsDir(), ProcessCacheFolder);
+		ETaskState ProcessResult = FileReader.ProcessFile();
 
 		ETaskState TaskState = ProcessResult;
 		SetTaskState(Task->Index, TaskState);
 
 		if (TaskState == ETaskState::ProcessOk)
 		{
-			TArray<CADLibrary::FFileDescription>& ExternalRefSet = FileParser.GetExternalRefSet();
+			const CADLibrary::FCADFileData& CADFileData = FileReader.GetCADFileData();
+			const TArray<CADLibrary::FFileDescription>& ExternalRefSet = CADFileData.GetExternalRefSet();
 			if (ExternalRefSet.Num() > 0)
 			{
-				for (CADLibrary::FFileDescription& ExternalFile : ExternalRefSet)
+				for (const CADLibrary::FFileDescription& ExternalFile : ExternalRefSet)
 				{
-					ExternalFile.MainCadFilePath = FileDescription.MainCadFilePath;
 					AddTask(ExternalFile);
 				}
 			}
-			LinkCTFileToUnrealCacheFile(FileParser.GetCADFileDescription(), FileParser.GetSceneGraphFile(), FileParser.GetMeshFileName());
+			LinkCTFileToUnrealCacheFile(CADFileData.GetCADFileDescription(), CADFileData.GetSceneGraphFileName(), CADFileData.GetMeshFileName());
 		}
 	}
 }
diff --git a/Engine/Plugins/Enterprise/DatasmithCADImporter/Source/DatasmithDispatcher/Public/DatasmithDispatcherTask.h b/Engine/Plugins/Enterprise/DatasmithCADImporter/Source/DatasmithDispatcher/Public/DatasmithDispatcherTask.h
--- a/Engine/Plugins/Enterprise/DatasmithCADImporter/Source/DatasmithDispatcher/Public/DatasmithDispatcherTask.h
+++ b/Engine/Plugins/Enterprise/DatasmithCADImporter/Source/DatasmithDispatcher/Public/DatasmithDispatcherTask.h
@@ -3,21 +3,20 @@
 #pragma once
 
 #include "CoreMinimal.h"
-#include "CoreTechFileParser.h"
 #include "CADData.h"
 
 namespace DatasmithDispatcher
 {
-using ETaskState = CADLibrary::ECoreTechParsingResult;
+using ETaskState = CADLibrary::ECADParsingResult;
 
 struct FTask
 {
 	FTask() = default;
 
 	FTask(const CADLibrary::FFileDescription& InFile)
+		: FileDescription(InFile)
+		, State(ETaskState::UnTreated)
 	{
-		FileDescription = InFile;
-		State = ETaskState::UnTreated;
 	}
 
 	CADLibrary::FFileDescription FileDescription;
diff --git a/Engine/Source/Programs/Enterprise/Datasmith/DatasmithCADWorker/Private/DatasmithCADWorkerImpl.cpp b/Engine/Source/Programs/Enterprise/Datasmith/DatasmithCADWorker/Private/DatasmithCADWorkerImpl.cpp
--- a/Engine/Source/Programs/Enterprise/Datasmith/DatasmithCADWorker/Private/DatasmithCADWorkerImpl.cpp
+++ b/Engine/Source/Programs/Enterprise/Datasmith/DatasmithCADWorker/Private/DatasmithCADWorkerImpl.cpp
@@ -159,7 +159,7 @@ void FDatasmithCADWorkerImpl::ProcessCommand(const FRunTaskCommand& RunTaskComma
 	FThread TimeCheckerThread = FThread(TEXT("TimeCheckerThread"), [&]() { CheckDuration(FileToProcess, MaxDuration); });
 
 	CADLibrary::FCoreTechFileParser FileParser(ImportParameters, EnginePluginsPath, CachePath);
-	CADLibrary::ECoreTechParsingResult ProcessResult = FileParser.ProcessFile(FileToProcess);
+	ETaskState ProcessResult = FileParser.ProcessFile(FileToProcess);
 
 	bProcessIsRunning = false;
 	TimeCheckerThread.Join();
