Fix build

diff --git a/Engine/Plugins/Enterprise/DatasmithCADImporter/Source/CADInterfaces/Private/CADFileReader.cpp b/Engine/Plugins/Enterprise/DatasmithCADImporter/Source/CADInterfaces/Private/CADFileReader.cpp
--- a/Engine/Plugins/Enterprise/DatasmithCADImporter/Source/CADInterfaces/Private/CADFileReader.cpp
+++ b/Engine/Plugins/Enterprise/DatasmithCADImporter/Source/CADInterfaces/Private/CADFileReader.cpp
@@ -6,7 +6,7 @@
 #include "CADOptions.h"
 #include "CoreTechFileParser.h"
 #include "HAL/FileManager.h"
-#include "TechSoftFileParser.h"
+//#include "TechSoftFileParser.h"
 #include "Templates/TypeHash.h"
 
 namespace CADLibrary
@@ -26,7 +26,7 @@ namespace CADLibrary
 #ifdef USE_TECHSOFT_SDK
 		if (GCADLibrary == TEXT("TechSoft"))
 		{
-			CADParser = MakeUnique<FTechSoftFileParser>(CADFileData, EnginePluginsPath);
+			//CADParser = MakeUnique<FTechSoftFileParser>(CADFileData, EnginePluginsPath);
 		}
 #endif
 	}
