CADInterfaces, CADTools refactoring:
FArchiveSceneGraph : rename arrays into Bodies, Components, UnloadedComponents, ExternalReferences, Instances;
FFileDescriptor become a class to better encapsulate the data + refactoring

diff --git a/Engine/Plugins/Enterprise/DatasmithCADImporter/Source/CADInterfaces/Private/CADFileReader.cpp b/Engine/Plugins/Enterprise/DatasmithCADImporter/Source/CADInterfaces/Private/CADFileReader.cpp
--- a/Engine/Plugins/Enterprise/DatasmithCADImporter/Source/CADInterfaces/Private/CADFileReader.cpp
+++ b/Engine/Plugins/Enterprise/DatasmithCADImporter/Source/CADInterfaces/Private/CADFileReader.cpp
@@ -6,12 +6,14 @@
 #include "CADOptions.h"
 #include "CoreTechFileParser.h"
 #include "HAL/FileManager.h"
+#ifdef USE_TECHSOFT_SDK
 //#include "TechSoftFileParser.h"
+#endif
 #include "Templates/TypeHash.h"
 
 namespace CADLibrary
 {
-	FCADFileReader::FCADFileReader(const FImportParameters& ImportParams, FFileDescription& InFile, const FString& EnginePluginsPath, const FString& InCachePath)
+	FCADFileReader::FCADFileReader(const FImportParameters& ImportParams, FFileDescriptor& InFile, const FString& EnginePluginsPath, const FString& InCachePath)
 		: CADFileData(ImportParams, InFile, InCachePath)
 	{
 #ifdef USE_KERNEL_IO_SDK
@@ -31,17 +33,17 @@ namespace CADLibrary
 #endif
 	}
 
-	bool FCADFileReader::FindFile(FFileDescription& File)
+	bool FCADFileReader::FindFile(FFileDescriptor& File)
 	{
-		FString FileName = File.Name;
+		const FString& FileName = File.GetFileName();
 
-		FString FilePath = FPaths::GetPath(File.Path);
-		FString RootFilePath = File.MainCadFilePath;
+		FString FilePath = FPaths::GetPath(File.GetSourcePath());
+		FString RootFilePath = File.GetRootFolder();
 
 		// Basic case: FilePath is, or is in a sub-folder of, RootFilePath
 		if (FilePath.StartsWith(RootFilePath))
 		{
-			return IFileManager::Get().FileExists(*File.Path);
+			return IFileManager::Get().FileExists(*File.GetSourcePath());
 		}
 
 		// Advance case: end of FilePath is in a upper-folder of RootFilePath
@@ -76,7 +78,7 @@ namespace CADLibrary
 				FString NewFilePath = FPaths::Combine(RootPaths[IndexFolderPath], FilePaths[IndexFilePath]);
 				if (IFileManager::Get().FileExists(*NewFilePath))
 				{
-					File.Path = NewFilePath;
+					File.SetSourceFilePath(NewFilePath);
 					return true;
 				};
 			}
@@ -84,12 +86,12 @@ namespace CADLibrary
 
 		// Last case: the FilePath is elsewhere and the file exist
 		// A Warning is launch because the file could be expected to not be loaded
-		if (IFileManager::Get().FileExists(*File.Path))
+		if (IFileManager::Get().FileExists(*File.GetSourcePath()))
 		{
 			return true;
 		}
 
-		CADFileData.AddWarningMessages(FString::Printf(TEXT("File %s cannot be found."), *File.Path));
+		CADFileData.AddWarningMessages(FString::Printf(TEXT("File %s cannot be found."), *File.GetFileName()));
 		return false;
 	}
 
@@ -121,7 +123,7 @@ namespace CADLibrary
 				}
 				else // the file has been converted into CT file but meshed with different parameters
 				{
-					CADFileData.ReplaceFileByCacheBackup(CADFileCachePath);
+					CADFileData.GetCADFileDescription().SetCacheFile(CADFileCachePath);
 				}
 			}
 
diff --git a/Engine/Plugins/Enterprise/DatasmithCADImporter/Source/CADInterfaces/Private/CADSceneGraph.cpp b/Engine/Plugins/Enterprise/DatasmithCADImporter/Source/CADInterfaces/Private/CADSceneGraph.cpp
--- a/Engine/Plugins/Enterprise/DatasmithCADImporter/Source/CADInterfaces/Private/CADSceneGraph.cpp
+++ b/Engine/Plugins/Enterprise/DatasmithCADImporter/Source/CADInterfaces/Private/CADSceneGraph.cpp
@@ -15,8 +15,8 @@ FArchive& operator<<(FArchive& Ar, FArchiveInstance& Instance)
 	Ar << Instance.MetaData;
 	Ar << Instance.TransformMatrix;
 	Ar << Instance.ReferenceNodeId;
-	Ar << Instance.bIsExternalRef;
-	Ar << Instance.ExternalRef;
+	Ar << Instance.bIsExternalReference;
+	Ar << Instance.ExternalReference;
 	return Ar;
 }
 
@@ -66,15 +66,15 @@ FArchive& operator<<(FArchive& Ar, FArchiveSceneGraph& SceneGraph)
 	Ar << SceneGraph.CADFileName;
 	Ar << SceneGraph.ArchiveFileName;
 	Ar << SceneGraph.FullPath;
-	Ar << SceneGraph.ExternalRefSet;
+	Ar << SceneGraph.ExternalReferences;
 
 	Ar << SceneGraph.ColorHIdToColor;
 	Ar << SceneGraph.MaterialHIdToMaterial;
 
 	Ar << SceneGraph.Instances;
-	Ar << SceneGraph.ComponentSet;
-	Ar << SceneGraph.UnloadedComponentSet;
-	Ar << SceneGraph.BodySet;
+	Ar << SceneGraph.Components;
+	Ar << SceneGraph.UnloadedComponents;
+	Ar << SceneGraph.Bodies;
 
 	Ar << SceneGraph.CADIdToInstanceIndex;
 	Ar << SceneGraph.CADIdToComponentIndex;
diff --git a/Engine/Plugins/Enterprise/DatasmithCADImporter/Source/CADInterfaces/Private/CoreTechFileParser.h b/Engine/Plugins/Enterprise/DatasmithCADImporter/Source/CADInterfaces/Private/CoreTechFileParser.h
--- a/Engine/Plugins/Enterprise/DatasmithCADImporter/Source/CADInterfaces/Private/CoreTechFileParser.h
+++ b/Engine/Plugins/Enterprise/DatasmithCADImporter/Source/CADInterfaces/Private/CoreTechFileParser.h
@@ -91,6 +91,7 @@ namespace CADLibrary
 		void GetAttributeValue(CT_ATTRIB_TYPE attrib_type, int ith_field, FString& value);
 
 		FCADFileData& CADFileData;
+		FFileDescriptor& FileDescription;
 	};
 } // CADLibrary
 
diff --git a/Engine/Plugins/Enterprise/DatasmithCADImporter/Source/CADInterfaces/Private/CoreTechFileParserImpl.cpp b/Engine/Plugins/Enterprise/DatasmithCADImporter/Source/CADInterfaces/Private/CoreTechFileParserImpl.cpp
--- a/Engine/Plugins/Enterprise/DatasmithCADImporter/Source/CADInterfaces/Private/CoreTechFileParserImpl.cpp
+++ b/Engine/Plugins/Enterprise/DatasmithCADImporter/Source/CADInterfaces/Private/CoreTechFileParserImpl.cpp
@@ -187,6 +187,7 @@ namespace CADLibrary
 
 	FCoreTechFileParser::FCoreTechFileParser(FCADFileData& InCADData, const FString& EnginePluginsPath)
 		: CADFileData(InCADData)
+		, FileDescription(InCADData.GetCADFileDescription())
 	{
 	}
 
@@ -204,22 +205,25 @@ namespace CADLibrary
 		CT_UINT32 NumberOfIds = 1;
 
 		// Set specific import option according to format
-		if (CADFileData.HasConfiguration())
+		if (FileDescription.HasConfiguration())
 		{
-			if (CADFileData.FileExtension() == TEXT("jt"))
+			switch(FileDescription.GetFileFormat())
 			{
-				LoadOption = CADFileData.GetConfiguration();
+			case ECADFormat::JT:
+			{
+				LoadOption = FileDescription.GetConfiguration();
+				break;
 			}
-			else
+			case ECADFormat::SOLIDWORKS:
 			{
-				NumberOfIds = CT_KERNEL_IO::AskFileNbOfIds(*CADFileData.FilePath());
+				NumberOfIds = CT_KERNEL_IO::AskFileNbOfIds(*FileDescription.GetPathOfFileToLoad());
 				if (NumberOfIds > 1)
 				{
-					CT_UINT32 ActiveConfig = CT_KERNEL_IO::AskFileActiveConfig(*CADFileData.FilePath());
+					CT_UINT32 ActiveConfig = CT_KERNEL_IO::AskFileActiveConfig(*FileDescription.GetPathOfFileToLoad());
 					for (CT_UINT32 i = 0; i < NumberOfIds; i++)
 					{
-						CT_STR ConfValue = CT_KERNEL_IO::AskFileIdIthName(*CADFileData.FilePath(), i);
-						if (CADFileData.GetConfiguration() == CoreTechFileParserUtils::AsFString(ConfValue))
+						CT_STR ConfValue = CT_KERNEL_IO::AskFileIdIthName(*FileDescription.GetPathOfFileToLoad(), i);
+						if (FileDescription.GetConfiguration() == CoreTechFileParserUtils::AsFString(ConfValue))
 						{
 							ActiveConfig = i;
 							break;
@@ -229,25 +233,29 @@ namespace CADLibrary
 					CTImportOption |= CT_LOAD_FLAGS_READ_SPECIFIC_OBJECT;
 					LoadOption = FString::FromInt((int32)ActiveConfig);
 				}
+				break;
+			}
+			default :
+				break;
 			}
 		}
 
 		CTKIO_ChangeUnit(CADFileData.MetricUnit());
-		Result = CT_KERNEL_IO::LoadFile(*CADFileData.FilePath(), MainId, CTImportOption, 0, *LoadOption);
+		Result = CT_KERNEL_IO::LoadFile(*FileDescription.GetPathOfFileToLoad(), MainId, CTImportOption, 0, *LoadOption);
 		if (Result == IO_ERROR_EMPTY_ASSEMBLY)
 		{
 			CT_KERNEL_IO::UnloadModel();
 			CT_FLAGS CTReImportOption = CTImportOption | CT_LOAD_FLAGS_LOAD_EXTERNAL_REF;
 			CTReImportOption &= ~CT_LOAD_FLAGS_READ_ASM_STRUCT_ONLY;  // BUG CT -> Ticket 11685
 			CTKIO_ChangeUnit(CADFileData.MetricUnit());
-			Result = CT_KERNEL_IO::LoadFile(*CADFileData.FilePath(), MainId, CTReImportOption, 0, *LoadOption);
+			Result = CT_KERNEL_IO::LoadFile(*FileDescription.GetPathOfFileToLoad(), MainId, CTReImportOption, 0, *LoadOption);
 		}
 
 		// the file is loaded but it's empty, so no data is generate
 		if (Result == IO_ERROR_EMPTY_ASSEMBLY)
 		{
 			CT_KERNEL_IO::UnloadModel();
-			CADFileData.AddWarningMessages(FString::Printf(TEXT("File %s has been loaded but no assembly has been detected."), *CADFileData.FileName()));
+			CADFileData.AddWarningMessages(FString::Printf(TEXT("File %s has been loaded but no assembly has been detected."), *FileDescription.GetFileName()));
 			return ECADParsingResult::ProcessOk;
 		}
 
@@ -260,7 +268,7 @@ namespace CADLibrary
 		if (CADFileData.IsCacheDefined())
 		{
 			FString CacheFilePath = CADFileData.GetCADCachePath();
-			if (CacheFilePath != CADFileData.FilePath())
+			if (CacheFilePath != FileDescription.GetPathOfFileToLoad())
 			{
 				CT_LIST_IO ObjectList;
 				ObjectList.PushBack(MainId);
@@ -291,10 +299,10 @@ namespace CADLibrary
 		CADFileData.ReserveBodyMeshes(NbElements[CT_BODY_INDEX]);
 
 		FArchiveSceneGraph& SceneGraphArchive = CADFileData.GetSceneGraphArchive();
-		SceneGraphArchive.BodySet.Reserve(NbElements[CT_BODY_INDEX]);
-		SceneGraphArchive.ComponentSet.Reserve(NbElements[CT_ASSEMBLY_INDEX] + NbElements[CT_PART_INDEX] + NbElements[CT_COMPONENT_INDEX]);
-		SceneGraphArchive.UnloadedComponentSet.Reserve(NbElements[CT_UNLOADED_COMPONENT_INDEX] + NbElements[CT_UNLOADED_ASSEMBLY_INDEX] + NbElements[CT_UNLOADED_PART_INDEX]);
-		SceneGraphArchive.ExternalRefSet.Reserve(NbElements[CT_UNLOADED_COMPONENT_INDEX] + NbElements[CT_UNLOADED_ASSEMBLY_INDEX] + NbElements[CT_UNLOADED_PART_INDEX]);
+		SceneGraphArchive.Bodies.Reserve(NbElements[CT_BODY_INDEX]);
+		SceneGraphArchive.Components.Reserve(NbElements[CT_ASSEMBLY_INDEX] + NbElements[CT_PART_INDEX] + NbElements[CT_COMPONENT_INDEX]);
+		SceneGraphArchive.UnloadedComponents.Reserve(NbElements[CT_UNLOADED_COMPONENT_INDEX] + NbElements[CT_UNLOADED_ASSEMBLY_INDEX] + NbElements[CT_UNLOADED_PART_INDEX]);
+		SceneGraphArchive.ExternalReferences.Reserve(NbElements[CT_UNLOADED_COMPONENT_INDEX] + NbElements[CT_UNLOADED_ASSEMBLY_INDEX] + NbElements[CT_UNLOADED_PART_INDEX]);
 		SceneGraphArchive.Instances.Reserve(NbElements[CT_INSTANCE_INDEX]);
 
 		SceneGraphArchive.CADIdToBodyIndex.Reserve(NbElements[CT_BODY_INDEX]);
@@ -315,7 +323,7 @@ namespace CADLibrary
 		CT_STR KernelIO_Version = CT_KERNEL_IO::AskVersion();
 		if (!KernelIO_Version.IsEmpty())
 		{
-			SceneGraphArchive.ComponentSet[0].MetaData.Add(TEXT("KernelIOVersion"), CoreTechFileParserUtils::AsFString(KernelIO_Version));
+			SceneGraphArchive.Components[0].MetaData.Add(TEXT("KernelIOVersion"), CoreTechFileParserUtils::AsFString(KernelIO_Version));
 		}
 
 		CT_KERNEL_IO::UnloadModel();
@@ -332,17 +340,18 @@ namespace CADLibrary
 	{
 		// Set import option
 		CT_FLAGS Flags = CT_LOAD_FLAGS_USE_DEFAULT;
+		Flags |= CT_LOAD_FLAGS_READ_META_DATA;
 
-		const FString& MainFileExt = CADFileData.FileExtension();
-
-		// Parallelization of monolithic JT file,
-		// For JT file, first step the file is read with "Structure only option"
-		// For each body, the JT file is read with "READ_SPECIFIC_OBJECT", Configuration == BodyId
-		if (MainFileExt == TEXT("jt"))
+		switch(FileDescription.GetFileFormat())
 		{
-			if (!CADFileData.HasConfiguration())
+		case ECADFormat::JT:
+		{
+			// Parallelization of monolithic JT file,
+			// For JT file, first step the file is read with "Structure only option"
+			// For each body, the JT file is read with "READ_SPECIFIC_OBJECT", Configuration == BodyId
+			if (!FileDescription.HasConfiguration())
 			{
-				FFileStatData FileStatData = IFileManager::Get().GetStatData(*CADFileData.FilePath());
+				FFileStatData FileStatData = IFileManager::Get().GetStatData(*FileDescription.GetSourcePath());
 
 				if (FileStatData.FileSize > 2e6 /* 2 Mb */ && CADFileData.IsCacheDefined()) // First step 
 				{
@@ -354,25 +363,31 @@ namespace CADLibrary
 				Flags &= ~CT_LOAD_FLAGS_REMOVE_EMPTY_COMPONENTS;
 				Flags |= CT_LOAD_FLAGS_READ_SPECIFIC_OBJECT;
 			}
+			break;
 		}
 
-		Flags |= CT_LOAD_FLAGS_READ_META_DATA;
 
-		if (MainFileExt == TEXT("catpart") || MainFileExt == TEXT("catproduct") || MainFileExt == TEXT("cgr"))
+		case ECADFormat::CATIA:
+		case ECADFormat::CATIA_CGR:
 		{
 			Flags |= CT_LOAD_FLAGS_V5_READ_GEOM_SET;
+			break;
 		}
 
-		// All the BRep topology is not available in IGES import
-		// Ask Kernel IO to complete or create missing topology
-		if (MainFileExt == TEXT("igs") || MainFileExt == TEXT("iges"))
+		case ECADFormat::IGES:
 		{
+			// All the BRep topology is not available in IGES import
+			// Ask Kernel IO to complete or create missing topology
 			Flags |= CT_LOAD_FLAG_COMPLETE_TOPOLOGY;
 			Flags |= CT_LOAD_FLAG_SEARCH_NEW_TOPOLOGY;
+			break;
+		}
+		default:
+			break;
 		}
 
 		// 3dxml file is zipped files, it's full managed by Kernel_io. We cannot read it in sequential mode
-		if (MainFileExt != TEXT("3dxml") && CADFileData.IsSequentialImport())
+		if (FileDescription.GetFileFormat() != ECADFormat::CATIA_3DXML && CADFileData.IsSequentialImport())
 		{
 			Flags &= ~CT_LOAD_FLAGS_LOAD_EXTERNAL_REF;
 		}
@@ -514,10 +529,10 @@ namespace CADLibrary
 		CT_OBJECT_IO::AskType(ReferenceNodeId, type);
 		if (type == CT_UNLOADED_PART_TYPE || type == CT_UNLOADED_COMPONENT_TYPE || type == CT_UNLOADED_ASSEMBLY_TYPE)
 		{
-			Instance.bIsExternalRef = true;
+			Instance.bIsExternalReference = true;
 			if (int32* Index = CADFileData.FindUnloadedComponentOfId(ReferenceNodeId))
 			{
-				Instance.ExternalRef = CADFileData.GetExternalReferences(*Index);
+				Instance.ExternalReference = CADFileData.GetExternalReferences(*Index);
 				return true;
 			}
 
@@ -535,18 +550,18 @@ namespace CADLibrary
 
 			if (ExternalRefFullPath.IsEmpty())
 			{
-				ExternalRefFullPath = CADFileData.FilePath();
+				ExternalRefFullPath = FileDescription.GetSourcePath();
 			}
 
 			FString Configuration;
-			if (CADFileData.FileExtension() == TEXT("jt"))
+			if (FileDescription.GetFileFormat() == ECADFormat::JT)
 			{
 				// Parallelization of monolithic JT file,
 				// is the external reference is the current file ? 
 				// Yes => this is an unloaded part that will be imported with CT_LOAD_FLAGS_READ_SPECIFIC_OBJECT Option
 				// No => the external reference is realy external... 
 				FString ExternalName = FPaths::GetCleanFilename(ExternalRefFullPath);
-				if (ExternalName == CADFileData.FileName())
+				if (ExternalName == FileDescription.GetFileName())
 				{
 					Configuration = FString::Printf(TEXT("%d"), InternalId);
 				}
@@ -561,13 +576,13 @@ namespace CADLibrary
 			FArchiveUnloadedComponent& UnloadedComponent = CADFileData.GetUnloadedComponentAt(UnloadedComponentIndex);
 			ReadNodeMetaData(ReferenceNodeId, UnloadedComponent.MetaData);
 
-			FFileDescription& NewFileDescription = CADFileData.AddExternalRef(*ExternalRefFullPath, *Configuration, *CADFileData.GetCADFileDescription().MainCadFilePath);
-			Instance.ExternalRef = NewFileDescription;
+			FFileDescriptor& NewFileDescription = CADFileData.AddExternalRef(*ExternalRefFullPath, *Configuration, *FileDescription.GetRootFolder());
+			Instance.ExternalReference = NewFileDescription;
 
 			return true;
 		}
 
-		Instance.bIsExternalRef = false;
+		Instance.bIsExternalReference = false;
 
 		return ReadComponent(ReferenceNodeId, DefaultMaterialHash);
 	}
diff --git a/Engine/Plugins/Enterprise/DatasmithCADImporter/Source/CADInterfaces/Public/CADFileData.h b/Engine/Plugins/Enterprise/DatasmithCADImporter/Source/CADInterfaces/Public/CADFileData.h
--- a/Engine/Plugins/Enterprise/DatasmithCADImporter/Source/CADInterfaces/Public/CADFileData.h
+++ b/Engine/Plugins/Enterprise/DatasmithCADImporter/Source/CADInterfaces/Public/CADFileData.h
@@ -19,7 +19,7 @@ namespace CADLibrary
 	class FCADFileData
 	{
 	public:
-		FCADFileData(const FImportParameters& InImportParameters, const FFileDescription& InFileDescription, const FString& InCachePath)
+		FCADFileData(const FImportParameters& InImportParameters, const FFileDescriptor& InFileDescription, const FString& InCachePath)
 			: CADImportSdk(ECADImportSdk::KernelIO)
 			, ImportParameters(InImportParameters)
 			, CachePath(InCachePath)
@@ -27,16 +27,16 @@ namespace CADLibrary
 			, bIsCacheDefined(!InCachePath.IsEmpty())
 			, FileDescription(InFileDescription)
 		{
-			SceneGraphArchive.FullPath = FileDescription.OriginalPath;
-			SceneGraphArchive.CADFileName = FileDescription.Name;
+			SceneGraphArchive.FullPath = FileDescription.GetSourcePath();
+			SceneGraphArchive.CADFileName = FileDescription.GetFileName();
 		}
 
 		uint32 GetSceneFileHash()
 		{
 			if (!SceneFileHash)
 			{
-				SceneFileHash = HashCombine(FileDescription.GetFileHash(), GetTypeHash(ImportParameters.StitchingTechnique));
-				SceneFileHash = HashCombine(SceneFileHash, GetTypeHash(CADImportSdk));
+				SceneFileHash = HashCombine(FileDescription.GetDescriptorHash(), ::GetTypeHash(ImportParameters.StitchingTechnique));
+				SceneFileHash = HashCombine(SceneFileHash, ::GetTypeHash(CADImportSdk));
 			}
 			return SceneFileHash;
 		}
@@ -46,12 +46,13 @@ namespace CADLibrary
 			if (!GeomFileHash)
 			{
 				GeomFileHash = GetSceneFileHash();
-				GeomFileHash = HashCombine(GeomFileHash, GetTypeHash(ImportParameters.ChordTolerance));
-				GeomFileHash = HashCombine(GeomFileHash, GetTypeHash(ImportParameters.MaxEdgeLength));
-				GeomFileHash = HashCombine(GeomFileHash, GetTypeHash(ImportParameters.MaxNormalAngle));
-				GeomFileHash = HashCombine(GeomFileHash, GetTypeHash(ImportParameters.MetricUnit));
-				GeomFileHash = HashCombine(GeomFileHash, GetTypeHash(ImportParameters.ScaleFactor));
-				GeomFileHash = HashCombine(GeomFileHash, GetTypeHash(ImportParameters.StitchingTechnique));
+				GeomFileHash = HashCombine(GeomFileHash, ::GetTypeHash(ImportParameters.bDisableCADKernelTessellation));
+				GeomFileHash = HashCombine(GeomFileHash, ::GetTypeHash(ImportParameters.ChordTolerance));
+				GeomFileHash = HashCombine(GeomFileHash, ::GetTypeHash(ImportParameters.MaxEdgeLength));
+				GeomFileHash = HashCombine(GeomFileHash, ::GetTypeHash(ImportParameters.MaxNormalAngle));
+				GeomFileHash = HashCombine(GeomFileHash, ::GetTypeHash(ImportParameters.MetricUnit));
+				GeomFileHash = HashCombine(GeomFileHash, ::GetTypeHash(ImportParameters.ScaleFactor));
+				GeomFileHash = HashCombine(GeomFileHash, ::GetTypeHash(ImportParameters.StitchingTechnique));
 			}
 			return GeomFileHash;
 		}
@@ -94,7 +95,7 @@ namespace CADLibrary
 		{
 			if (IsCacheDefined())
 			{
-				FString CacheFileName = FString::Printf(TEXT("UEx%08x"), FileDescription.GetFileHash());
+				FString CacheFileName = FString::Printf(TEXT("UEx%08x"), FileDescription.GetDescriptorHash());
 				return FPaths::Combine(CachePath, TEXT("cad"), CacheFileName + TEXT(".ct"));
 			}
 			return FString();
@@ -130,14 +131,14 @@ namespace CADLibrary
 
 		int32 AddComponent(FCadId ComponentId)
 		{
-			int32 Index = SceneGraphArchive.ComponentSet.Emplace(ComponentId);
+			int32 Index = SceneGraphArchive.Components.Emplace(ComponentId);
 			SceneGraphArchive.CADIdToComponentIndex.Add(ComponentId, Index);
 			return Index;
 		}
 
 		FArchiveComponent& GetComponentAt(int32 Index)
 		{
-			return SceneGraphArchive.ComponentSet[Index];
+			return SceneGraphArchive.Components[Index];
 		}
 
 
@@ -166,14 +167,14 @@ namespace CADLibrary
 
 		int32 AddBody(FCadId BodyId)
 		{
-			int32 Index = SceneGraphArchive.BodySet.Emplace(BodyId);
+			int32 Index = SceneGraphArchive.Bodies.Emplace(BodyId);
 			SceneGraphArchive.CADIdToBodyIndex.Add(BodyId, Index);
 			return Index;
 		}
 
 		FArchiveBody& GetBodyAt(int32 Index)
 		{
-			return SceneGraphArchive.BodySet[Index];
+			return SceneGraphArchive.Bodies[Index];
 		}
 
 
@@ -189,29 +190,29 @@ namespace CADLibrary
 
 		int32 AddUnloadedComponent(FCadId ComponentId)
 		{
-			int32 Index = SceneGraphArchive.UnloadedComponentSet.Emplace(ComponentId);
+			int32 Index = SceneGraphArchive.UnloadedComponents.Emplace(ComponentId);
 			SceneGraphArchive.CADIdToUnloadedComponentIndex.Add(ComponentId, Index);
 			return Index;
 		}
 
 		FArchiveUnloadedComponent& GetUnloadedComponentAt(int32 Index)
 		{
-			return SceneGraphArchive.UnloadedComponentSet[Index];
+			return SceneGraphArchive.UnloadedComponents[Index];
 		}
 
-		FFileDescription& GetExternalReferences(int32 Index)
+		FFileDescriptor& GetExternalReferences(int32 Index)
 		{
-			return SceneGraphArchive.ExternalRefSet[Index];
+			return SceneGraphArchive.ExternalReferences[Index];
 		}
 
-		FFileDescription& AddExternalRef(const TCHAR* InFilePath, const TCHAR* InConfiguration, const TCHAR* InRootFilePath)
+		FFileDescriptor& AddExternalRef(const TCHAR* InFilePath, const TCHAR* InConfiguration, const TCHAR* InRootFilePath)
 		{
-			return SceneGraphArchive.ExternalRefSet.Emplace_GetRef(InFilePath, InConfiguration, InRootFilePath);
+			return SceneGraphArchive.ExternalReferences.Emplace_GetRef(InFilePath, InConfiguration, InRootFilePath);
 		}
 
-		FFileDescription& AddExternalRef(const FFileDescription& InFileDescription)
+		FFileDescriptor& AddExternalRef(const FFileDescriptor& InFileDescription)
 		{
-			return SceneGraphArchive.ExternalRefSet.Emplace_GetRef(InFileDescription);
+			return SceneGraphArchive.ExternalReferences.Emplace_GetRef(InFileDescription);
 		}
 
 		FBodyMesh& AddBodyMesh(FCadId BodyId)
@@ -226,9 +227,9 @@ namespace CADLibrary
 			SerializeBodyMeshSet(*MeshArchiveFilePath, BodyMeshes);
 		}
 
-		const TArray<FFileDescription>& GetExternalRefSet() const 
+		const TArray<FFileDescriptor>& GetExternalRefSet() const 
 		{
-			return SceneGraphArchive.ExternalRefSet;
+			return SceneGraphArchive.ExternalReferences;
 		}
 
 		const FString& GetSceneGraphFileName() const
@@ -286,46 +287,16 @@ namespace CADLibrary
 			return BodyMeshes;
 		}
 
-		const FFileDescription& GetCADFileDescription() const
+		const FFileDescriptor& GetCADFileDescription() const
 		{
 			return FileDescription;
 		}
 
-		FFileDescription& GetCADFileDescription()
+		FFileDescriptor& GetCADFileDescription()
 		{
 			return FileDescription;
 		}
 
-		bool HasConfiguration() const
-		{
-			return !FileDescription.Configuration.IsEmpty();
-		}
-
-		const FString& GetConfiguration() const
-		{
-			return FileDescription.Configuration;
-		}
-
-		const FString& FileExtension() const
-		{
-			return FileDescription.Extension;
-		}
-
-		const FString& FilePath() const
-		{
-			return FileDescription.Path;
-		}
-
-		const FString& FileName() const
-		{
-			return FileDescription.Name;
-		}
-
-		void ReplaceFileByCacheBackup(const FString& CacheFilePath)
-		{
-			FileDescription.ReplaceByKernelIOBackup(CacheFilePath);
-		}
-
 		void ReserveBodyMeshes(int32 MaxBodyCount)
 		{
 			BodyMeshes.Reserve(MaxBodyCount);
@@ -378,7 +349,7 @@ namespace CADLibrary
 		const FString BodyCacheExt;
 		const bool bIsCacheDefined;
 
-		FFileDescription FileDescription;
+		FFileDescriptor FileDescription;
 
 		FString MeshArchiveFile;
 
diff --git a/Engine/Plugins/Enterprise/DatasmithCADImporter/Source/CADInterfaces/Public/CADFileReader.h b/Engine/Plugins/Enterprise/DatasmithCADImporter/Source/CADInterfaces/Public/CADFileReader.h
--- a/Engine/Plugins/Enterprise/DatasmithCADImporter/Source/CADInterfaces/Public/CADFileReader.h
+++ b/Engine/Plugins/Enterprise/DatasmithCADImporter/Source/CADInterfaces/Public/CADFileReader.h
@@ -20,7 +20,7 @@ namespace CADLibrary
 		 * @param EnginePluginsPath Full Path of EnginePlugins. Mandatory to set KernelIO to import DWG, or DGN files
 		 * @param InCachePath Full path of the cache in which the data will be saved
 		 */
-		FCADFileReader(const FImportParameters& ImportParams, FFileDescription& InCTFileDescription, const FString& EnginePluginsPath = TEXT(""), const FString& InCachePath = TEXT(""));
+		FCADFileReader(const FImportParameters& ImportParams, FFileDescriptor& InCTFileDescription, const FString& EnginePluginsPath = TEXT(""), const FString& InCachePath = TEXT(""));
 		ECADParsingResult ProcessFile();
 
 		const FCADFileData& GetCADFileData() const
@@ -34,7 +34,7 @@ namespace CADLibrary
 		}
 
 	private:
-		bool FindFile(FFileDescription& File);
+		bool FindFile(FFileDescriptor& File);
 
 		uint32 GetSceneFileHash(const uint32 InSGHash, const FImportParameters& ImportParam);
 		uint32 GetGeomFileHash(const uint32 InSGHash, const FImportParameters& ImportParam);
diff --git a/Engine/Plugins/Enterprise/DatasmithCADImporter/Source/CADInterfaces/Public/CADSceneGraph.h b/Engine/Plugins/Enterprise/DatasmithCADImporter/Source/CADInterfaces/Public/CADSceneGraph.h
--- a/Engine/Plugins/Enterprise/DatasmithCADImporter/Source/CADInterfaces/Public/CADSceneGraph.h
+++ b/Engine/Plugins/Enterprise/DatasmithCADImporter/Source/CADInterfaces/Public/CADSceneGraph.h
@@ -38,8 +38,8 @@ public:
 public:
 	FMatrix TransformMatrix = FMatrix::Identity;
 	FCadId ReferenceNodeId = 0;
-	bool bIsExternalRef = false;
-	FFileDescription ExternalRef;
+	bool bIsExternalReference = false;
+	FFileDescriptor ExternalReference;
 };
 
 class CADINTERFACES_API FArchiveComponent : public ICADArchiveObject
@@ -133,10 +133,10 @@ public:
 	TMap<FColorId, FArchiveColor> ColorHIdToColor;
 	TMap<FMaterialId, FArchiveMaterial> MaterialHIdToMaterial;
 
-	TArray<FArchiveBody> BodySet;
-	TArray<FArchiveComponent> ComponentSet;
-	TArray<FArchiveUnloadedComponent> UnloadedComponentSet;
-	TArray<FFileDescription> ExternalRefSet;
+	TArray<FArchiveBody> Bodies;
+	TArray<FArchiveComponent> Components;
+	TArray<FArchiveUnloadedComponent> UnloadedComponents;
+	TArray<FFileDescriptor> ExternalReferences;
 	TArray<FArchiveInstance> Instances;
 
 	TMap<FCadId, int32> CADIdToBodyIndex;
diff --git a/Engine/Plugins/Enterprise/DatasmithCADImporter/Source/CADTools/Private/CADData.cpp b/Engine/Plugins/Enterprise/DatasmithCADImporter/Source/CADTools/Private/CADData.cpp
--- a/Engine/Plugins/Enterprise/DatasmithCADImporter/Source/CADTools/Private/CADData.cpp
+++ b/Engine/Plugins/Enterprise/DatasmithCADImporter/Source/CADTools/Private/CADData.cpp
@@ -29,27 +29,31 @@ void GetCTColorIdAlpha(FColorId ColorId, uint32& CTColorId, uint8& Alpha)
 
 int32 BuildColorName(const FColor& Color)
 {
-	FString Name = FString::Printf(TEXT("%02x%02x%02x%02x"), Color.R, Color.G, Color.B, Color.A);
-	return FGenericPlatformMath::Abs((int32)GetTypeHash(Name));
+	return FMath::Abs((int32)GetTypeHash(Color));
 }
 
 int32 BuildMaterialName(const FCADMaterial& Material)
 {
-	FString Name;
+	using ::GetTypeHash;
+
+	uint32 MaterialName = 0;
 	if (!Material.MaterialName.IsEmpty())
 	{
-		Name += Material.MaterialName;  // we add material name because it could be used by the end user so two material with same parameters but different name are different.
+		MaterialName = GetTypeHash(*Material.MaterialName); // we add material name because it could be used by the end user so two material with same parameters but different name are different.
 	}
-	Name += FString::Printf(TEXT("%02x%02x%02x "), Material.Diffuse.R, Material.Diffuse.G, Material.Diffuse.B);
-	Name += FString::Printf(TEXT("%02x%02x%02x "), Material.Ambient.R, Material.Ambient.G, Material.Ambient.B);
-	Name += FString::Printf(TEXT("%02x%02x%02x "), Material.Specular.R, Material.Specular.G, Material.Specular.B);
-	Name += FString::Printf(TEXT("%02x%02x%02x"), (int)(Material.Shininess * 255.0), (int)(Material.Transparency * 255.0), (int)(Material.Reflexion * 255.0));
+
+	MaterialName = HashCombine(MaterialName, GetTypeHash(Material.Diffuse));
+	MaterialName = HashCombine(MaterialName, GetTypeHash(Material.Ambient));
+	MaterialName = HashCombine(MaterialName, GetTypeHash(Material.Specular));
+	MaterialName = HashCombine(MaterialName, GetTypeHash((int)(Material.Shininess * 255.0)));
+	MaterialName = HashCombine(MaterialName, GetTypeHash((int)(Material.Transparency * 255.0)));
+	MaterialName = HashCombine(MaterialName, GetTypeHash((int)(Material.Reflexion * 255.0)));
 
 	if (!Material.TextureName.IsEmpty())
 	{
-		Name += Material.TextureName;
+		MaterialName = HashCombine(MaterialName, GetTypeHash(*Material.TextureName));
 	}
-	return FMath::Abs((int32)GetTypeHash(Name));
+	return FMath::Abs((int32) MaterialName);
 }
 
 FArchive& operator<<(FArchive& Ar, FCADMaterial& Material)
@@ -65,12 +69,14 @@ FArchive& operator<<(FArchive& Ar, FCADMaterial& Material)
 	return Ar;
 }
 
-FArchive& operator<<(FArchive& Ar, FFileDescription& File)
+FArchive& operator<<(FArchive& Ar, FFileDescriptor& File)
 {
-	Ar << File.Path;
+	Ar << File.SourceFilePath;
+	Ar << File.CacheFilePath;
 	Ar << File.Name;
-	Ar << File.Extension;
 	Ar << File.Configuration;
+	Ar << File.Format;
+	Ar << File.RootFolder;
 	return Ar;
 }
 
@@ -152,32 +158,57 @@ void GetCleanFilenameAndExtension(const FString& InFilePath, FString& OutFilenam
 	if (!OutExtension.IsEmpty() && FCString::IsNumeric(*OutExtension))
 	{
 		BaseFile = OutFilename;
-		BaseFile.Split(TEXT("."), &OutFilename, &OutExtension, ESearchCase::CaseSensitive, ESearchDir::FromEnd);
-		OutExtension = OutExtension + TEXT(".*");
+		FString NewExtension;
+		BaseFile.Split(TEXT("."), &OutFilename, &NewExtension, ESearchCase::CaseSensitive, ESearchDir::FromEnd);
+		if (!NewExtension.IsEmpty())
+		{
+			OutExtension = NewExtension + TEXT(".*");
+		}
 	}
 }
 
-uint32 FFileDescription::GetFileHash() const
+FString GetExtension(const FString& InFilePath)
 {
-	FFileStatData FileStatData = IFileManager::Get().GetStatData(*OriginalPath);
+	if (InFilePath.IsEmpty())
+	{
+		return FString();
+	}
+
+	FString Extension;
+	FString BaseFileWithoutExt;
+
+	FString BaseFile = FPaths::GetCleanFilename(InFilePath);
+	BaseFile.Split(TEXT("."), &BaseFileWithoutExt, &Extension, ESearchCase::CaseSensitive, ESearchDir::FromEnd);
+
+	if (!Extension.IsEmpty() && FCString::IsNumeric(*Extension))
+	{
+		FString NewExtension;
+		BaseFileWithoutExt.Split(TEXT("."), &BaseFile, &NewExtension, ESearchCase::CaseSensitive, ESearchDir::FromEnd);
+		if (!NewExtension.IsEmpty())
+		{
+			NewExtension += TEXT(".*");
+			return NewExtension;
+		}
+	}
+	return Extension;
+}
 
-	FDateTime ModificationTime = FileStatData.ModificationTime;
+uint32 GetTypeHash(const FFileDescriptor& FileDescriptor)
+{
+	using ::GetTypeHash;
+	FFileStatData FileStatData = IFileManager::Get().GetStatData(*FileDescriptor.SourceFilePath);
 
-	uint32 FileHash = GetTypeHash(*Name);
-	FileHash = HashCombine(FileHash, GetTypeHash(*Configuration));
-	FileHash = HashCombine(FileHash, GetTypeHash(FileStatData.FileSize));
-	FileHash = HashCombine(FileHash, GetTypeHash(ModificationTime));
+	uint32 DescriptorHash = GetTypeHash(*FileDescriptor.Name);
+	DescriptorHash = HashCombine(DescriptorHash, GetTypeHash(*FileDescriptor.Configuration));
+	DescriptorHash = HashCombine(DescriptorHash, GetTypeHash(FileStatData.FileSize));
+	DescriptorHash = HashCombine(DescriptorHash, GetTypeHash(FileStatData.ModificationTime));
 
-	return FileHash;
+	return DescriptorHash;
 }
 
 ECADFormat FileFormat(const FString& Extension)
 {
-	if (Extension == TEXT("sldprt") || Extension == TEXT("sldasm"))
-	{
-		return ECADFormat::SOLIDWORKS;
-	}
-	else if (Extension == TEXT("catpart") || Extension == TEXT("catproduct"))
+	if (Extension == TEXT("catpart") || Extension == TEXT("catproduct"))
 	{
 		return ECADFormat::CATIA;
 	}
@@ -185,10 +216,6 @@ ECADFormat FileFormat(const FString& Extension)
 	{
 		return ECADFormat::CATIA_CGR;
 	}
-	else if (Extension == TEXT("x_t") || Extension == TEXT("x_b"))
-	{
-		return ECADFormat::PARASOLID;
-	}
 	else if (Extension == TEXT("iges") || Extension == TEXT("igs"))
 	{
 		return ECADFormat::IGES;
@@ -197,10 +224,18 @@ ECADFormat FileFormat(const FString& Extension)
 	{
 		return ECADFormat::STEP;
 	}
+	else if (Extension == TEXT("ipt") || Extension == TEXT("iam"))
+	{
+		return ECADFormat::INVENTOR;
+	}
 	else if (Extension == TEXT("jt"))
 	{
 		return ECADFormat::JT;
 	}
+	else if (Extension == TEXT("model"))
+	{
+		return ECADFormat::CATIAV4;
+	}
 	else if (Extension == TEXT("prt.*") || Extension == TEXT("asm.*"))
 	{
 		return ECADFormat::CREO;
@@ -209,25 +244,21 @@ ECADFormat FileFormat(const FString& Extension)
 	{
 		return ECADFormat::NX;
 	}
-	else if (Extension == TEXT("step") || Extension == TEXT("stp"))
-	{
-		return ECADFormat::STEP;
-	}
 	else if (Extension == TEXT("sat"))
 	{
 		return ECADFormat::ACIS;
 	}
-	else if (Extension == TEXT("model"))
+	else if (Extension == TEXT("sldprt") || Extension == TEXT("sldasm"))
 	{
-		return ECADFormat::CATIAV4;
+		return ECADFormat::SOLIDWORKS;
 	}
-	else if (Extension == TEXT("3dxml"))
+	else if (Extension == TEXT("x_t") || Extension == TEXT("x_b"))
 	{
-		return ECADFormat::CATIA_3DXML;
+		return ECADFormat::PARASOLID;
 	}
-	else if (Extension == TEXT("ipt") || Extension == TEXT("iam"))
+	else if (Extension == TEXT("3dxml"))
 	{
-		return ECADFormat::INVENTOR;
+		return ECADFormat::CATIA_3DXML;
 	}
 	else
 	{
diff --git a/Engine/Plugins/Enterprise/DatasmithCADImporter/Source/CADTools/Public/CADData.h b/Engine/Plugins/Enterprise/DatasmithCADImporter/Source/CADTools/Public/CADData.h
--- a/Engine/Plugins/Enterprise/DatasmithCADImporter/Source/CADTools/Public/CADData.h
+++ b/Engine/Plugins/Enterprise/DatasmithCADImporter/Source/CADTools/Public/CADData.h
@@ -52,6 +52,8 @@ enum class ECADParsingResult : uint8
 
 // TODO: Remove from hear and replace by DatasmithUtils::GetCleanFilenameAndExtension... But need to remove DatasmithCore dependancies 
 CADTOOLS_API void GetCleanFilenameAndExtension(const FString& InFilePath, FString& OutFilename, FString& OutExtension);
+CADTOOLS_API FString GetExtension(const FString& InFilePath);
+
 
 class CADTOOLS_API FCADMaterial
 {
@@ -76,51 +78,109 @@ struct CADTOOLS_API FObjectDisplayDataId
 	FColorId Color = 0; // => FastHash == ColorId+Transparency
 };
 
-struct CADTOOLS_API FFileDescription
+class CADTOOLS_API FFileDescriptor
 {
-	explicit FFileDescription(const TCHAR* InFilePath = nullptr, const TCHAR* InConfiguration = nullptr, const TCHAR* InRootFilePath = nullptr)
-		: Path(InFilePath)
-		, OriginalPath(InFilePath)
+public:
+	FFileDescriptor() = default;
+	
+	explicit FFileDescriptor(const TCHAR* InFilePath, const TCHAR* InConfiguration = nullptr, const TCHAR* InRootFolder = nullptr)
+		: SourceFilePath(InFilePath)
 		, Configuration(InConfiguration)
-		, MainCadFilePath(InRootFilePath)
 	{
-		if (MainCadFilePath.IsEmpty() && !Path.IsEmpty())
-		{
-			MainCadFilePath = FPaths::GetPath(Path);
-		}
-
-		GetCleanFilenameAndExtension(Path, Name, Extension);
-		Name += TEXT(".") + Extension;
+		Name = FPaths::GetCleanFilename(InFilePath);
+		Format = FileFormat(GetExtension(InFilePath));
+		RootFolder = InRootFolder ? InRootFolder : FPaths::GetPath(InFilePath);
 	}
 
 	/**
-	 * Used to replace CADFile path by the path of the file saved in KernelIO format (*.ct)
+	 * Used define and then load the cache of the CAD File instead of the source file
 	 */ 
-	void ReplaceByKernelIOBackup(const FString& InKernelIOBackupPath)
+	void SetCacheFile(const FString& InCacheFilePath)
 	{
-		Path = InKernelIOBackupPath;
+		CacheFilePath = InCacheFilePath;
 	}
 
-	bool operator==(const FFileDescription& Other) const
+	bool operator==(const FFileDescriptor& Other) const
 	{
 		return (Name.Equals(Other.Name, ESearchCase::IgnoreCase) && (Configuration == Other.Configuration));
 	}
 
-	friend CADTOOLS_API FArchive& operator<<(FArchive& Ar, FFileDescription& File);
+	friend CADTOOLS_API FArchive& operator<<(FArchive& Ar, FFileDescriptor& File);
+
+	friend CADTOOLS_API uint32 GetTypeHash(const FFileDescriptor& FileDescription);
+
+	uint32 GetDescriptorHash() const
+	{
+		if (!DescriptorHash)
+		{
+			DescriptorHash = GetTypeHash(*this);
+		}
+		return DescriptorHash;
+	}
+
+	const FString& GetSourcePath() const 
+	{
+		return SourceFilePath;
+	}
+	
+	bool HasConfiguration() const
+	{
+		return !Configuration.IsEmpty();
+	}
+
+	const FString& GetConfiguration() const
+	{
+		return Configuration;
+	}
+
+	void SetConfiguration(const FString& NewConfiguration)
+	{
+		Configuration = NewConfiguration;
+	}
+
+	ECADFormat GetFileFormat() const
+	{
+		return Format;
+	}
 
-	uint32 GetFileHash() const;
+	const FString& GetPathOfFileToLoad() const
+	{
+		if (CacheFilePath.IsEmpty())
+		{
+			return SourceFilePath;
+		}
+		else
+		{
+			return CacheFilePath;
+		}
+	}
 
-	ECADFormat GetFormat()
+	/** Set the file path if SourceFilePath was not the real path */
+	void SetSourceFilePath(const FString& NewFilePath)
 	{
-		return FileFormat(Extension);
+		SourceFilePath = NewFilePath;
 	}
 
-	FString Path;
-	FString OriginalPath;
-	FString Name;
-	FString Extension;
-	FString Configuration;
-	FString MainCadFilePath;
+	const FString& GetRootFolder() const 
+	{
+		return RootFolder;
+	}
+
+	const FString& GetFileName() const
+	{
+		return Name;
+	}
+
+private:
+
+	FString SourceFilePath; // e.g. d:/folder/content.jt
+	FString CacheFilePath; // if the file has already been loaded 
+	FString Name; // content.jt
+	ECADFormat Format; // ECADFormat::JT
+	FString Configuration; // dedicated to JT or SW file to read the good configuration (SW) or only a sub-file (JT)
+	FString RootFolder; // alternative folder where the file could be if its path is not valid.
+
+	mutable uint32 DescriptorHash = 0;
 };
 
 /**
@@ -278,11 +338,5 @@ inline void CopyValue(const void* Source, int Offset, uint8 Size, int32 Dest[3])
 	}
 }
 
-using ::GetTypeHash;
-FORCEINLINE CADTOOLS_API uint32 GetTypeHash(const FFileDescription& FileDescription)
-{
-	return HashCombine(GetTypeHash(FileDescription.Name), GetTypeHash(FileDescription.Configuration));
-};
-
 }
 
diff --git a/Engine/Plugins/Enterprise/DatasmithCADImporter/Source/DatasmithCADTranslator/Private/DatasmithCADTranslator.cpp b/Engine/Plugins/Enterprise/DatasmithCADImporter/Source/DatasmithCADTranslator/Private/DatasmithCADTranslator.cpp
--- a/Engine/Plugins/Enterprise/DatasmithCADImporter/Source/DatasmithCADTranslator/Private/DatasmithCADTranslator.cpp
+++ b/Engine/Plugins/Enterprise/DatasmithCADImporter/Source/DatasmithCADTranslator/Private/DatasmithCADTranslator.cpp
@@ -95,39 +95,50 @@ bool FDatasmithCADTranslator::LoadScene(TSharedRef<IDatasmithScene> DatasmithSce
 	ImportParameters.MaxNormalAngle = TesselationOptions.NormalTolerance;
 	ImportParameters.StitchingTechnique = (CADLibrary::EStitchingTechnique)TesselationOptions.StitchingTechnique;
 
-	CADLibrary::FFileDescription FileDescription(*FPaths::ConvertRelativePathToFull(GetSource().GetSourceFile()),
-		TEXT(""),
-		*FPaths::GetPath(FPaths::ConvertRelativePathToFull(GetSource().GetSourceFile())));
+	CADLibrary::FFileDescriptor FileDescription(*FPaths::ConvertRelativePathToFull(GetSource().GetSourceFile()));
 
 	// Do not change the model unit when translator is called by the Datasmith runtime plugin.
-#if WITH_EDITOR
-	if (FileDescription.Extension == TEXT("jt") && IsInGameThread())
+	ImportParameters.ModelCoordSys = FDatasmithUtils::EModelCoordSystem::ZUp_RightHanded;
+
+	switch (FileDescription.GetFileFormat())
 	{
-		ImportParameters.MetricUnit = 1.;
-		ImportParameters.ScaleFactor = 100.;
-	}
+	case CADLibrary::ECADFormat::JT:
+	{
+#if WITH_EDITOR
+		if (IsInGameThread())
+		{
+			// In UE Editor, for historical QA reason (i.e. due to some validation files) the unit of JT file is set to meter as it seemed to be the default unit
+			// But setting the meter as the default unit requires resetting KernelIO that crash at the runtime...
+			// so for runtime use, metric unit is mm.
+			// this will change with the next release of CAD importer
+			ImportParameters.MetricUnit = 1.;
+			ImportParameters.ScaleFactor = 100.;
+		}
 #endif
-
-	ImportParameters.ModelCoordSys = FDatasmithUtils::EModelCoordSystem::ZUp_RightHanded;
-	if (FileDescription.Extension == TEXT("prt")) // NX
+		break;
+	}
+	case CADLibrary::ECADFormat::NX:
 	{
 		ImportParameters.ModelCoordSys = FDatasmithUtils::EModelCoordSystem::ZUp_RightHanded;
 		ImportParameters.DisplayPreference = CADLibrary::EDisplayPreference::ColorOnly;
 		ImportParameters.Propagation = CADLibrary::EDisplayDataPropagationMode::BodyOnly;
+		break;
 	}
-	else if (FileDescription.Extension == TEXT("sldprt") || FileDescription.Extension == TEXT("sldasm") || // Solidworks
-		FileDescription.Extension == TEXT("iam") || FileDescription.Extension == TEXT("ipt") || // Inventor
-		FileDescription.Extension.StartsWith(TEXT("asm")) || FileDescription.Extension.StartsWith(TEXT("creo")) || FileDescription.Extension.StartsWith(TEXT("prt")) || FileDescription.Extension.StartsWith(TEXT("neu")) // Creo
-		)
+	case CADLibrary::ECADFormat::SOLIDWORKS:
+	case CADLibrary::ECADFormat::INVENTOR:
+	case CADLibrary::ECADFormat::CREO:
 	{
 		ImportParameters.ModelCoordSys = FDatasmithUtils::EModelCoordSystem::YUp_RightHanded;
 		ImportParameters.DisplayPreference = CADLibrary::EDisplayPreference::ColorOnly;
 		ImportParameters.Propagation = CADLibrary::EDisplayDataPropagationMode::BodyOnly;
+		break;
 	}
-	else if (FileDescription.Extension == TEXT("dwg")) // Autocad
+	case CADLibrary::ECADFormat::DWG:
 	{
 		ImportParameters.DisplayPreference = CADLibrary::EDisplayPreference::ColorOnly;
 		ImportParameters.Propagation = CADLibrary::EDisplayDataPropagationMode::BodyOnly;
+		break;
+	}
 	}
 
 	ImportParameters.bEnableTimeControl = CADLibrary::bGEnableTimeControl;
diff --git a/Engine/Plugins/Enterprise/DatasmithCADImporter/Source/DatasmithCADTranslator/Private/DatasmithSceneGraphBuilder.cpp b/Engine/Plugins/Enterprise/DatasmithCADImporter/Source/DatasmithCADTranslator/Private/DatasmithSceneGraphBuilder.cpp
--- a/Engine/Plugins/Enterprise/DatasmithCADImporter/Source/DatasmithCADTranslator/Private/DatasmithSceneGraphBuilder.cpp
+++ b/Engine/Plugins/Enterprise/DatasmithCADImporter/Source/DatasmithCADTranslator/Private/DatasmithSceneGraphBuilder.cpp
@@ -87,13 +87,13 @@ bool FDatasmithSceneGraphBuilder::Build()
 {
 	LoadSceneGraphDescriptionFiles();
 
-	uint32 rootHash = GetTypeHash(rootFileDescription);
-	SceneGraph = CADFileToSceneGraphArchive.FindRef(rootHash);
+	uint32 RootHash = RootFileDescription.GetDescriptorHash();
+	SceneGraph = CADFileToSceneGraphArchive.FindRef(RootHash);
 	if (!SceneGraph)
 	{
 		return false;
 	}
-	AncestorSceneGraphHash.Add(rootHash);
+	AncestorSceneGraphHash.Add(RootHash);
 
 	return FDatasmithSceneBaseGraphBuilder::Build();
 }
@@ -128,7 +128,7 @@ void FDatasmithSceneGraphBuilder::LoadSceneGraphDescriptionFiles()
 
 void FDatasmithSceneGraphBuilder::FillAnchorActor(const TSharedRef<IDatasmithActorElement>& ActorElement, const FString& CleanFilenameOfCADFile)
 {
-	CADLibrary::FFileDescription AnchorDescription(*CleanFilenameOfCADFile);
+	CADLibrary::FFileDescriptor AnchorDescription(*CleanFilenameOfCADFile);
 	SceneGraph = CADFileToSceneGraphArchive.FindRef(GetTypeHash(AnchorDescription));
 	if (!SceneGraph)
 	{
@@ -146,7 +146,7 @@ void FDatasmithSceneGraphBuilder::FillAnchorActor(const TSharedRef<IDatasmithAct
 
 	// TODO: check ParentData and Index validity?
 	ActorData ParentData(ActorElement->GetName());
-	CADLibrary::FArchiveComponent& Component = SceneGraph->ComponentSet[*Index];
+	CADLibrary::FArchiveComponent& Component = SceneGraph->Components[*Index];
 
 	TMap<FString, FString> InstanceNodeMetaDataMap;
 	FString ActorUUID;
@@ -180,7 +180,7 @@ FDatasmithSceneBaseGraphBuilder::FDatasmithSceneBaseGraphBuilder(CADLibrary::FAr
 	, DatasmithScene(InScene)
 	, ImportParameters(InImportParameters)
 	, ImportParametersHash(ImportParameters.GetHash())
-	, rootFileDescription(*InSource.GetSourceFile())
+	, RootFileDescription(*InSource.GetSourceFile())
 {
 	if (InSceneGraph)
 	{
@@ -208,7 +208,7 @@ bool FDatasmithSceneBaseGraphBuilder::Build()
 	}
 
 	ActorData Data(TEXT(""));
-	CADLibrary::FArchiveComponent& Component = SceneGraph->ComponentSet[*Index];
+	CADLibrary::FArchiveComponent& Component = SceneGraph->Components[*Index];
 	TSharedPtr< IDatasmithActorElement > RootActor = BuildComponent(Component, Data);
 	DatasmithScene->AddActor(RootActor);
 
@@ -221,77 +221,70 @@ bool FDatasmithSceneBaseGraphBuilder::Build()
 		DatasmithScene->SetProductVersion(**ProductVersion);
 	}
 
-	FString* ProductName = Component.MetaData.Find(TEXT("Input_Format_and_Emitter"));
-	if(ProductName)
+	FString ProductName;
+	const FString* ProductNamePtr = Component.MetaData.Find(TEXT("Input_Format_and_Emitter"));
+	if(ProductNamePtr)
 	{
-		ProductName->TrimStartAndEndInline();
-		if (!ProductName->IsEmpty())
+		ProductName = *ProductNamePtr;
+		ProductName.TrimStartAndEndInline();
+		if (!ProductName.IsEmpty())
 		{
-			DatasmithScene->SetProductName(**ProductName);
-		}
-		else
-		{
-			ProductName = nullptr;
+			DatasmithScene->SetProductName(*ProductName);
 		}
 	}
 
-	if(!ProductName)
+	if(ProductName.IsEmpty())
 	{
-		if (rootFileDescription.Extension == TEXT("jt"))
+		switch (RootFileDescription.GetFileFormat())
 		{
+		case CADLibrary::ECADFormat::JT:
 			DatasmithScene->SetProductName(TEXT("Jt"));
-		}
-		else if (rootFileDescription.Extension == TEXT("sldprt") || rootFileDescription.Extension == TEXT("sldasm"))
-		{
+			break;
+		case CADLibrary::ECADFormat::SOLIDWORKS:
 			DatasmithScene->SetProductName(TEXT("SolidWorks"));
-		}
-		else if (rootFileDescription.Extension == TEXT("catpart") || rootFileDescription.Extension == TEXT("catproduct") || rootFileDescription.Extension == TEXT("cgr"))
-		{
+			break;
+		case CADLibrary::ECADFormat::ACIS:
+			DatasmithScene->SetProductName(TEXT("3D ACIS"));
+			break;
+		case CADLibrary::ECADFormat::CATIA:
 			DatasmithScene->SetProductName(TEXT("CATIA V5"));
-		}
-		else if (rootFileDescription.Extension == TEXT("3dxml") || rootFileDescription.Extension == TEXT("3drep"))
-		{
+			break;
+		case CADLibrary::ECADFormat::CATIA_CGR:
+			DatasmithScene->SetProductName(TEXT("CATIA V5"));
+			break;
+		case CADLibrary::ECADFormat::CATIAV4:
+			DatasmithScene->SetProductName(TEXT("CATIA V4"));
+			break;
+		case CADLibrary::ECADFormat::CATIA_3DXML:
 			DatasmithScene->SetProductName(TEXT("3D XML"));
-		}
-		else if (rootFileDescription.Extension == TEXT("iam") || rootFileDescription.Extension == TEXT("ipt"))
-		{
+			break;
+		case CADLibrary::ECADFormat::CREO:
+			DatasmithScene->SetProductName(TEXT("Creo"));
+			break;
+		case CADLibrary::ECADFormat::IGES:
+			DatasmithScene->SetProductName(TEXT("IGES"));
+			break;
+		case CADLibrary::ECADFormat::INVENTOR:
 			DatasmithScene->SetProductName(TEXT("Inventor"));
-		}
-		else if (rootFileDescription.Extension == TEXT("prt") || rootFileDescription.Extension == TEXT("asm"))
-		{
+			break;
+		case CADLibrary::ECADFormat::NX:
 			DatasmithScene->SetProductName(TEXT("NX"));
-		}
-		else if (rootFileDescription.Extension == TEXT("stp") || rootFileDescription.Extension == TEXT("step"))
-		{
-			DatasmithScene->SetProductName(TEXT("STEP"));
-		}
-		else if (rootFileDescription.Extension == TEXT("igs") || rootFileDescription.Extension == TEXT("iges"))
-		{
-			DatasmithScene->SetProductName(TEXT("IGES"));
-		}
-		else if (rootFileDescription.Extension == TEXT("x_t") || rootFileDescription.Extension == TEXT("x_b"))
-		{
+			break;
+		case CADLibrary::ECADFormat::PARASOLID:
 			DatasmithScene->SetProductName(TEXT("Parasolid"));
-		}
-		else if (rootFileDescription.Extension == TEXT("dwg"))
-		{
+			break;
+		case CADLibrary::ECADFormat::STEP:
+			DatasmithScene->SetProductName(TEXT("STEP"));
+			break;
+		case CADLibrary::ECADFormat::DWG:
 			DatasmithScene->SetProductName(TEXT("AutoCAD"));
-		}
-		else if (rootFileDescription.Extension == TEXT("dgn"))
-		{
+			break;
+		case CADLibrary::ECADFormat::DGN:
 			DatasmithScene->SetProductName(TEXT("Micro Station"));
-		}
-		else if (rootFileDescription.Extension == TEXT("sat"))
-		{
-			DatasmithScene->SetProductName(TEXT("3D ACIS"));
-		}
-		else if (rootFileDescription.Extension.StartsWith(TEXT("asm")) || rootFileDescription.Extension.StartsWith(TEXT("creo")) || rootFileDescription.Extension.StartsWith(TEXT("prt")) || rootFileDescription.Extension.StartsWith(TEXT("neu")))
-		{
-			DatasmithScene->SetProductName(TEXT("Creo"));
-		}
-		else
-		{
+			break;
+		default:
 			DatasmithScene->SetProductName(TEXT("Unknown"));
+			break;
 		}
 	}
 
@@ -306,11 +299,11 @@ TSharedPtr< IDatasmithActorElement >  FDatasmithSceneBaseGraphBuilder::BuildInst
 	CADLibrary::FArchiveInstance& Instance = SceneGraph->Instances[InstanceIndex];
 
 	CADLibrary::FArchiveSceneGraph* InstanceSceneGraph = SceneGraph;
-	if (Instance.bIsExternalRef)
+	if (Instance.bIsExternalReference)
 	{
-		if (!Instance.ExternalRef.Path.IsEmpty())
+		if (!Instance.ExternalReference.GetSourcePath().IsEmpty())
 		{
-			uint32 InstanceSceneGraphHash = GetTypeHash(Instance.ExternalRef);
+			uint32 InstanceSceneGraphHash = GetTypeHash(Instance.ExternalReference);
 			SceneGraph = CADFileToSceneGraphArchive.FindRef(InstanceSceneGraphHash);
 			if (SceneGraph)
 			{
@@ -322,7 +315,7 @@ TSharedPtr< IDatasmithActorElement >  FDatasmithSceneBaseGraphBuilder::BuildInst
 					const int32* Index = SceneGraph->CADIdToComponentIndex.Find(RootId);
 					if (Index)
 					{
-						Reference = &(SceneGraph->ComponentSet[*Index]);
+						Reference = &(SceneGraph->Components[*Index]);
 					}
 				}
 			}
@@ -334,7 +327,7 @@ TSharedPtr< IDatasmithActorElement >  FDatasmithSceneBaseGraphBuilder::BuildInst
 			const int32* Index = SceneGraph->CADIdToUnloadedComponentIndex.Find(Instance.ReferenceNodeId);
 			if (Index)
 			{
-				Reference = &(SceneGraph->UnloadedComponentSet[*Index]);
+				Reference = &(SceneGraph->UnloadedComponents[*Index]);
 			}
 		}
 	}
@@ -343,7 +336,7 @@ TSharedPtr< IDatasmithActorElement >  FDatasmithSceneBaseGraphBuilder::BuildInst
 		const int32* Index = SceneGraph->CADIdToComponentIndex.Find(Instance.ReferenceNodeId);
 		if (Index)
 		{
-			Reference = &(SceneGraph->ComponentSet[*Index]);
+			Reference = &(SceneGraph->Components[*Index]);
 		}
 	}
 
@@ -492,7 +485,7 @@ TSharedPtr<IDatasmithActorElement> FDatasmithSceneBaseGraphBuilder::BuildBody(in
 {
 	TMap<FString, FString> InstanceNodeMetaDataMap;
 
-	CADLibrary::FArchiveBody& Body = SceneGraph->BodySet[BodyIndex];
+	CADLibrary::FArchiveBody& Body = SceneGraph->Bodies[BodyIndex];
 
 	FString BodyUUID;
 	FString BodyLabel;
@@ -684,7 +677,7 @@ void FDatasmithSceneBaseGraphBuilder::AddMetaData(TSharedPtr< IDatasmithActorEle
 			}
 			else
 			{
-				FString FileDir = FPaths::GetPath(rootFileDescription.Path);
+				FString FileDir = RootFileDescription.GetRootFolder();
 				FString FilePath = FPaths::Combine(FileDir, OFilePath);
 
 				if (FPaths::FileExists(FilePath))
diff --git a/Engine/Plugins/Enterprise/DatasmithCADImporter/Source/DatasmithCADTranslator/Public/DatasmithSceneGraphBuilder.h b/Engine/Plugins/Enterprise/DatasmithCADImporter/Source/DatasmithCADTranslator/Public/DatasmithSceneGraphBuilder.h
--- a/Engine/Plugins/Enterprise/DatasmithCADImporter/Source/DatasmithCADTranslator/Public/DatasmithSceneGraphBuilder.h
+++ b/Engine/Plugins/Enterprise/DatasmithCADImporter/Source/DatasmithCADTranslator/Public/DatasmithSceneGraphBuilder.h
@@ -86,7 +86,7 @@ protected:
 	TSharedRef<IDatasmithScene> DatasmithScene;
 	const CADLibrary::FImportParameters& ImportParameters;
 	const uint32 ImportParametersHash;
-	CADLibrary::FFileDescription rootFileDescription;
+	CADLibrary::FFileDescriptor RootFileDescription;
 
 	TArray<CADLibrary::FArchiveSceneGraph> ArchiveMockUps;
 	TMap<uint32, CADLibrary::FArchiveSceneGraph*> CADFileToSceneGraphArchive;
diff --git a/Engine/Plugins/Enterprise/DatasmithCADImporter/Source/DatasmithDispatcher/Private/DatasmithDispatcher.cpp b/Engine/Plugins/Enterprise/DatasmithCADImporter/Source/DatasmithDispatcher/Private/DatasmithDispatcher.cpp
--- a/Engine/Plugins/Enterprise/DatasmithCADImporter/Source/DatasmithDispatcher/Private/DatasmithDispatcher.cpp
+++ b/Engine/Plugins/Enterprise/DatasmithCADImporter/Source/DatasmithDispatcher/Private/DatasmithDispatcher.cpp
@@ -38,7 +38,7 @@ FDatasmithDispatcher::FDatasmithDispatcher(const CADLibrary::FImportParameters&
 	}
 }
 
-void FDatasmithDispatcher::AddTask(const CADLibrary::FFileDescription & InFileDescription)
+void FDatasmithDispatcher::AddTask(const CADLibrary::FFileDescriptor& InFileDescription)
 {
 	FScopeLock Lock(&TaskPoolCriticalSection);
 	for (const FTask& Task : TaskPool)
@@ -81,7 +81,7 @@ TOptional<FTask> FDatasmithDispatcher::GetNextTask()
 
 void FDatasmithDispatcher::SetTaskState(int32 TaskIndex, ETaskState TaskState)
 {
-	CADLibrary::FFileDescription FileDescription;
+	CADLibrary::FFileDescriptor FileDescriptor;
 	{
 		FScopeLock Lock(&TaskPoolCriticalSection);
 
@@ -92,7 +92,7 @@ void FDatasmithDispatcher::SetTaskState(int32 TaskIndex, ETaskState TaskState)
 
 		FTask& Task = TaskPool[TaskIndex];
 		Task.State = TaskState;
-		FileDescription = Task.FileDescription;
+		FileDescriptor = Task.FileDescription;
 
 		if (TaskState == ETaskState::ProcessOk
 			|| TaskState == ETaskState::ProcessFailed
@@ -107,10 +107,10 @@ void FDatasmithDispatcher::SetTaskState(int32 TaskIndex, ETaskState TaskState)
 		}
 	}
 
-	UE_CLOG(TaskState == ETaskState::ProcessOk, LogDatasmithDispatcher, Verbose, TEXT("File processed: %s"), *FileDescription.Name);
-	UE_CLOG(TaskState == ETaskState::UnTreated, LogDatasmithDispatcher, Warning, TEXT("File resubmitted: %s"), *FileDescription.Name);
-	UE_CLOG(TaskState == ETaskState::ProcessFailed, LogDatasmithDispatcher, Error, TEXT("File processing failure: %s"), *FileDescription.Name);
-	UE_CLOG(TaskState == ETaskState::FileNotFound, LogDatasmithDispatcher, Warning, TEXT("file not found: %s"), *FileDescription.Path);
+	UE_CLOG(TaskState == ETaskState::ProcessOk, LogDatasmithDispatcher, Verbose, TEXT("File processed: %s"), *FileDescriptor.GetFileName());
+	UE_CLOG(TaskState == ETaskState::UnTreated, LogDatasmithDispatcher, Warning, TEXT("File resubmitted: %s"), *FileDescriptor.GetFileName());
+	UE_CLOG(TaskState == ETaskState::ProcessFailed, LogDatasmithDispatcher, Error, TEXT("File processing failure: %s"), *FileDescriptor.GetFileName());
+	UE_CLOG(TaskState == ETaskState::FileNotFound, LogDatasmithDispatcher, Warning, TEXT("file not found: %s"), *FileDescriptor.GetSourcePath());
 }
 
 void FDatasmithDispatcher::Process(bool bWithProcessor)
@@ -199,7 +199,7 @@ bool FDatasmithDispatcher::IsOver()
 	return CompletedTaskCount == TaskPool.Num();
 }
 
-void FDatasmithDispatcher::LinkCTFileToUnrealCacheFile(const CADLibrary::FFileDescription& CTFileDescription, const FString& UnrealSceneGraphFile, const FString& UnrealMeshFile)
+void FDatasmithDispatcher::LinkCTFileToUnrealCacheFile(const CADLibrary::FFileDescriptor& CTFileDescription, const FString& UnrealSceneGraphFile, const FString& UnrealMeshFile)
 {
 	FScopeLock Lock(&TaskPoolCriticalSection);
 
@@ -246,7 +246,7 @@ void FDatasmithDispatcher::ProcessLocal()
 {
 	while (TOptional<FTask> Task = GetNextTask())
 	{
-		CADLibrary::FFileDescription& FileDescription = Task->FileDescription;
+		CADLibrary::FFileDescriptor& FileDescription = Task->FileDescription;
 
 		CADLibrary::FCADFileReader FileReader(ImportParameters, FileDescription, *FPaths::EnginePluginsDir(), ProcessCacheFolder);
 		ETaskState ProcessResult = FileReader.ProcessFile();
@@ -257,10 +257,10 @@ void FDatasmithDispatcher::ProcessLocal()
 		if (TaskState == ETaskState::ProcessOk)
 		{
 			const CADLibrary::FCADFileData& CADFileData = FileReader.GetCADFileData();
-			const TArray<CADLibrary::FFileDescription>& ExternalRefSet = CADFileData.GetExternalRefSet();
+			const TArray<CADLibrary::FFileDescriptor>& ExternalRefSet = CADFileData.GetExternalRefSet();
 			if (ExternalRefSet.Num() > 0)
 			{
-				for (const CADLibrary::FFileDescription& ExternalFile : ExternalRefSet)
+				for (const CADLibrary::FFileDescriptor& ExternalFile : ExternalRefSet)
 				{
 					AddTask(ExternalFile);
 				}
diff --git a/Engine/Plugins/Enterprise/DatasmithCADImporter/Source/DatasmithDispatcher/Private/DatasmithWorkerHandler.cpp b/Engine/Plugins/Enterprise/DatasmithCADImporter/Source/DatasmithDispatcher/Private/DatasmithWorkerHandler.cpp
--- a/Engine/Plugins/Enterprise/DatasmithCADImporter/Source/DatasmithDispatcher/Private/DatasmithWorkerHandler.cpp
+++ b/Engine/Plugins/Enterprise/DatasmithCADImporter/Source/DatasmithDispatcher/Private/DatasmithWorkerHandler.cpp
@@ -333,7 +333,7 @@ void FDatasmithWorkerHandler::ProcessCommand(FCompletedTaskCommand& CompletedTas
 		return;
 	}
 
-	for (const CADLibrary::FFileDescription& ExternalReferenceFile : CompletedTaskCommand.ExternalReferences)
+	for (const CADLibrary::FFileDescriptor& ExternalReferenceFile : CompletedTaskCommand.ExternalReferences)
 	{
 		Dispatcher.AddTask(ExternalReferenceFile);
 	}
diff --git a/Engine/Plugins/Enterprise/DatasmithCADImporter/Source/DatasmithDispatcher/Public/DatasmithCommands.h b/Engine/Plugins/Enterprise/DatasmithCADImporter/Source/DatasmithDispatcher/Public/DatasmithCommands.h
--- a/Engine/Plugins/Enterprise/DatasmithCADImporter/Source/DatasmithDispatcher/Public/DatasmithCommands.h
+++ b/Engine/Plugins/Enterprise/DatasmithCADImporter/Source/DatasmithDispatcher/Public/DatasmithCommands.h
@@ -86,7 +86,7 @@ protected:
 	virtual void SerializeImpl(FArchive&) override;
 
 public:
-	CADLibrary::FFileDescription JobFileDescription;
+	CADLibrary::FFileDescriptor JobFileDescription;
 	int32 JobIndex = -1;
 };
 
@@ -99,7 +99,7 @@ protected:
 	virtual void SerializeImpl(FArchive&) override;
 
 public:
-	TArray<CADLibrary::FFileDescription> ExternalReferences;
+	TArray<CADLibrary::FFileDescriptor> ExternalReferences;
 	FString SceneGraphFileName;
 	FString GeomFileName;
 	ETaskState ProcessResult = ETaskState::Unknown;
diff --git a/Engine/Plugins/Enterprise/DatasmithCADImporter/Source/DatasmithDispatcher/Public/DatasmithDispatcher.h b/Engine/Plugins/Enterprise/DatasmithCADImporter/Source/DatasmithDispatcher/Public/DatasmithDispatcher.h
--- a/Engine/Plugins/Enterprise/DatasmithCADImporter/Source/DatasmithDispatcher/Public/DatasmithDispatcher.h
+++ b/Engine/Plugins/Enterprise/DatasmithCADImporter/Source/DatasmithDispatcher/Public/DatasmithDispatcher.h
@@ -18,14 +18,14 @@ class DATASMITHDISPATCHER_API FDatasmithDispatcher
 public:
 	FDatasmithDispatcher(const CADLibrary::FImportParameters& InImportParameters, const FString& InCacheDir, int32 InNumberOfWorkers, TMap<uint32, FString>& CADFileToUnrealFileMap, TMap<uint32, FString>& CADFileToUnrealGeomMap);
 
-	void AddTask(const CADLibrary::FFileDescription & FileDescription);
+	void AddTask(const CADLibrary::FFileDescriptor & FileDescription);
 	TOptional<FTask> GetNextTask();
 	void SetTaskState(int32 TaskIndex, ETaskState TaskState);
 
 	void Process(bool bWithProcessor);
 	bool IsOver();
 
-	void LinkCTFileToUnrealCacheFile(const CADLibrary::FFileDescription& CTFileDescription, const FString& UnrealSceneGraphFile, const FString& UnrealGeomFile);
+	void LinkCTFileToUnrealCacheFile(const CADLibrary::FFileDescriptor& CTFileDescription, const FString& UnrealSceneGraphFile, const FString& UnrealGeomFile);
 
 	void LogWarningMessages(const TArray<FString>& Warnings) const;
 
diff --git a/Engine/Plugins/Enterprise/DatasmithCADImporter/Source/DatasmithDispatcher/Public/DatasmithDispatcherTask.h b/Engine/Plugins/Enterprise/DatasmithCADImporter/Source/DatasmithDispatcher/Public/DatasmithDispatcherTask.h
--- a/Engine/Plugins/Enterprise/DatasmithCADImporter/Source/DatasmithDispatcher/Public/DatasmithDispatcherTask.h
+++ b/Engine/Plugins/Enterprise/DatasmithCADImporter/Source/DatasmithDispatcher/Public/DatasmithDispatcherTask.h
@@ -13,13 +13,13 @@ struct FTask
 {
 	FTask() = default;
 
-	FTask(const CADLibrary::FFileDescription& InFile)
+	FTask(const CADLibrary::FFileDescriptor& InFile)
 		: FileDescription(InFile)
 		, State(ETaskState::UnTreated)
 	{
 	}
 
-	CADLibrary::FFileDescription FileDescription;
+	CADLibrary::FFileDescriptor FileDescription;
 	int32 Index = -1;
 	ETaskState State = ETaskState::Unknown;
 };
diff --git a/Engine/Plugins/Enterprise/DatasmithCADImporter/Source/DatasmithPLMXMLTranslator/Private/DatasmithPlmXmlImporter.cpp b/Engine/Plugins/Enterprise/DatasmithCADImporter/Source/DatasmithPLMXMLTranslator/Private/DatasmithPlmXmlImporter.cpp
--- a/Engine/Plugins/Enterprise/DatasmithCADImporter/Source/DatasmithPLMXMLTranslator/Private/DatasmithPlmXmlImporter.cpp
+++ b/Engine/Plugins/Enterprise/DatasmithCADImporter/Source/DatasmithPLMXMLTranslator/Private/DatasmithPlmXmlImporter.cpp
@@ -261,7 +261,7 @@ namespace PlmXml
 		// Adds geom file to load and returns Id to use in InstantiateMesh later(after all is loaded)
 		int32 AddMeshToLoad(const FString& FullPath)
 		{
-			CADLibrary::FFileDescription FileDescription(*FullPath);
+			CADLibrary::FFileDescriptor FileDescription(*FullPath);
 			DatasmithDispatcher->AddTask(FileDescription);
 			return FilePaths.Add(FullPath);
 		}
