Datasmith: untangle OpenNurbsBRepToCADKernelConverter

diff --git a/Engine/Plugins/Enterprise/DatasmithCADImporter/Source/DatasmithOpenNurbsTranslator/Private/OpenNurbsBRepToCADKernelConverter.cpp b/Engine/Plugins/Enterprise/DatasmithCADImporter/Source/DatasmithOpenNurbsTranslator/Private/OpenNurbsBRepToCADKernelConverter.cpp
--- a/Engine/Plugins/Enterprise/DatasmithCADImporter/Source/DatasmithOpenNurbsTranslator/Private/OpenNurbsBRepToCADKernelConverter.cpp
+++ b/Engine/Plugins/Enterprise/DatasmithCADImporter/Source/DatasmithOpenNurbsTranslator/Private/OpenNurbsBRepToCADKernelConverter.cpp
@@ -39,76 +39,51 @@ using namespace CADKernel;
 namespace
 {
 	enum EAxis { U, V };
-	struct PerAxisInfo
+
+	void FillPerAxisInfo(EAxis Axis, ON_NurbsSurface& OpenNurbsSurface, FNurbsSurfaceHomogeneousData& OutNurbsInfo)
 	{
-		ON_NurbsSurface& OpenNurbsSurface;
-		EAxis Axis;
-		int32& Degree; // Degree + 1
-		TArray<double>& Knots; // t values with superflux values
-		int32& CtrlVertCount; // number of control points
-
-		uint32 KnotSize;  // ON knots
-		uint32 KnotCount; // CT knots
-
-		TArray<uint32> KnotMultiplicities; // from ON, not relevant as we send n-plicated knots to CT (dbg only)
-
-		PerAxisInfo(EAxis InAxis, ON_NurbsSurface& InSurface, FNurbsSurfaceHomogeneousData& NurbsInfo)
-			: OpenNurbsSurface(InSurface)
-			, Axis(InAxis)
-			, Degree(Axis == U ? NurbsInfo.UDegree : NurbsInfo.VDegree)
-			, Knots(Axis == U ? NurbsInfo.UNodalVector : NurbsInfo.VNodalVector)
-			, CtrlVertCount(Axis == U ? NurbsInfo.PoleUCount : NurbsInfo.PoleVCount)
-		{
-			// detect cases not handled by CADKernel, that is knot vectors with multiplicity < order on either end
-			Degree = OpenNurbsSurface.Order(Axis) - 1;
-			if (OpenNurbsSurface.KnotMultiplicity(Axis, 0) < Degree || OpenNurbsSurface.KnotMultiplicity(Axis, KnotSize - 3) < Degree)
-			{
-				OpenNurbsSurface.IncreaseDegree(Axis, OpenNurbsSurface.Degree(Axis) + 1);
-			}
-			Populate();
-		}
+		int32& Degree = Axis == EAxis::U ? OutNurbsInfo.UDegree : OutNurbsInfo.VDegree;
+		Degree = OpenNurbsSurface.Order(Axis) - 1;
+
+		int32& CtrlVertCount = Axis == EAxis::U ? OutNurbsInfo.PoleUCount : OutNurbsInfo.PoleVCount; // number of control points
+		CtrlVertCount = OpenNurbsSurface.CVCount(Axis);
 
-	private:
-		void Populate()
+		uint32 KnotSize = Degree + CtrlVertCount + 1;
+
+		// detect cases not handled by CADKernel, that is knot vectors with multiplicity < order on either end
+		if (OpenNurbsSurface.KnotMultiplicity(Axis, 0) < Degree || OpenNurbsSurface.KnotMultiplicity(Axis, KnotSize - 3) < Degree)
 		{
+			OpenNurbsSurface.IncreaseDegree(Axis, OpenNurbsSurface.Degree(Axis) + 1);
 			Degree = OpenNurbsSurface.Order(Axis) - 1;
-			CtrlVertCount = OpenNurbsSurface.CVCount(Axis);
 			KnotSize = Degree + CtrlVertCount + 1;
-			KnotCount = OpenNurbsSurface.KnotCount(Axis);
-
-			KnotMultiplicities.SetNumUninitialized(KnotSize - 2);
-			for (uint32 Index = 0; Index < KnotSize - 2; ++Index)
-			{
-				KnotMultiplicities[Index] = OpenNurbsSurface.KnotMultiplicity(Axis, Index); // 0 and < Order + CV_count - 2
-			}
+		}
 
-			Knots.Reserve(KnotSize);
-			Knots.Add(OpenNurbsSurface.SuperfluousKnot(Axis, 0));
-			for (uint32 i = 0; i < KnotCount; ++i)
-			{
-				Knots.Add(OpenNurbsSurface.Knot(Axis, i));
-			}
-			Knots.Add(OpenNurbsSurface.SuperfluousKnot(Axis, 1));
+		TArray<double>& Knots = Axis == EAxis::U ? OutNurbsInfo.UNodalVector : OutNurbsInfo.VNodalVector; // t values with superfluous values
+		Knots.Reserve(KnotSize);
+		Knots.Add(OpenNurbsSurface.SuperfluousKnot(Axis, 0));
+		uint32 KnotCount = OpenNurbsSurface.KnotCount(Axis);
+		for (uint32 i = 0; i < KnotCount; ++i)
+		{
+			Knots.Add(OpenNurbsSurface.Knot(Axis, i));
 		}
-	};
+		Knots.Add(OpenNurbsSurface.SuperfluousKnot(Axis, 1));
+	}
 }
 
 TSharedRef<FSurface> FOpenNurbsBRepToCADKernelConverter::AddSurface(ON_NurbsSurface& OpenNurbsSurface)
 {
 	FNurbsSurfaceHomogeneousData NurbsData;
+	FillPerAxisInfo(U, OpenNurbsSurface, NurbsData);
+	FillPerAxisInfo(V, OpenNurbsSurface, NurbsData);
 
 	int32 ControlVertexDimension = OpenNurbsSurface.CVSize();
-
-	PerAxisInfo UInfo(U, OpenNurbsSurface, NurbsData);
-	PerAxisInfo VInfo(V, OpenNurbsSurface, NurbsData);
-
-	NurbsData.HomogeneousPoles.SetNumUninitialized(UInfo.CtrlVertCount * VInfo.CtrlVertCount * OpenNurbsSurface.CVSize());
+	NurbsData.HomogeneousPoles.SetNumUninitialized(NurbsData.PoleUCount * NurbsData.PoleVCount * ControlVertexDimension);
 	double* ControlPoints = NurbsData.HomogeneousPoles.GetData();
 	NurbsData.bIsRational = OpenNurbsSurface.IsRational();
 	ON::point_style PointStyle = OpenNurbsSurface.IsRational() ? ON::point_style::euclidean_rational : ON::point_style::not_rational;
-	for (int32 UIndex = 0; UIndex < UInfo.CtrlVertCount; ++UIndex)
+	for (int32 UIndex = 0; UIndex < NurbsData.PoleUCount; ++UIndex)
 	{
-		for (int32 VIndex = 0; VIndex < VInfo.CtrlVertCount; ++VIndex, ControlPoints += ControlVertexDimension)
+		for (int32 VIndex = 0; VIndex < NurbsData.PoleVCount; ++VIndex, ControlPoints += ControlVertexDimension)
 		{
 			OpenNurbsSurface.GetCV(UIndex, VIndex, PointStyle, ControlPoints);
 		}
