encapsulation of ImportParameter to well manage its setting and the respect of CVar

diff --git a/Engine/Config/ConsoleVariables.ini b/Engine/Config/ConsoleVariables.ini
--- a/Engine/Config/ConsoleVariables.ini
+++ b/Engine/Config/ConsoleVariables.ini
@@ -111,20 +111,20 @@ fx.SkipVectorVMBackendOptimizations=1
 ;ds.CADTranslator.MaxImportThreads = 0
 
 ; Uncomment to enable/disable temporary CAD processing file cache. These file will be use in a next import to avoid CAD file processing. 
-; If MaxImportThreads != 1, EnableCADCache value is ignored
+; EnableCADCache = true enforce MaxImportThreads = 1
 ; Default is enable
-;ds.CADTranslator.EnableCADCache=1
+;ds.CADTranslator.EnableCADCache=true
 
 ; Uncomment to overwrite any existing cache associated with the file being imported. 
 ; Default is false
-;ds.CADTranslator.OverwriteCache=0
+;ds.CADTranslator.OverwriteCache=false
 
 ; Uncomment to enable/disable the timer that kill the worker if the import time is unusually long. With this time control, the load of the corrupted file 
 ; is canceled but the rest of the scene is imported. Default is enable
 ;ds.CADTranslator.EnableTimeControl=1
 
 ; Uncomment to enable/disable CADKernel tessellation. If disable, CAD SDK tessellator is used.
-;ds.CADTranslator.DisableCADKernelTessellation=0
+;ds.CADTranslator.DisableCADKernelTessellation=false
 
 ; Uncomment to select the SDK to load CAD files: "KernelIO" or "TechSoft". Default is "KernelIO"
 ;ds.CADTranslator.CADLibrary="KernelIO"
diff --git a/Engine/Plugins/Enterprise/DatasmithCADImporter/Source/CADInterfaces/Private/CADFileReader.cpp b/Engine/Plugins/Enterprise/DatasmithCADImporter/Source/CADInterfaces/Private/CADFileReader.cpp
--- a/Engine/Plugins/Enterprise/DatasmithCADImporter/Source/CADInterfaces/Private/CADFileReader.cpp
+++ b/Engine/Plugins/Enterprise/DatasmithCADImporter/Source/CADInterfaces/Private/CADFileReader.cpp
@@ -4,12 +4,16 @@
 
 #include "CADFileParser.h"
 #include "CADOptions.h"
-#include "CoreTechFileParser.h"
 #include "HAL/FileManager.h"
+#include "Templates/TypeHash.h"
+
+#ifdef USE_KERNEL_IO_SDK
+#include "CoreTechFileParser.h"
+#endif
+
 #ifdef USE_TECHSOFT_SDK
-//#include "TechSoftFileParser.h"
+#include "TechSoftFileParser.h"
 #endif
-#include "Templates/TypeHash.h"
 
 namespace CADLibrary
 {
@@ -28,7 +32,7 @@ namespace CADLibrary
 #ifdef USE_TECHSOFT_SDK
 		if (GCADLibrary == TEXT("TechSoft"))
 		{
-			//CADParser = MakeUnique<FTechSoftFileParser>(CADFileData, EnginePluginsPath);
+			CADParser = MakeUnique<FTechSoftFileParser>(CADFileData, EnginePluginsPath);
 		}
 #endif
 	}
@@ -107,14 +111,14 @@ namespace CADLibrary
 			return ECADParsingResult::FileNotFound;
 		}
 
-		if (CADFileData.IsSequentialImport())
+		if (FImportParameters::bGEnableCADCache)
 		{
 			CADFileData.SetArchiveNames();
 
 			bool bNeedToProceed = true;
 
 			FString CADFileCachePath = CADFileData.GetCADCachePath();
-			if (!CADFileData.GetImportParameters().bOverwriteCache && IFileManager::Get().FileExists(*CADFileCachePath))
+			if (!FImportParameters::bGOverwriteCache && IFileManager::Get().FileExists(*CADFileCachePath))
 			{
 				FString MeshArchiveFilePath = CADFileData.GetMeshArchiveFilePath();
 				if (IFileManager::Get().FileExists(*MeshArchiveFilePath)) // the file has been proceed with same meshing parameters
diff --git a/Engine/Plugins/Enterprise/DatasmithCADImporter/Source/CADInterfaces/Private/CTSession.cpp b/Engine/Plugins/Enterprise/DatasmithCADImporter/Source/CADInterfaces/Private/CTSession.cpp
--- a/Engine/Plugins/Enterprise/DatasmithCADImporter/Source/CADInterfaces/Private/CTSession.cpp
+++ b/Engine/Plugins/Enterprise/DatasmithCADImporter/Source/CADInterfaces/Private/CTSession.cpp
@@ -22,7 +22,7 @@ namespace CADLibrary
 
 	bool FCTSession::TopoFixes(double SewingToleranceFactor)
 	{
-		return CADLibrary::CTKIO_Repair(MainObjectId, ImportParams.StitchingTechnique, SewingToleranceFactor);
+		return CADLibrary::CTKIO_Repair(MainObjectId, ImportParams.GetStitchingTechnique(), SewingToleranceFactor);
 	}
 
 	void FCTSession::SetSceneUnit(double InMetricUnit)
@@ -31,15 +31,10 @@ namespace CADLibrary
 		CTKIO_ChangeUnit(InMetricUnit);
 	}
 
-	void FCTSession::SetImportParameters(double ChordTolerance, double MaxEdgeLength, double NormalTolerance, CADLibrary::EStitchingTechnique StitchingTechnique, bool bScaleUVMap)
+	void FCTSession::SetImportParameters(double ChordTolerance, double MaxEdgeLength, double NormalTolerance, CADLibrary::EStitchingTechnique StitchingTechnique)
 	{
-		ImportParams.ChordTolerance = ChordTolerance;
-		ImportParams.MaxEdgeLength = MaxEdgeLength;
-		ImportParams.MaxNormalAngle = NormalTolerance;
-		ImportParams.StitchingTechnique = StitchingTechnique;
-		ImportParams.bScaleUVMap = bScaleUVMap;
-
-		CTKIO_ChangeTesselationParameters(ImportParams.ChordTolerance, ImportParams.MaxEdgeLength, ImportParams.MaxNormalAngle);
+		ImportParams.SetTesselationParameters(ChordTolerance, MaxEdgeLength, NormalTolerance, StitchingTechnique);
+		CTKIO_ChangeTesselationParameters(ImportParams.GetChordTolerance(), ImportParams.GetMaxEdgeLength(), ImportParams.GetMaxNormalAngle());
 	}
 
 } // namespace CADLibrary
diff --git a/Engine/Plugins/Enterprise/DatasmithCADImporter/Source/CADInterfaces/Private/CoreTechFileParserImpl.cpp b/Engine/Plugins/Enterprise/DatasmithCADImporter/Source/CADInterfaces/Private/CoreTechFileParserImpl.cpp
--- a/Engine/Plugins/Enterprise/DatasmithCADImporter/Source/CADInterfaces/Private/CoreTechFileParserImpl.cpp
+++ b/Engine/Plugins/Enterprise/DatasmithCADImporter/Source/CADInterfaces/Private/CoreTechFileParserImpl.cpp
@@ -240,14 +240,15 @@ namespace CADLibrary
 			}
 		}
 
-		CTKIO_ChangeUnit(CADFileData.MetricUnit());
+		const FImportParameters& ImportParameters = CADFileData.GetImportParameters();
+		CTKIO_ChangeUnit(ImportParameters.GetMetricUnit());
 		Result = CT_KERNEL_IO::LoadFile(*FileDescription.GetPathOfFileToLoad(), MainId, CTImportOption, 0, *LoadOption);
 		if (Result == IO_ERROR_EMPTY_ASSEMBLY)
 		{
 			CT_KERNEL_IO::UnloadModel();
 			CT_FLAGS CTReImportOption = CTImportOption | CT_LOAD_FLAGS_LOAD_EXTERNAL_REF;
 			CTReImportOption &= ~CT_LOAD_FLAGS_READ_ASM_STRUCT_ONLY;  // BUG CT -> Ticket 11685
-			CTKIO_ChangeUnit(CADFileData.MetricUnit());
+			CTKIO_ChangeUnit(ImportParameters.GetMetricUnit());
 			Result = CT_KERNEL_IO::LoadFile(*FileDescription.GetPathOfFileToLoad(), MainId, CTReImportOption, 0, *LoadOption);
 		}
 
@@ -279,12 +280,12 @@ namespace CADLibrary
 
 		CoreTechFileParserUtils::AddFaceIdAttribut(MainId);
 
-		if (CADFileData.GetStitchingTechnique() != StitchingNone && !CADFileData.IsCADKernelTessellation())
+		if (ImportParameters.GetStitchingTechnique() != StitchingNone && !FImportParameters::bGDisableCADKernelTessellation)
 		{
-			CADLibrary::CTKIO_Repair(MainId, CADFileData.GetStitchingTechnique(), 10.);
+			CADLibrary::CTKIO_Repair(MainId, ImportParameters.GetStitchingTechnique(), 10.);
 		}
 
-		CTKIO_SetCoreTechTessellationState(CADFileData.GetImportParameters());
+		CTKIO_SetCoreTechTessellationState(ImportParameters);
 
 		const CT_OBJECT_TYPE TypeSet[] = { CT_INSTANCE_TYPE, CT_ASSEMBLY_TYPE, CT_PART_TYPE, CT_COMPONENT_TYPE, CT_BODY_TYPE, CT_UNLOADED_COMPONENT_TYPE, CT_UNLOADED_ASSEMBLY_TYPE, CT_UNLOADED_PART_TYPE };
 		enum EObjectTypeIndex : uint8	{ CT_INSTANCE_INDEX = 0, CT_ASSEMBLY_INDEX, CT_PART_INDEX, CT_COMPONENT_INDEX, CT_BODY_INDEX, CT_UNLOADED_COMPONENT_INDEX, CT_UNLOADED_ASSEMBLY_INDEX, CT_UNLOADED_PART_INDEX };
@@ -346,11 +347,11 @@ namespace CADLibrary
 		{
 		case ECADFormat::JT:
 		{
-			// Parallelization of monolithic JT file,
-			// For JT file, first step the file is read with "Structure only option"
-			// For each body, the JT file is read with "READ_SPECIFIC_OBJECT", Configuration == BodyId
+		// Parallelization of monolithic JT file,
+		// For JT file, first step the file is read with "Structure only option"
+		// For each body, the JT file is read with "READ_SPECIFIC_OBJECT", Configuration == BodyId
 			if (!FileDescription.HasConfiguration())
-			{
+		{
 				FFileStatData FileStatData = IFileManager::Get().GetStatData(*FileDescription.GetSourcePath());
 
 				if (FileStatData.FileSize > 2e6 /* 2 Mb */ && CADFileData.IsCacheDefined()) // First step 
@@ -376,8 +377,8 @@ namespace CADLibrary
 
 		case ECADFormat::IGES:
 		{
-			// All the BRep topology is not available in IGES import
-			// Ask Kernel IO to complete or create missing topology
+		// All the BRep topology is not available in IGES import
+		// Ask Kernel IO to complete or create missing topology
 			Flags |= CT_LOAD_FLAG_COMPLETE_TOPOLOGY;
 			Flags |= CT_LOAD_FLAG_SEARCH_NEW_TOPOLOGY;
 			break;
@@ -387,7 +388,7 @@ namespace CADLibrary
 		}
 
 		// 3dxml file is zipped files, it's full managed by Kernel_io. We cannot read it in sequential mode
-		if (FileDescription.GetFileFormat() != ECADFormat::CATIA_3DXML && CADFileData.IsSequentialImport())
+		if (FileDescription.GetFileFormat() != ECADFormat::CATIA_3DXML && FImportParameters::bGEnableCADCache)
 		{
 			Flags &= ~CT_LOAD_FLAGS_LOAD_EXTERNAL_REF;
 		}
@@ -472,7 +473,7 @@ namespace CADLibrary
 			CT_FLAGS BodyProperties;
 			CT_BODY_IO::AskProperties(BodyId, BodyProperties);
 
-			if (!CADFileData.IsCADKernelTessellation() || !(BodyProperties & CT_BODY_PROP_EXACT))
+			if (!FImportParameters::bGDisableCADKernelTessellation || !(BodyProperties & CT_BODY_PROP_EXACT))
 			{
 				if (ReadKioBody(BodyId, ComponentId, DefaultMaterialHash, false))
 				{
@@ -648,9 +649,9 @@ namespace CADLibrary
 			CoreTechFileParserUtils::GetCTObjectDisplayDataIds(FaceID, FaceMaterial);
 			SetFaceMainMaterial(FaceMaterial, BodyMaterial, BodyMesh, Index);
 
-			if (CADFileData.NeedUVMapScaling() && Tessellation.TexCoordArray.Num() > 0)
+			if (CADFileData.GetImportParameters().NeedScaleUVMap() && Tessellation.TexCoordArray.Num() > 0)
 			{
-				CoreTechFileParserUtils::ScaleUV(FaceID, Tessellation.TexCoordArray, (float) CADFileData.GetScaleFactor());
+				CoreTechFileParserUtils::ScaleUV(FaceID, Tessellation.TexCoordArray, (float) CADFileData.GetImportParameters().GetScaleFactor());
 			}
 		};
 
@@ -696,7 +697,7 @@ namespace CADLibrary
 		BodyMesh.MeshActorName = ArchiveBody.MeshActorName;
 
 		{
-			TSharedRef<CADKernel::FSession> CADKernelSession = MakeShared<CADKernel::FSession>(0.00001 / CADFileData.GetMetricUnit());
+			TSharedRef<CADKernel::FSession> CADKernelSession = MakeShared<CADKernel::FSession>(0.00001 / CADFileData.GetImportParameters().GetMetricUnit());
 
 			TSharedRef<CADKernel::FModel> CADKernelModel = CADKernelSession->GetModel();
 
@@ -1620,24 +1621,6 @@ namespace CADLibrary
 			return CtName.IsEmpty() ? FString() : CtName.toUnicode();
 		};
 
-		uint32 GetSceneFileHash(const uint32 InSGHash, const FImportParameters& ImportParam)
-		{
-			uint32 FileHash = HashCombine(InSGHash, ::GetTypeHash(ImportParam.StitchingTechnique));
-			return FileHash;
-		}
-
-		uint32 GetGeomFileHash(const uint32 InSGHash, const FImportParameters& ImportParam)
-		{
-			uint32 FileHash = InSGHash;
-			FileHash = HashCombine(FileHash, ::GetTypeHash(ImportParam.ChordTolerance));
-			FileHash = HashCombine(FileHash, ::GetTypeHash(ImportParam.MaxEdgeLength));
-			FileHash = HashCombine(FileHash, ::GetTypeHash(ImportParam.MaxNormalAngle));
-			FileHash = HashCombine(FileHash, ::GetTypeHash(ImportParam.MetricUnit));
-			FileHash = HashCombine(FileHash, ::GetTypeHash(ImportParam.ScaleFactor));
-			FileHash = HashCombine(FileHash, ::GetTypeHash(ImportParam.StitchingTechnique));
-			return FileHash;
-		}
-
 		bool GetColor(uint32 ColorUuid, FColor& OutColor)
 		{
 			uint32 ColorId;
diff --git a/Engine/Plugins/Enterprise/DatasmithCADImporter/Source/CADInterfaces/Private/CoreTechInterfaceImpl.cpp b/Engine/Plugins/Enterprise/DatasmithCADImporter/Source/CADInterfaces/Private/CoreTechInterfaceImpl.cpp
--- a/Engine/Plugins/Enterprise/DatasmithCADImporter/Source/CADInterfaces/Private/CoreTechInterfaceImpl.cpp
+++ b/Engine/Plugins/Enterprise/DatasmithCADImporter/Source/CADInterfaces/Private/CoreTechInterfaceImpl.cpp
@@ -449,13 +449,13 @@ namespace CADLibrary
 
 		// convert max edge length to model units
 		double ModelUnit = FLT_MAX; // default value is huge, as 0.0 causes bugs in KernelIO...
-		if (ImportParams.MaxEdgeLength > SMALL_NUMBER)
+		if (ImportParams.GetMaxEdgeLength() > SMALL_NUMBER)
 		{
-			ModelUnit = ImportParams.MaxEdgeLength / ImportParams.ScaleFactor;
+			ModelUnit = ImportParams.GetMaxEdgeLength() / ImportParams.GetScaleFactor();
 		}
 
 		// Apply retrieved tessellation parameters to CoreTech tessellation engine
-		return ChangeTesselationParameters(ImportParams.ChordTolerance / ImportParams.ScaleFactor, ModelUnit, ImportParams.MaxNormalAngle);
+		return ChangeTesselationParameters(ImportParams.GetChordTolerance() / ImportParams.GetScaleFactor(), ModelUnit, ImportParams.GetMaxNormalAngle());
 	}
 
 	//ECADParsingResult FCoreTechInterfaceImpl::LoadFile(const FFileDescription& InFileDescription, const FImportParameters& InImportParameters, const FString& InCachePath, FArchiveSceneGraph& OutSceneGraphArchive, TArray<FString>& OutWarningMessages, TArray<FBodyMesh>& OutBodyMeshes)
diff --git a/Engine/Plugins/Enterprise/DatasmithCADImporter/Source/CADInterfaces/Public/CADFileData.h b/Engine/Plugins/Enterprise/DatasmithCADImporter/Source/CADInterfaces/Public/CADFileData.h
--- a/Engine/Plugins/Enterprise/DatasmithCADImporter/Source/CADInterfaces/Public/CADFileData.h
+++ b/Engine/Plugins/Enterprise/DatasmithCADImporter/Source/CADInterfaces/Public/CADFileData.h
@@ -23,7 +23,7 @@ namespace CADLibrary
 			: CADImportSdk(ECADImportSdk::KernelIO)
 			, ImportParameters(InImportParameters)
 			, CachePath(InCachePath)
-			, BodyCacheExt(InImportParameters.bDisableCADKernelTessellation ? TEXT(".ct") : TEXT(".ugeom"))
+			, BodyCacheExt(FImportParameters::bGDisableCADKernelTessellation ? TEXT(".ct") : TEXT(".ugeom"))
 			, bIsCacheDefined(!InCachePath.IsEmpty())
 			, FileDescription(InFileDescription)
 		{
@@ -35,7 +35,7 @@ namespace CADLibrary
 		{
 			if (!SceneFileHash)
 			{
-				SceneFileHash = HashCombine(FileDescription.GetDescriptorHash(), ::GetTypeHash(ImportParameters.StitchingTechnique));
+				SceneFileHash = HashCombine(FileDescription.GetDescriptorHash(), ::GetTypeHash(ImportParameters.GetStitchingTechnique()));
 				SceneFileHash = HashCombine(SceneFileHash, ::GetTypeHash(CADImportSdk));
 			}
 			return SceneFileHash;
@@ -46,13 +46,7 @@ namespace CADLibrary
 			if (!GeomFileHash)
 			{
 				GeomFileHash = GetSceneFileHash();
-				GeomFileHash = HashCombine(GeomFileHash, ::GetTypeHash(ImportParameters.bDisableCADKernelTessellation));
-				GeomFileHash = HashCombine(GeomFileHash, ::GetTypeHash(ImportParameters.ChordTolerance));
-				GeomFileHash = HashCombine(GeomFileHash, ::GetTypeHash(ImportParameters.MaxEdgeLength));
-				GeomFileHash = HashCombine(GeomFileHash, ::GetTypeHash(ImportParameters.MaxNormalAngle));
-				GeomFileHash = HashCombine(GeomFileHash, ::GetTypeHash(ImportParameters.MetricUnit));
-				GeomFileHash = HashCombine(GeomFileHash, ::GetTypeHash(ImportParameters.ScaleFactor));
-				GeomFileHash = HashCombine(GeomFileHash, ::GetTypeHash(ImportParameters.StitchingTechnique));
+				GeomFileHash = HashCombine(GeomFileHash, GetTypeHash(ImportParameters));
 			}
 			return GeomFileHash;
 		}
@@ -215,11 +209,26 @@ namespace CADLibrary
 			return SceneGraphArchive.ExternalReferences.Emplace_GetRef(InFileDescription);
 		}
 
-		FBodyMesh& AddBodyMesh(FCadId BodyId)
+		/** return a unique value that will be used to define the static mesh name */
+		uint32 GetStaticMeshHash(const int32 BodyId)
+		{
+			return HashCombine(GetSceneFileHash(), ::GetTypeHash(BodyId));
+		}
+
+		FBodyMesh& AddBodyMesh(FCadId BodyId, FArchiveBody& Body)
 		{
-			return BodyMeshes.Emplace_GetRef(BodyId);
+			FBodyMesh& BodyMesh = BodyMeshes.Emplace_GetRef(BodyId);
+			BodyMesh.MeshActorName = GetStaticMeshHash(BodyId);
+			Body.MeshActorName = BodyMesh.MeshActorName;
+			return BodyMesh;
 		}
 
+		FBodyMesh& AddBodyMesh(FCadId BodyId)
+		{
+			FBodyMesh& BodyMesh = BodyMeshes.Emplace_GetRef(BodyId);
+			BodyMesh.MeshActorName = GetStaticMeshHash(BodyId);
+			return BodyMesh;
+		}
 
 		void ExportMeshArchiveFile()
 		{
@@ -307,41 +316,6 @@ namespace CADLibrary
 			return ImportParameters;
 		}
 
-		double MetricUnit() const
-		{
-			return ImportParameters.MetricUnit;
-		}
-
-		const EStitchingTechnique& GetStitchingTechnique() const
-		{
-			return ImportParameters.StitchingTechnique;
-		}
-
-		bool IsCADKernelTessellation() const
-		{
-			return ImportParameters.bDisableCADKernelTessellation;
-		}
-
-		bool IsSequentialImport() const
-		{
-			return ImportParameters.bEnableSequentialImport;
-		}
-
-		bool NeedUVMapScaling() const
-		{
-			return ImportParameters.bScaleUVMap;
-		}
-
-		double GetScaleFactor() const
-		{
-			return ImportParameters.ScaleFactor;
-		}
-
-		double GetMetricUnit() const
-		{
-			return ImportParameters.MetricUnit;
-		}
-
 	private:
 		ECADImportSdk CADImportSdk = ECADImportSdk::KernelIO;
 		const FImportParameters ImportParameters;
diff --git a/Engine/Plugins/Enterprise/DatasmithCADImporter/Source/CADInterfaces/Public/CTSession.h b/Engine/Plugins/Enterprise/DatasmithCADImporter/Source/CADInterfaces/Public/CTSession.h
--- a/Engine/Plugins/Enterprise/DatasmithCADImporter/Source/CADInterfaces/Public/CTSession.h
+++ b/Engine/Plugins/Enterprise/DatasmithCADImporter/Source/CADInterfaces/Public/CTSession.h
@@ -25,7 +25,6 @@ public:
 		: FCoreTechSessionBase(InOwner)
 		, ImportParams(InImportParameters)
 	{
-		ImportParams.bDisableCADKernelTessellation = true;
 	}
 
 	void ClearData();
@@ -39,28 +38,20 @@ public:
 	 */
 	bool TopoFixes(double SewingToleranceFactor = 100);
 
-	/**
-	 * @param InScaleFactor : use to scale meshing from Kernel-IO
-	 */
-	void SetScaleFactor(double InScaleFactor)
-	{
-		ImportParams.ScaleFactor = InScaleFactor;
-	}
-
 	double GetScaleFactor() const
 	{
-		return ImportParams.ScaleFactor;
+		return ImportParams.GetScaleFactor();
 	}
 
 	double GetSceneUnit() const
 	{
-		return ImportParams.MetricUnit;
+		return ImportParams.GetMetricUnit();
 	}
 
 	/**
 	 * Handle input file unit
 	 * @param FileMetricUnit: number of meters per file unit.
-	 * eg. For a file in inches, arg should be 0.0254
+	 * e.g. For a file in inches, arg should be 0.0254
 	 */
 	void SetSceneUnit(double InMetricUnit);
 
@@ -72,11 +63,11 @@ public:
 	 * @param NormalTolerance : Angle between two adjacent triangles
 	 * @param StitchingTechnique : CAD topology correction technique
 	 */
-	void SetImportParameters(double ChordTolerance, double MaxEdgeLength, double NormalTolerance, CADLibrary::EStitchingTechnique StitchingTechnique, bool bScaleUVMap);
+	void SetImportParameters(double ChordTolerance, double MaxEdgeLength, double NormalTolerance, CADLibrary::EStitchingTechnique StitchingTechnique);
 	
 	void SetModelCoordinateSystem(FDatasmithUtils::EModelCoordSystem NewCoordinateSystem)
 	{
-		ImportParams.ModelCoordSys = NewCoordinateSystem;
+		ImportParams.SetModelCoordinateSystem(NewCoordinateSystem);
 	}
 
 	const CADLibrary::FImportParameters& GetImportParameters() const
diff --git a/Engine/Plugins/Enterprise/DatasmithCADImporter/Source/CADKernelSurface/Private/CADKernelSurfaceExtension.cpp b/Engine/Plugins/Enterprise/DatasmithCADImporter/Source/CADKernelSurface/Private/CADKernelSurfaceExtension.cpp
--- a/Engine/Plugins/Enterprise/DatasmithCADImporter/Source/CADKernelSurface/Private/CADKernelSurfaceExtension.cpp
+++ b/Engine/Plugins/Enterprise/DatasmithCADImporter/Source/CADKernelSurface/Private/CADKernelSurfaceExtension.cpp
@@ -34,14 +34,9 @@ bool UCADKernelParametricSurfaceData::Tessellate(UStaticMesh& StaticMesh, const
 	bool bSuccessfulTessellation = false;
 
 #if WITH_EDITOR
-	CADLibrary::FImportParameters ImportParameters;
-	ImportParameters.MetricUnit = SceneParameters.MetricUnit;
-	ImportParameters.ScaleFactor = SceneParameters.ScaleFactor;
-	ImportParameters.ChordTolerance = RetessellateOptions.ChordTolerance;
-	ImportParameters.MaxEdgeLength = RetessellateOptions.MaxEdgeLength;
-	ImportParameters.MaxNormalAngle = RetessellateOptions.NormalTolerance;
-	ImportParameters.ModelCoordSys = static_cast<FDatasmithUtils::EModelCoordSystem>(SceneParameters.ModelCoordSys);
-	ImportParameters.StitchingTechnique = CADLibrary::EStitchingTechnique(RetessellateOptions.StitchingTechnique);
+	CADLibrary::FImportParameters ImportParameters((double) SceneParameters.MetricUnit, (double) SceneParameters.ScaleFactor);
+	ImportParameters.SetModelCoordinateSystem((FDatasmithUtils::EModelCoordSystem) SceneParameters.ModelCoordSys);
+	ImportParameters.SetTesselationParameters(RetessellateOptions.ChordTolerance, RetessellateOptions.MaxEdgeLength, RetessellateOptions.NormalTolerance, (CADLibrary::EStitchingTechnique) RetessellateOptions.StitchingTechnique);
 
 	CADLibrary::FMeshParameters CadMeshParameters;
 	CadMeshParameters.bNeedSwapOrientation = MeshParameters.bNeedSwapOrientation;
@@ -61,7 +56,7 @@ bool UCADKernelParametricSurfaceData::Tessellate(UStaticMesh& StaticMesh, const
 			CADLibrary::CopyPatchGroups(*DestinationMeshDescription, MeshDescription);
 		}
 
-		TSharedRef<CADKernel::FSession> CADKernelSession = MakeShared<CADKernel::FSession>(0.00001 / ImportParameters.MetricUnit);
+		TSharedRef<CADKernel::FSession> CADKernelSession = MakeShared<CADKernel::FSession>(0.00001 / ImportParameters.GetMetricUnit());
 		CADKernelSession->AddDatabase(RawData);
 
 		TSharedRef<CADKernel::FModel> CADKernelModel = CADKernelSession->GetModel();
@@ -110,9 +105,9 @@ namespace CADKernelSurface
 			{
 				UCADKernelParametricSurfaceData* CADKernelData = Datasmith::MakeAdditionalData<UCADKernelParametricSurfaceData>();
 				CADKernelData->RawData = MoveTemp(ByteArray);
-				CADKernelData->SceneParameters.ModelCoordSys = uint8(InSceneParameters.ModelCoordSys);
-				CADKernelData->SceneParameters.MetricUnit = InSceneParameters.MetricUnit;
-				CADKernelData->SceneParameters.ScaleFactor = InSceneParameters.ScaleFactor;
+				CADKernelData->SceneParameters.ModelCoordSys = uint8(InSceneParameters.GetModelCoordSys());
+				CADKernelData->SceneParameters.MetricUnit = InSceneParameters.GetMetricUnit();
+				CADKernelData->SceneParameters.ScaleFactor = InSceneParameters.GetScaleFactor();
 
 				CADKernelData->MeshParameters.bNeedSwapOrientation = InMeshParameters.bNeedSwapOrientation;
 				CADKernelData->MeshParameters.bIsSymmetric = InMeshParameters.bIsSymmetric;
diff --git a/Engine/Plugins/Enterprise/DatasmithCADImporter/Source/CADKernelSurface/Public/CADKernelSurfaceExtension.h b/Engine/Plugins/Enterprise/DatasmithCADImporter/Source/CADKernelSurface/Public/CADKernelSurfaceExtension.h
--- a/Engine/Plugins/Enterprise/DatasmithCADImporter/Source/CADKernelSurface/Public/CADKernelSurfaceExtension.h
+++ b/Engine/Plugins/Enterprise/DatasmithCADImporter/Source/CADKernelSurface/Public/CADKernelSurfaceExtension.h
@@ -37,7 +37,7 @@ struct FDatasmithMeshElementPayload;
 
 namespace CADLibrary
 {
-	struct FImportParameters;
+	class FImportParameters;
 	struct FMeshParameters;
 }
 
diff --git a/Engine/Plugins/Enterprise/DatasmithCADImporter/Source/CADKernelSurface/Public/CADModelToCADKernelConverterBase.h b/Engine/Plugins/Enterprise/DatasmithCADImporter/Source/CADKernelSurface/Public/CADModelToCADKernelConverterBase.h
--- a/Engine/Plugins/Enterprise/DatasmithCADImporter/Source/CADKernelSurface/Public/CADModelToCADKernelConverterBase.h
+++ b/Engine/Plugins/Enterprise/DatasmithCADImporter/Source/CADKernelSurface/Public/CADModelToCADKernelConverterBase.h
@@ -21,12 +21,11 @@ class FCADModelToCADKernelConverterBase : public CADLibrary::ICADModelConverter
 public:
 
 	FCADModelToCADKernelConverterBase(CADLibrary::FImportParameters InImportParameters)
-		: CADKernelSession(0.00001 / InImportParameters.MetricUnit)
+		: CADKernelSession(0.00001 / InImportParameters.GetMetricUnit())
 		, ImportParameters(InImportParameters)
-		, GeometricTolerance(0.00001 / InImportParameters.MetricUnit)
+		, GeometricTolerance(0.00001 / InImportParameters.GetMetricUnit())
 		, SquareTolerance(GeometricTolerance* GeometricTolerance)
 	{
-		ImportParameters.bDisableCADKernelTessellation = false;
 	}
 
 	virtual void InitializeProcess(double InMetricUnit) override
@@ -54,13 +53,9 @@ public:
 		return CADLibrary::FCADKernelTools::Tessellate(CADKernelEntity, ImportParameters, InMeshParameters, OutMeshDescription);
 	}
 
-	virtual void SetImportParameters(double ChordTolerance, double MaxEdgeLength, double NormalTolerance, CADLibrary::EStitchingTechnique StitchingTechnique, bool bScaleUVMap) override
+	virtual void SetImportParameters(double ChordTolerance, double MaxEdgeLength, double NormalTolerance, CADLibrary::EStitchingTechnique StitchingTechnique) override
 	{
-		ImportParameters.ChordTolerance = ChordTolerance;
-		ImportParameters.MaxEdgeLength = MaxEdgeLength;
-		ImportParameters.MaxNormalAngle = NormalTolerance;
-		ImportParameters.StitchingTechnique = StitchingTechnique;
-		ImportParameters.bScaleUVMap = bScaleUVMap;
+		ImportParameters.SetTesselationParameters(ChordTolerance, MaxEdgeLength, NormalTolerance, StitchingTechnique);
 	}
 
 	virtual void SetMetricUnit(double NewMetricUnit) override
@@ -70,12 +65,12 @@ public:
 
 	virtual double GetScaleFactor() const override
 	{
-		return ImportParameters.ScaleFactor;
+		return ImportParameters.GetScaleFactor();
 	}
 
 	virtual double GetMetricUnit() const override
 	{
-		return ImportParameters.MetricUnit;
+		return ImportParameters.GetMetricUnit();
 	}
 
 	virtual bool IsSessionValid() override
diff --git a/Engine/Plugins/Enterprise/DatasmithCADImporter/Source/CADLibrary/Private/CADKernelTools.cpp b/Engine/Plugins/Enterprise/DatasmithCADImporter/Source/CADLibrary/Private/CADKernelTools.cpp
--- a/Engine/Plugins/Enterprise/DatasmithCADImporter/Source/CADLibrary/Private/CADKernelTools.cpp
+++ b/Engine/Plugins/Enterprise/DatasmithCADImporter/Source/CADLibrary/Private/CADKernelTools.cpp
@@ -54,7 +54,7 @@ namespace CADLibrary
 
 		for (FVector& Vertex : VertexArray)
 		{
-			Vertex *= Context.ImportParams.ScaleFactor;
+			Vertex *= Context.ImportParams.GetScaleFactor();
 		}
 
 		int32 VertexCount = VertexArray.Num();
@@ -71,7 +71,7 @@ namespace CADLibrary
 			VertexIndex++;
 
 			FVertexID VertexID = MeshDescription.CreateVertex();
-			VertexPositions[VertexID] = FDatasmithUtils::ConvertVector((FDatasmithUtils::EModelCoordSystem) Context.ImportParams.ModelCoordSys, Vertex);
+			VertexPositions[VertexID] = FDatasmithUtils::ConvertVector((FDatasmithUtils::EModelCoordSystem) Context.ImportParams.GetModelCoordSys(), Vertex);
 			Context.VertexIds[VertexIndex] = VertexID;
 		}
 
@@ -86,7 +86,7 @@ namespace CADLibrary
 			for (const FVector& Vertex : VertexArray)
 			{
 				FVertexID VertexID = MeshDescription.CreateVertex();
-				VertexPositions[VertexID] = FDatasmithUtils::ConvertVector((FDatasmithUtils::EModelCoordSystem) Context.ImportParams.ModelCoordSys, Vertex);
+				VertexPositions[VertexID] = FDatasmithUtils::ConvertVector((FDatasmithUtils::EModelCoordSystem) Context.ImportParams.GetModelCoordSys(), Vertex);
 				VertexPositions[VertexID] = SymmetricMatrix.TransformPosition(VertexPositions[VertexID]);
 				Context.SymmetricVertexIds[VertexIndex++] = VertexID;
 			}
@@ -221,7 +221,7 @@ namespace CADLibrary
 
 				if (!Step)
 				{
-					FDatasmithUtils::ConvertVectorArray(Context.ImportParams.ModelCoordSys, FaceMesh->Normals);
+					FDatasmithUtils::ConvertVectorArray(Context.ImportParams.GetModelCoordSys(), FaceMesh->Normals);
 					for (FVector& Normal : FaceMesh->Normals)
 					{
 						Normal = Normal.GetSafeNormal();
@@ -387,22 +387,22 @@ namespace CADLibrary
 			MeshModel->AddCriterion(CurvatureCriterion);
 		}
 
-		if (ImportParameters.MaxEdgeLength > SMALL_NUMBER)
+		if (ImportParameters.GetMaxEdgeLength() > SMALL_NUMBER)
 		{
 
-			TSharedPtr<CADKernel::FCriterion> MaxSizeCriterion = CADKernel::FCriterion::CreateCriterion(CADKernel::ECriterion::MaxSize, ImportParameters.MaxEdgeLength / ImportParameters.ScaleFactor);
+			TSharedPtr<CADKernel::FCriterion> MaxSizeCriterion = CADKernel::FCriterion::CreateCriterion(CADKernel::ECriterion::MaxSize, ImportParameters.GetMaxEdgeLength() / ImportParameters.GetScaleFactor());
 			MeshModel->AddCriterion(MaxSizeCriterion);
 		}
 
-		if (ImportParameters.ChordTolerance > SMALL_NUMBER)
+		if (ImportParameters.GetChordTolerance() > SMALL_NUMBER)
 		{
-			TSharedPtr<CADKernel::FCriterion> ChordCriterion = CADKernel::FCriterion::CreateCriterion(CADKernel::ECriterion::Sag, ImportParameters.ChordTolerance / ImportParameters.ScaleFactor);
+			TSharedPtr<CADKernel::FCriterion> ChordCriterion = CADKernel::FCriterion::CreateCriterion(CADKernel::ECriterion::Sag, ImportParameters.GetChordTolerance() / ImportParameters.GetScaleFactor());
 			MeshModel->AddCriterion(ChordCriterion);
 		}
 
-		if (ImportParameters.MaxNormalAngle > SMALL_NUMBER)
+		if (ImportParameters.GetMaxNormalAngle() > SMALL_NUMBER)
 		{
-			TSharedPtr<CADKernel::FCriterion> MaxNormalAngleCriterion = CADKernel::FCriterion::CreateCriterion(CADKernel::ECriterion::Angle, ImportParameters.MaxNormalAngle);
+			TSharedPtr<CADKernel::FCriterion> MaxNormalAngleCriterion = CADKernel::FCriterion::CreateCriterion(CADKernel::ECriterion::Angle, ImportParameters.GetMaxNormalAngle());
 			MeshModel->AddCriterion(MaxNormalAngleCriterion);
 		}
 	}
diff --git a/Engine/Plugins/Enterprise/DatasmithCADImporter/Source/CADLibrary/Private/MeshDescriptionHelper.cpp b/Engine/Plugins/Enterprise/DatasmithCADImporter/Source/CADLibrary/Private/MeshDescriptionHelper.cpp
--- a/Engine/Plugins/Enterprise/DatasmithCADImporter/Source/CADLibrary/Private/MeshDescriptionHelper.cpp
+++ b/Engine/Plugins/Enterprise/DatasmithCADImporter/Source/CADLibrary/Private/MeshDescriptionHelper.cpp
@@ -80,7 +80,7 @@ namespace CADLibrary
 		{
 			for (const FVector& Vertex : CTTessellation.PositionArray)
 			{
-				VertexDataSet.Emplace(GlobalVertexCount, Vertex * ImportParams.ScaleFactor);
+				VertexDataSet.Emplace(GlobalVertexCount, Vertex * ImportParams.GetScaleFactor());
 				++GlobalVertexCount;
 			}
 		}
@@ -160,7 +160,7 @@ namespace CADLibrary
 
 			FVertexID VertexID = MeshDescription.CreateVertex();
 			Body.VertexIds.Add(VertexID);
-			VertexPositions[VertexID] = FDatasmithUtils::ConvertVector((FDatasmithUtils::EModelCoordSystem) ImportParams.ModelCoordSys, VertexDataSet[VertexIndex].Coordinates);
+			VertexPositions[VertexID] = FDatasmithUtils::ConvertVector((FDatasmithUtils::EModelCoordSystem) ImportParams.GetModelCoordSys(), VertexDataSet[VertexIndex].Coordinates);
 			VertexDataSet[VertexIndex].VertexID = VertexID;
 		}
 
@@ -184,7 +184,7 @@ namespace CADLibrary
 
 				FVertexID VertexID = MeshDescription.CreateVertex();
 				Body.SymmetricVertexIds.Add(VertexID);
-				VertexPositions[VertexID] = FDatasmithUtils::ConvertVector((FDatasmithUtils::EModelCoordSystem) ImportParams.ModelCoordSys, VertexDataSet[VertexIndex].Coordinates);
+				VertexPositions[VertexID] = FDatasmithUtils::ConvertVector((FDatasmithUtils::EModelCoordSystem) ImportParams.GetModelCoordSys(), VertexDataSet[VertexIndex].Coordinates);
 				VertexPositions[VertexID] = SymmetricMatrix.TransformPosition(VertexPositions[VertexID]);
 
 				VertexDataSet[VertexIndex].SymVertexID = VertexID;
@@ -211,39 +211,19 @@ namespace CADLibrary
 		int32 TriangleCount = Body.TriangleCount;
 		TArray<FTessellationData>& FaceTessellationSet = Body.Faces;
 
-		// Add offset to the bounding box to avoid to remove good vertex
-		FVector Size = Body.BBox.GetSize();
-		FBox BBox = Body.BBox.ExpandBy(Size.Size());
-		BBox.IsValid = Body.BBox.IsValid;
-		BBox.Min *= ImportParams.ScaleFactor;
-		BBox.Max *= ImportParams.ScaleFactor;
-
 		TVertexAttributesRef<FVector3f> VertexPositions = MeshDescription.GetVertexPositions();
 
 		TArray<FVector>& VertexArray = Body.VertexArray;
 
 		for (FVector& Vertex : VertexArray)
 		{
-			Vertex *= ImportParams.ScaleFactor;
-			}
+			Vertex *= ImportParams.GetScaleFactor();
+		}
 
 		TArray<int32>& VertexIdSet = Body.VertexIds;
 		int32 VertexCount = VertexArray.Num();
 		VertexIdSet.SetNumZeroed(VertexCount);
 
-		if (BBox.IsValid)
-		{
-			int32 VertexIndex = 0;
-			for (const FVector& Vertex : VertexArray)
-			{
-				if (!BBox.IsInside(Vertex))
-				{
-					VertexIdSet[VertexIndex] = INDEX_NONE;
-				}
-				VertexIndex++;
-			}
-		}
-
 		// Make MeshDescription.VertexPositions and VertexID
 		MeshDescription.ReserveNewVertices(MeshParameters.bIsSymmetric ? VertexCount * 2 : VertexCount);
 
@@ -259,7 +239,7 @@ namespace CADLibrary
 			}
 
 			FVertexID VertexID = MeshDescription.CreateVertex();
-			VertexPositions[VertexID] = FDatasmithUtils::ConvertVector((FDatasmithUtils::EModelCoordSystem) ImportParams.ModelCoordSys, Vertex);
+			VertexPositions[VertexID] = FDatasmithUtils::ConvertVector((FDatasmithUtils::EModelCoordSystem) ImportParams.GetModelCoordSys(), Vertex);
 			VertexIdSet[VertexIndex] = VertexID;
 		}
 
@@ -281,7 +261,7 @@ namespace CADLibrary
 				}
 
 				FVertexID VertexID = MeshDescription.CreateVertex();
-				VertexPositions[VertexID] = FDatasmithUtils::ConvertVector((FDatasmithUtils::EModelCoordSystem) ImportParams.ModelCoordSys, Vertex);
+				VertexPositions[VertexID] = FDatasmithUtils::ConvertVector((FDatasmithUtils::EModelCoordSystem) ImportParams.GetModelCoordSys(), Vertex);
 				VertexPositions[VertexID] = SymmetricMatrix.TransformPosition(VertexPositions[VertexID]);
 				SymmetricVertexIds[VertexIndex++] = VertexID;
 			}
@@ -406,7 +386,7 @@ namespace CADLibrary
 
 		if (Algo::AnyOf(BodyTessellation.Faces, [](const FTessellationData& FaceTessellation) { return !FaceTessellation.TexCoordArray.IsEmpty(); }))
 		{
-		VertexInstanceUVs.SetNumChannels(1);
+			VertexInstanceUVs.SetNumChannels(1);
 		}
 
 		int32 NbStep = 1;
@@ -509,7 +489,7 @@ namespace CADLibrary
 
 				if (!Step)
 				{
-					FDatasmithUtils::ConvertVectorArray(ImportParams.ModelCoordSys, Tessellation.NormalArray);
+					FDatasmithUtils::ConvertVectorArray(ImportParams.GetModelCoordSys(), Tessellation.NormalArray);
 					for (FVector& Normal : Tessellation.NormalArray)
 					{
 						Normal = Normal.GetSafeNormal();
diff --git a/Engine/Plugins/Enterprise/DatasmithCADImporter/Source/CADLibrary/Public/CADKernelTools.h b/Engine/Plugins/Enterprise/DatasmithCADImporter/Source/CADLibrary/Public/CADKernelTools.h
--- a/Engine/Plugins/Enterprise/DatasmithCADImporter/Source/CADLibrary/Public/CADKernelTools.h
+++ b/Engine/Plugins/Enterprise/DatasmithCADImporter/Source/CADLibrary/Public/CADKernelTools.h
@@ -19,7 +19,7 @@ struct FMeshDescription;
 
 namespace CADLibrary
 {
-	struct FImportParameters;
+	class FImportParameters;
 	struct FMeshParameters;
 
 	class CADLIBRARY_API FCADKernelTools
diff --git a/Engine/Plugins/Enterprise/DatasmithCADImporter/Source/CADTools/Private/CADOption.cpp b/Engine/Plugins/Enterprise/DatasmithCADImporter/Source/CADTools/Private/CADOption.cpp
--- a/Engine/Plugins/Enterprise/DatasmithCADImporter/Source/CADTools/Private/CADOption.cpp
+++ b/Engine/Plugins/Enterprise/DatasmithCADImporter/Source/CADTools/Private/CADOption.cpp
@@ -11,13 +11,6 @@ FAutoConsoleVariableRef GCADTranslatorLibrary(
 	TEXT("Define the used library to import CAD file i.e. \"TechSoft\" or \"KernelIO\".\n"),
 	ECVF_Default);
 
-bool bGDisableCADKernelTessellation = true;
-FAutoConsoleVariableRef GCADTranslatorDisableCADKernelTessellation(
-	TEXT("ds.CADTranslator.DisableCADKernelTessellation"),
-	bGDisableCADKernelTessellation,
-	TEXT("Disable to use CAD import library tessellator.\n"),
-	ECVF_Default);
-
 int32 GMaxImportThreads = 0;
 FAutoConsoleVariableRef GCADTranslatorMaxImportThreads(
 	TEXT("ds.CADTranslator.MaxImportThreads"),
@@ -30,28 +23,47 @@ Default is MaxImportThreads = 0\n\
    -if CADCache is disable, the scene is read all at once\n"),
 	ECVF_Default);
 
-bool bGEnableCADCache = true;
+bool FImportParameters::bGDisableCADKernelTessellation = true;
+FAutoConsoleVariableRef GCADTranslatorDisableCADKernelTessellation(
+	TEXT("ds.CADTranslator.DisableCADKernelTessellation"),
+	FImportParameters::bGDisableCADKernelTessellation,
+	TEXT("Disable to use CAD import library tessellator.\n"),
+	ECVF_Default);
+
+bool FImportParameters::bGEnableCADCache = true;
 FAutoConsoleVariableRef GCADTranslatorEnableCADCache(
 	TEXT("ds.CADTranslator.EnableCADCache"),
-	bGEnableCADCache,
+	FImportParameters::bGEnableCADCache,
 	TEXT("\
 Enable/disable temporary CAD processing file cache. These file will be use in a next import to avoid CAD file processing.\n\
 If MaxImportThreads != 1, EnableCADCache value is ignored\n\
 Default is enable\n"),
 ECVF_Default);
 
-bool bGOverwriteCache = false;
+bool FImportParameters::bGOverwriteCache = false;
 FAutoConsoleVariableRef GCADTranslatorOverwriteCache(
 	TEXT("ds.CADTranslator.OverwriteCache"),
-	bGOverwriteCache,
+	FImportParameters::bGOverwriteCache,
 	TEXT("Overwrite any existing cache associated with the file being imported.\n"),
 	ECVF_Default);
 
-bool bGEnableTimeControl = true;
+bool FImportParameters::bGEnableTimeControl = true;
 FAutoConsoleVariableRef GCADTranslatorEnableTimeControl(
 	TEXT("ds.CADTranslator.EnableTimeControl"),
-	bGEnableTimeControl,
+	FImportParameters::bGEnableTimeControl,
 	TEXT("Enable the timer that kill the worker if the import time is unusually long. With this time control, the load of the corrupted file is canceled but the rest of the scene is imported.\n"),
 	ECVF_Default);
 
+uint32 GetTypeHash(const FImportParameters& ImportParameters)
+{
+	uint32 ParametersHash = ::GetTypeHash(ImportParameters.bGDisableCADKernelTessellation);
+	ParametersHash = HashCombine(ParametersHash, ::GetTypeHash(ImportParameters.ChordTolerance));
+	ParametersHash = HashCombine(ParametersHash, ::GetTypeHash(ImportParameters.MaxEdgeLength));
+	ParametersHash = HashCombine(ParametersHash, ::GetTypeHash(ImportParameters.MaxNormalAngle));
+	ParametersHash = HashCombine(ParametersHash, ::GetTypeHash(ImportParameters.MetricUnit));
+	ParametersHash = HashCombine(ParametersHash, ::GetTypeHash(ImportParameters.ScaleFactor));
+	ParametersHash = HashCombine(ParametersHash, ::GetTypeHash(ImportParameters.StitchingTechnique));
+	return ParametersHash;
+}
+
 }
diff --git a/Engine/Plugins/Enterprise/DatasmithCADImporter/Source/CADTools/Public/CADOptions.h b/Engine/Plugins/Enterprise/DatasmithCADImporter/Source/CADTools/Public/CADOptions.h
--- a/Engine/Plugins/Enterprise/DatasmithCADImporter/Source/CADTools/Public/CADOptions.h
+++ b/Engine/Plugins/Enterprise/DatasmithCADImporter/Source/CADTools/Public/CADOptions.h
@@ -10,10 +10,6 @@
 
 namespace CADLibrary
 {
-	CADTOOLS_API extern bool bGDisableCADKernelTessellation;
-	CADTOOLS_API extern bool bGEnableCADCache;
-	CADTOOLS_API extern bool bGEnableTimeControl;
-	CADTOOLS_API extern bool bGOverwriteCache;
 	CADTOOLS_API extern int32 GMaxImportThreads;
 	CADTOOLS_API extern FString GCADLibrary;
 
@@ -39,22 +35,48 @@ namespace CADLibrary
 		BodyOnly,
 	};
 
-	struct FImportParameters
+	class FImportParameters
 	{
-		double MetricUnit = 0.001;
-		double ScaleFactor = 0.1;
-		double ChordTolerance = 0.2;
-		double MaxEdgeLength = 0.0;
-		double MaxNormalAngle = 20.0;
-		EStitchingTechnique StitchingTechnique = EStitchingTechnique::StitchingNone;
-		FDatasmithUtils::EModelCoordSystem ModelCoordSys = FDatasmithUtils::EModelCoordSystem::ZUp_RightHanded;
-		EDisplayDataPropagationMode Propagation = EDisplayDataPropagationMode::TopDown;
-		EDisplayPreference DisplayPreference = EDisplayPreference::MaterialPrefered;
-		bool bScaleUVMap = true;
-		bool bEnableSequentialImport = true;
-		bool bOverwriteCache = false;
-		bool bDisableCADKernelTessellation = false;
-		bool bEnableTimeControl = true;
+	private:
+		double MetricUnit;
+		double ScaleFactor;
+		double ChordTolerance;
+		double MaxEdgeLength;
+		double MaxNormalAngle;
+		EStitchingTechnique StitchingTechnique;
+		FDatasmithUtils::EModelCoordSystem ModelCoordSys;
+		EDisplayDataPropagationMode Propagation;
+		EDisplayPreference DisplayPreference;
+		bool bScaleUVMap;
+		
+	public:
+		CADTOOLS_API static bool bGDisableCADKernelTessellation;
+		CADTOOLS_API static bool bGEnableTimeControl;
+		CADTOOLS_API static bool bGEnableCADCache;
+		CADTOOLS_API static bool bGOverwriteCache;
+
+	public:
+		FImportParameters(double InMetricUnit = 0.001, double InScaleFactor = 0.1, FDatasmithUtils::EModelCoordSystem NewCoordinateSystem = FDatasmithUtils::EModelCoordSystem::ZUp_RightHanded)
+			: MetricUnit(InMetricUnit)
+			, ScaleFactor(InScaleFactor)
+			, ChordTolerance(0.2)
+			, MaxEdgeLength(0.0)
+			, MaxNormalAngle(20.0)
+			, StitchingTechnique(EStitchingTechnique::StitchingNone)
+			, ModelCoordSys(NewCoordinateSystem)
+			, Propagation(EDisplayDataPropagationMode::TopDown)
+			, DisplayPreference(EDisplayPreference::MaterialPrefered)
+			, bScaleUVMap(true)
+		{
+		}
+	
+		void SetTesselationParameters(double InChordTolerance, double InMaxEdgeLength, double InMaxNormalAngle, CADLibrary::EStitchingTechnique InStitchingTechnique)
+		{
+			ChordTolerance = InChordTolerance;
+			MaxEdgeLength = InMaxEdgeLength;
+			MaxNormalAngle = InMaxNormalAngle;
+			StitchingTechnique = InStitchingTechnique;
+		}
 
 		void SetMetricUnit(double NewMetricUnit)
 		{
@@ -64,7 +86,7 @@ namespace CADLibrary
 
 		uint32 GetHash() const
 		{
-			uint32 Hash = 0; 
+			uint32 Hash = 0;
 			for (double Param : {MetricUnit, ScaleFactor, ChordTolerance, MaxEdgeLength, MaxNormalAngle})
 			{
 				Hash = HashCombine(Hash, ::GetTypeHash(Param));
@@ -89,16 +111,19 @@ namespace CADLibrary
 			Ar << (uint8&) ImportParameters.Propagation;
 			Ar << (uint8&) ImportParameters.DisplayPreference;
 			Ar << ImportParameters.bScaleUVMap;
-			Ar << ImportParameters.bEnableSequentialImport;
-			Ar << ImportParameters.bOverwriteCache;
-			Ar << ImportParameters.bDisableCADKernelTessellation;
+			// these 4 static variables have to be serialized to be transmitted to CADWorkers
+			// because CADWorker has not access to CVars
+			Ar << ImportParameters.bGOverwriteCache;
+			Ar << ImportParameters.bGDisableCADKernelTessellation;
+			Ar << ImportParameters.bGEnableTimeControl;
+			Ar << ImportParameters.bGEnableCADCache;
 			return Ar;
 		}
 
 		FString DefineCADFilePath(const TCHAR* FolderPath, const TCHAR* InFileName) const
 		{
 			FString OutFileName = FPaths::Combine(FolderPath, InFileName);
-			if (bDisableCADKernelTessellation)
+			if (bGDisableCADKernelTessellation)
 			{
 				OutFileName += TEXT(".ct");
 			}
@@ -108,6 +133,79 @@ namespace CADLibrary
 			}
 			return OutFileName;
 		}
+
+		double GetMetricUnit() const
+		{
+			return MetricUnit;
+		}
+
+		double GetScaleFactor() const
+		{
+			return ScaleFactor;
+		}
+
+		double GetChordTolerance() const
+		{
+			return ChordTolerance;
+		}
+
+		double GetMaxNormalAngle() const
+		{
+			return MaxNormalAngle;
+		}
+
+		double GetMaxEdgeLength() const
+		{
+			return MaxEdgeLength;
+		}
+
+		EStitchingTechnique GetStitchingTechnique() const
+		{
+			return StitchingTechnique;
+		}
+
+		FDatasmithUtils::EModelCoordSystem GetModelCoordSys() const
+		{
+			return ModelCoordSys;
+		}
+
+		EDisplayDataPropagationMode GetPropagation() const
+		{
+			return Propagation;
+		}
+
+		EDisplayPreference GetDisplayPreference() const
+		{
+			return DisplayPreference;
+		}
+
+		bool NeedScaleUVMap() const
+		{	
+			return bScaleUVMap;
+		}
+
+		void SwitchOffUVMapScaling()
+		{
+			bScaleUVMap = false;
+		}
+
+		void SetModelCoordinateSystem(FDatasmithUtils::EModelCoordSystem NewCoordinateSystem)
+		{
+			ModelCoordSys = NewCoordinateSystem;
+		}
+
+		void SetDisplayPreference(CADLibrary::EDisplayPreference InDisplayPreference)
+		{
+			DisplayPreference = InDisplayPreference;
+		}
+
+		void SetPropagationMode(CADLibrary::EDisplayDataPropagationMode InMode)
+		{
+			Propagation = InMode;
+		}
+
+		CADTOOLS_API friend uint32 GetTypeHash(const FImportParameters& ImportParameters);
+
 	};
 
 	struct FMeshParameters
diff --git a/Engine/Plugins/Enterprise/DatasmithCADImporter/Source/CoreTechSurface/Private/CoreTechSurfaceExtension.cpp b/Engine/Plugins/Enterprise/DatasmithCADImporter/Source/CoreTechSurface/Private/CoreTechSurfaceExtension.cpp
--- a/Engine/Plugins/Enterprise/DatasmithCADImporter/Source/CoreTechSurface/Private/CoreTechSurfaceExtension.cpp
+++ b/Engine/Plugins/Enterprise/DatasmithCADImporter/Source/CoreTechSurface/Private/CoreTechSurfaceExtension.cpp
@@ -46,14 +46,8 @@ bool UCoreTechParametricSurfaceData::Tessellate(UStaticMesh& StaticMesh, const F
 
 	CADLibrary::FCTMesh Mesh;
 
-	CADLibrary::FImportParameters ImportParameters;
-	ImportParameters.MetricUnit = SceneParameters.MetricUnit;
-	ImportParameters.ScaleFactor = SceneParameters.ScaleFactor;
-	ImportParameters.ChordTolerance = RetessellateOptions.ChordTolerance;
-	ImportParameters.MaxEdgeLength = RetessellateOptions.MaxEdgeLength;
-	ImportParameters.MaxNormalAngle = RetessellateOptions.NormalTolerance;
-	ImportParameters.ModelCoordSys = static_cast<FDatasmithUtils::EModelCoordSystem>(SceneParameters.ModelCoordSys);
-	ImportParameters.StitchingTechnique = CADLibrary::EStitchingTechnique(RetessellateOptions.StitchingTechnique);
+	CADLibrary::FImportParameters ImportParameters(SceneParameters.MetricUnit, SceneParameters.ScaleFactor, (FDatasmithUtils::EModelCoordSystem) SceneParameters.ModelCoordSys);
+	ImportParameters.SetTesselationParameters(RetessellateOptions.ChordTolerance, RetessellateOptions.MaxEdgeLength, RetessellateOptions.NormalTolerance, (CADLibrary::EStitchingTechnique)RetessellateOptions.StitchingTechnique);
 
 	CADLibrary::FMeshParameters CadMeshParameters;
 	CadMeshParameters.bNeedSwapOrientation = MeshParameters.bNeedSwapOrientation;
diff --git a/Engine/Plugins/Enterprise/DatasmithCADImporter/Source/CoreTechSurface/Private/CoreTechSurfaceHelper.cpp b/Engine/Plugins/Enterprise/DatasmithCADImporter/Source/CoreTechSurface/Private/CoreTechSurfaceHelper.cpp
--- a/Engine/Plugins/Enterprise/DatasmithCADImporter/Source/CoreTechSurface/Private/CoreTechSurfaceHelper.cpp
+++ b/Engine/Plugins/Enterprise/DatasmithCADImporter/Source/CoreTechSurface/Private/CoreTechSurfaceHelper.cpp
@@ -47,7 +47,7 @@ namespace CoreTechSurface
 			return false;
 		}
 
-		CADLibrary::CTKIO_ChangeUnit(ImportParameters.MetricUnit);
+		CADLibrary::CTKIO_ChangeUnit(ImportParameters.GetMetricUnit());
 		uint64 MainObjectID;
 		if(!CADLibrary::CTKIO_LoadModel(*FileName, MainObjectID, 0x00020000 /* CT_LOAD_FLAGS_READ_META_DATA */))
 		{
@@ -55,7 +55,7 @@ namespace CoreTechSurface
 			return false;
 		}
 
-		if (ImportParameters.StitchingTechnique != CADLibrary::EStitchingTechnique::StitchingNone)
+		if (ImportParameters.GetStitchingTechnique() != CADLibrary::EStitchingTechnique::StitchingNone)
 		{
 			CADLibrary::CTKIO_Repair(MainObjectID, CADLibrary::EStitchingTechnique::StitchingSew);
 		}
@@ -75,9 +75,9 @@ namespace CoreTechSurface
 					UCoreTechParametricSurfaceData* CoreTechData = Datasmith::MakeAdditionalData<UCoreTechParametricSurfaceData>();
 					CoreTechData->SourceFile = InFilePath;
 					CoreTechData->RawData = MoveTemp(ByteArray);
-					CoreTechData->SceneParameters.ModelCoordSys = uint8(InSceneParameters.ModelCoordSys);
-					CoreTechData->SceneParameters.MetricUnit = InSceneParameters.MetricUnit;
-					CoreTechData->SceneParameters.ScaleFactor = InSceneParameters.ScaleFactor;
+					CoreTechData->SceneParameters.ModelCoordSys = uint8(InSceneParameters.GetModelCoordSys());
+					CoreTechData->SceneParameters.MetricUnit = InSceneParameters.GetMetricUnit();
+					CoreTechData->SceneParameters.ScaleFactor = InSceneParameters.GetScaleFactor();
 
 					CoreTechData->MeshParameters.bNeedSwapOrientation = InMeshParameters.bNeedSwapOrientation;
 					CoreTechData->MeshParameters.bIsSymmetric = InMeshParameters.bIsSymmetric;
diff --git a/Engine/Plugins/Enterprise/DatasmithCADImporter/Source/CoreTechSurface/Public/CADModelToCoretechConverterBase.h b/Engine/Plugins/Enterprise/DatasmithCADImporter/Source/CoreTechSurface/Public/CADModelToCoretechConverterBase.h
--- a/Engine/Plugins/Enterprise/DatasmithCADImporter/Source/CoreTechSurface/Public/CADModelToCoretechConverterBase.h
+++ b/Engine/Plugins/Enterprise/DatasmithCADImporter/Source/CoreTechSurface/Public/CADModelToCoretechConverterBase.h
@@ -53,9 +53,9 @@ public:
 		return false;
 	}
 
-	virtual void SetImportParameters(double ChordTolerance, double MaxEdgeLength, double NormalTolerance, CADLibrary::EStitchingTechnique StitchingTechnique, bool bScaleUVMap) override
+	virtual void SetImportParameters(double ChordTolerance, double MaxEdgeLength, double NormalTolerance, CADLibrary::EStitchingTechnique StitchingTechnique) override
 	{
-		FCTSession::SetImportParameters(ChordTolerance, MaxEdgeLength, NormalTolerance, StitchingTechnique, bScaleUVMap);
+		FCTSession::SetImportParameters(ChordTolerance, MaxEdgeLength, NormalTolerance, StitchingTechnique);
 	}
 
 	virtual void SetMetricUnit(double NewMetricUnit) override
@@ -65,12 +65,12 @@ public:
 
 	virtual double GetScaleFactor() const override
 	{
-		return ImportParams.ScaleFactor;
+		return ImportParams.GetScaleFactor();
 	}
 
 	virtual double GetMetricUnit() const override
 	{
-		return ImportParams.MetricUnit;
+		return ImportParams.GetMetricUnit();
 	}
 
 	virtual bool IsSessionValid() override
diff --git a/Engine/Plugins/Enterprise/DatasmithCADImporter/Source/CoreTechSurface/Public/CoreTechSurfaceHelper.h b/Engine/Plugins/Enterprise/DatasmithCADImporter/Source/CoreTechSurface/Public/CoreTechSurfaceHelper.h
--- a/Engine/Plugins/Enterprise/DatasmithCADImporter/Source/CoreTechSurface/Public/CoreTechSurfaceHelper.h
+++ b/Engine/Plugins/Enterprise/DatasmithCADImporter/Source/CoreTechSurface/Public/CoreTechSurfaceHelper.h
@@ -5,7 +5,7 @@
 
 namespace CADLibrary
 {
-	struct FImportParameters;
+	class FImportParameters;
 	struct FMeshParameters;
 }
 
diff --git a/Engine/Plugins/Enterprise/DatasmithCADImporter/Source/DatasmithCADTranslator/Private/DatasmithCADTranslator.cpp b/Engine/Plugins/Enterprise/DatasmithCADImporter/Source/DatasmithCADTranslator/Private/DatasmithCADTranslator.cpp
--- a/Engine/Plugins/Enterprise/DatasmithCADImporter/Source/DatasmithCADTranslator/Private/DatasmithCADTranslator.cpp
+++ b/Engine/Plugins/Enterprise/DatasmithCADImporter/Source/DatasmithCADTranslator/Private/DatasmithCADTranslator.cpp
@@ -85,70 +85,56 @@ bool FDatasmithCADTranslator::IsSourceSupported(const FDatasmithSceneSource& Sou
 
 bool FDatasmithCADTranslator::LoadScene(TSharedRef<IDatasmithScene> DatasmithScene)
 {
-	ImportParameters.bDisableCADKernelTessellation = CADLibrary::bGDisableCADKernelTessellation;
-
-	ImportParameters.MetricUnit = 0.001;
-	ImportParameters.ScaleFactor = 0.1;
 	const FDatasmithTessellationOptions& TesselationOptions = GetCommonTessellationOptions();
-	ImportParameters.ChordTolerance = TesselationOptions.ChordTolerance;
-	ImportParameters.MaxEdgeLength = TesselationOptions.MaxEdgeLength;
-	ImportParameters.MaxNormalAngle = TesselationOptions.NormalTolerance;
-	ImportParameters.StitchingTechnique = (CADLibrary::EStitchingTechnique)TesselationOptions.StitchingTechnique;
-
-	CADLibrary::FFileDescriptor FileDescription(*FPaths::ConvertRelativePathToFull(GetSource().GetSourceFile()));
+	CADLibrary::FFileDescriptor FileDescriptor(*FPaths::ConvertRelativePathToFull(GetSource().GetSourceFile()));
 
-	// Do not change the model unit when translator is called by the Datasmith runtime plugin.
-	ImportParameters.ModelCoordSys = FDatasmithUtils::EModelCoordSystem::ZUp_RightHanded;
+	ImportParameters.SetTesselationParameters(TesselationOptions.ChordTolerance, TesselationOptions.MaxEdgeLength, TesselationOptions.NormalTolerance, (CADLibrary::EStitchingTechnique)TesselationOptions.StitchingTechnique);
+	ImportParameters.SetModelCoordinateSystem(FDatasmithUtils::EModelCoordSystem::ZUp_RightHanded);
 
-	switch (FileDescription.GetFileFormat())
+	switch (FileDescriptor.GetFileFormat())
 	{
 	case CADLibrary::ECADFormat::JT:
 	{
 #if WITH_EDITOR
+		// Do not change the model unit when translator is called by the Datasmith runtime plugin.
 		if (IsInGameThread())
 		{
 			// In UE Editor, for historical QA reason (i.e. due to some validation files) the unit of JT file is set to meter as it seemed to be the default unit
 			// But setting the meter as the default unit requires resetting KernelIO that crash at the runtime...
 			// so for runtime use, metric unit is mm.
 			// this will change with the next release of CAD importer
-			ImportParameters.MetricUnit = 1.;
-			ImportParameters.ScaleFactor = 100.;
+			ImportParameters.SetMetricUnit(1.);
 		}
 #endif
 		break;
 	}
+
 	case CADLibrary::ECADFormat::NX:
 	{
-		ImportParameters.ModelCoordSys = FDatasmithUtils::EModelCoordSystem::ZUp_RightHanded;
-		ImportParameters.DisplayPreference = CADLibrary::EDisplayPreference::ColorOnly;
-		ImportParameters.Propagation = CADLibrary::EDisplayDataPropagationMode::BodyOnly;
+		ImportParameters.SetDisplayPreference(CADLibrary::EDisplayPreference::ColorOnly);
+		ImportParameters.SetPropagationMode(CADLibrary::EDisplayDataPropagationMode::BodyOnly);
 		break;
 	}
+
 	case CADLibrary::ECADFormat::SOLIDWORKS:
 	case CADLibrary::ECADFormat::INVENTOR:
 	case CADLibrary::ECADFormat::CREO:
 	{
-		ImportParameters.ModelCoordSys = FDatasmithUtils::EModelCoordSystem::YUp_RightHanded;
-		ImportParameters.DisplayPreference = CADLibrary::EDisplayPreference::ColorOnly;
-		ImportParameters.Propagation = CADLibrary::EDisplayDataPropagationMode::BodyOnly;
+		ImportParameters.SetModelCoordinateSystem(FDatasmithUtils::EModelCoordSystem::YUp_RightHanded);
+		ImportParameters.SetDisplayPreference(CADLibrary::EDisplayPreference::ColorOnly);
+		ImportParameters.SetPropagationMode(CADLibrary::EDisplayDataPropagationMode::BodyOnly);
 		break;
 	}
+
 	case CADLibrary::ECADFormat::DWG:
 	{
-		ImportParameters.DisplayPreference = CADLibrary::EDisplayPreference::ColorOnly;
-		ImportParameters.Propagation = CADLibrary::EDisplayDataPropagationMode::BodyOnly;
+		ImportParameters.SetDisplayPreference(CADLibrary::EDisplayPreference::ColorOnly);
+		ImportParameters.SetPropagationMode(CADLibrary::EDisplayDataPropagationMode::BodyOnly);
 		break;
 	}
-	}
 
-	ImportParameters.bEnableTimeControl = CADLibrary::bGEnableTimeControl;
-	if (CADLibrary::GMaxImportThreads == 1)
-	{
-		ImportParameters.bEnableSequentialImport = CADLibrary::bGEnableCADCache;
-	}
-	else
-	{
-		ImportParameters.bEnableSequentialImport = true;
+	default:
+		break;
 	}
 
 	FString CachePath = FDatasmithCADTranslatorModule::Get().GetCacheDir();
@@ -156,14 +142,9 @@ bool FDatasmithCADTranslator::LoadScene(TSharedRef<IDatasmithScene> DatasmithSce
 	{
 		CachePath = FPaths::ConvertRelativePathToFull(CachePath);
 	}
-	else
-	{
-		// Sequential Import is disabling because it needs CADCache
-		ImportParameters.bEnableSequentialImport = false;
-	}
 
 	// Use sequential translation (multi-processed or not)
-	if (ImportParameters.bEnableSequentialImport)
+	if (CADLibrary::FImportParameters::bGEnableCADCache)
 	{
 		TMap<uint32, FString> CADFileToUEFileMap;
 		{
@@ -173,7 +154,7 @@ bool FDatasmithCADTranslator::LoadScene(TSharedRef<IDatasmithScene> DatasmithSce
 				NumCores = FMath::Min(CADLibrary::GMaxImportThreads, NumCores);
 			}
 			DatasmithDispatcher::FDatasmithDispatcher Dispatcher(ImportParameters, CachePath, NumCores, CADFileToUEFileMap, CADFileToUEGeomMap);
-			Dispatcher.AddTask(FileDescription);
+			Dispatcher.AddTask(FileDescriptor);
 
 			Dispatcher.Process(CADLibrary::GMaxImportThreads != 1);
 		}
@@ -186,7 +167,7 @@ bool FDatasmithCADTranslator::LoadScene(TSharedRef<IDatasmithScene> DatasmithSce
 		return true;
 	}
 
-	CADLibrary::FCADFileReader FileReader(ImportParameters, FileDescription, *FPaths::EnginePluginsDir());
+	CADLibrary::FCADFileReader FileReader(ImportParameters, FileDescriptor, *FPaths::EnginePluginsDir());
 	if (FileReader.ProcessFile() != CADLibrary::ECADParsingResult::ProcessOk)
 	{
 		return false;
@@ -220,7 +201,7 @@ bool FDatasmithCADTranslator::LoadStaticMesh(const TSharedRef<IDatasmithMeshElem
 	if (TOptional< FMeshDescription > Mesh = MeshBuilderPtr->GetMeshDescription(MeshElement, MeshParameters))
 	{
 		OutMeshPayload.LodMeshes.Add(MoveTemp(Mesh.GetValue()));
-		if (ImportParameters.bDisableCADKernelTessellation)
+		if (CADLibrary::FImportParameters::bGDisableCADKernelTessellation)
 		{
 			CoreTechSurface::AddSurfaceDataForMesh(MeshElement->GetFile(), ImportParameters, MeshParameters, GetCommonTessellationOptions(), OutMeshPayload);
 		}
diff --git a/Engine/Plugins/Enterprise/DatasmithCADImporter/Source/DatasmithCADTranslator/Private/DatasmithCADTranslatorModule.cpp b/Engine/Plugins/Enterprise/DatasmithCADImporter/Source/DatasmithCADTranslator/Private/DatasmithCADTranslatorModule.cpp
--- a/Engine/Plugins/Enterprise/DatasmithCADImporter/Source/DatasmithCADTranslator/Private/DatasmithCADTranslatorModule.cpp
+++ b/Engine/Plugins/Enterprise/DatasmithCADImporter/Source/DatasmithCADTranslator/Private/DatasmithCADTranslatorModule.cpp
@@ -25,13 +25,7 @@ void FDatasmithCADTranslatorModule::StartupModule()
 		IFileManager::Get().DeleteDirectory(*OldCacheDir, true, true);
 	}
 
-	bool bCreateCacheFolder = false;
-	if (IConsoleVariable* CVar = IConsoleManager::Get().FindConsoleVariable(TEXT("ds.CADTranslator.EnableCADCache")))
-	{
-		bCreateCacheFolder = CVar->GetInt() != 0;
-	}
-
-	if (bCreateCacheFolder)
+	if (CADLibrary::FImportParameters::bGEnableCADCache)
 	{
 		CacheDir = FPaths::ConvertRelativePathToFull(FPaths::Combine(FPaths::ProjectIntermediateDir(), TEXT("DatasmithCADCache"), *FString::FromInt(CacheVersion)));
 		IFileManager::Get().MakeDirectory(*CacheDir);
diff --git a/Engine/Plugins/Enterprise/DatasmithCADImporter/Source/DatasmithCADTranslator/Private/DatasmithSceneGraphBuilder.cpp b/Engine/Plugins/Enterprise/DatasmithCADImporter/Source/DatasmithCADTranslator/Private/DatasmithSceneGraphBuilder.cpp
--- a/Engine/Plugins/Enterprise/DatasmithCADImporter/Source/DatasmithCADTranslator/Private/DatasmithSceneGraphBuilder.cpp
+++ b/Engine/Plugins/Enterprise/DatasmithCADImporter/Source/DatasmithCADTranslator/Private/DatasmithSceneGraphBuilder.cpp
@@ -39,9 +39,9 @@ namespace DatasmithSceneGraphBuilderImpl
 		}
 
 		FTransform LocalTransform(Instance.TransformMatrix);
-		FTransform LocalUETransform = FDatasmithUtils::ConvertTransform(ImportParameters.ModelCoordSys, LocalTransform);
+		FTransform LocalUETransform = FDatasmithUtils::ConvertTransform(ImportParameters.GetModelCoordSys(), LocalTransform);
 
-		Actor->SetTranslation(LocalUETransform.GetTranslation() * ImportParameters.ScaleFactor);
+		Actor->SetTranslation(LocalUETransform.GetTranslation() * ImportParameters.GetScaleFactor());
 		Actor->SetScale(LocalUETransform.GetScale3D());
 		Actor->SetRotation(LocalUETransform.GetRotation());
 	}
@@ -516,7 +516,7 @@ TSharedPtr<IDatasmithActorElement> FDatasmithSceneBaseGraphBuilder::BuildBody(in
 	ActorElement->SetLabel(*BodyLabel);
 	ActorElement->SetStaticMeshPathName(MeshElement->GetName());
 
-	if (MaterialUuid && ImportParameters.Propagation != CADLibrary::EDisplayDataPropagationMode::BodyOnly)
+	if (MaterialUuid && ImportParameters.GetPropagation() != CADLibrary::EDisplayDataPropagationMode::BodyOnly)
 	{
 		TSharedPtr< IDatasmithMaterialIDElement > PartMaterialIDElement = FindOrAddMaterial(MaterialUuid);
 		const TCHAR* MaterialIDElementName = PartMaterialIDElement->GetName();
@@ -733,19 +733,20 @@ void FDatasmithSceneBaseGraphBuilder::AddMetaData(TSharedPtr< IDatasmithActorEle
 
 bool FDatasmithSceneBaseGraphBuilder::DoesActorHaveChildrenOrIsAStaticMesh(const TSharedPtr< IDatasmithActorElement >& ActorElement)
 {
-	if (ActorElement != nullptr)
-	{
-		if (ActorElement->GetChildrenCount() > 0)
-		{
-			return true;
-		}
-		else if (ActorElement->IsA(EDatasmithElementType::StaticMeshActor))
-		{
-			const TSharedPtr< IDatasmithMeshActorElement >& MeshActorElement = StaticCastSharedPtr< IDatasmithMeshActorElement >(ActorElement);
-			return FCString::Strlen(MeshActorElement->GetStaticMeshPathName()) > 0;
-		}
-	}
-	return false;
+	return true;
+	//if (ActorElement != nullptr)
+	//{
+	//	if (ActorElement->GetChildrenCount() > 0)
+	//	{
+	//		return true;
+	//	}
+	//	else if (ActorElement->IsA(EDatasmithElementType::StaticMeshActor))
+	//	{
+	//		const TSharedPtr< IDatasmithMeshActorElement >& MeshActorElement = StaticCastSharedPtr< IDatasmithMeshActorElement >(ActorElement);
+	//		return FCString::Strlen(MeshActorElement->GetStaticMeshPathName()) > 0;
+	//	}
+	//}
+	//return false;
 }
 
 
diff --git a/Engine/Plugins/Enterprise/DatasmithCADImporter/Source/DatasmithDispatcher/Private/DatasmithDispatcher.cpp b/Engine/Plugins/Enterprise/DatasmithCADImporter/Source/DatasmithDispatcher/Private/DatasmithDispatcher.cpp
--- a/Engine/Plugins/Enterprise/DatasmithCADImporter/Source/DatasmithDispatcher/Private/DatasmithDispatcher.cpp
+++ b/Engine/Plugins/Enterprise/DatasmithCADImporter/Source/DatasmithDispatcher/Private/DatasmithDispatcher.cpp
@@ -28,7 +28,7 @@ FDatasmithDispatcher::FDatasmithDispatcher(const CADLibrary::FImportParameters&
 	, NumberOfWorkers(InNumberOfWorkers)
 	, NextWorkerId(0)
 {
-	if (ImportParameters.bEnableSequentialImport)
+	if (CADLibrary::FImportParameters::bGEnableCADCache)
 	{
 		// init cache folders
 		IFileManager::Get().MakeDirectory(*FPaths::Combine(ProcessCacheFolder, TEXT("scene")), true);
diff --git a/Engine/Plugins/Enterprise/DatasmithCADImporter/Source/DatasmithDispatcher/Private/DatasmithWorkerHandler.cpp b/Engine/Plugins/Enterprise/DatasmithCADImporter/Source/DatasmithDispatcher/Private/DatasmithWorkerHandler.cpp
--- a/Engine/Plugins/Enterprise/DatasmithCADImporter/Source/DatasmithDispatcher/Private/DatasmithWorkerHandler.cpp
+++ b/Engine/Plugins/Enterprise/DatasmithCADImporter/Source/DatasmithDispatcher/Private/DatasmithWorkerHandler.cpp
@@ -44,11 +44,11 @@ FDatasmithWorkerHandler::FDatasmithWorkerHandler(FDatasmithDispatcher& InDispatc
 	, WorkerState(EWorkerState::Uninitialized)
 	, ErrorState(EWorkerErrorState::Ok)
 	, CachePath(InCachePath)
+	, ImportParametersCommand(InImportParameters)
 	, bShouldTerminate(false)
 {
 	ThreadName = FString(TEXT("DatasmithWorkerHandler_")) + FString::FromInt(Id);
 	IOThread = FThread(*ThreadName, [this]() { Run(); } );
-	ImportParametersCommand.ImportParameters = InImportParameters;
 }
 
 FDatasmithWorkerHandler::~FDatasmithWorkerHandler()
diff --git a/Engine/Plugins/Enterprise/DatasmithCADImporter/Source/DatasmithDispatcher/Public/DatasmithCommands.h b/Engine/Plugins/Enterprise/DatasmithCADImporter/Source/DatasmithDispatcher/Public/DatasmithCommands.h
--- a/Engine/Plugins/Enterprise/DatasmithCADImporter/Source/DatasmithDispatcher/Public/DatasmithCommands.h
+++ b/Engine/Plugins/Enterprise/DatasmithCADImporter/Source/DatasmithDispatcher/Public/DatasmithCommands.h
@@ -109,6 +109,16 @@ public:
 class FImportParametersCommand : public ICommand
 {
 public:
+	FImportParametersCommand(const CADLibrary::FImportParameters& InImportParameters)
+		: ImportParameters(InImportParameters)
+	{
+	}
+
+	FImportParametersCommand()
+		: ImportParameters()
+	{
+	}
+
 	virtual ECommandId GetType() const override { return ECommandId::ImportParams; }
 
 protected:
diff --git a/Engine/Plugins/Enterprise/DatasmithCADImporter/Source/DatasmithOpenNurbsTranslator/Private/DatasmithOpenNurbsTranslator.cpp b/Engine/Plugins/Enterprise/DatasmithCADImporter/Source/DatasmithOpenNurbsTranslator/Private/DatasmithOpenNurbsTranslator.cpp
--- a/Engine/Plugins/Enterprise/DatasmithCADImporter/Source/DatasmithOpenNurbsTranslator/Private/DatasmithOpenNurbsTranslator.cpp
+++ b/Engine/Plugins/Enterprise/DatasmithCADImporter/Source/DatasmithOpenNurbsTranslator/Private/DatasmithOpenNurbsTranslator.cpp
@@ -640,18 +640,17 @@ public:
 		, ArchiveOpenNurbsVersion(0)
 		, FileLength(0)
 		, NumCRCErrors(0)
+		, ImportParameters(0.001, 0.1, FDatasmithUtils::EModelCoordSystem::ZUp_RightHanded_FBXLegacy)
 	{
+
 		if (!TranslationCache.IsValid())
 		{
 			TranslationCache = MakeShared<FTranslationCache>();
 		}
 
-		CADLibrary::FImportParameters ImportParameters;
-		ImportParameters.MetricUnit = 0.001;
-		ImportParameters.ScaleFactor = 0.1;
-		ImportParameters.ModelCoordSys = FDatasmithUtils::EModelCoordSystem::ZUp_RightHanded_FBXLegacy;
+		ImportParameters.SwitchOffUVMapScaling();
 
-		if (CADLibrary::bGDisableCADKernelTessellation)
+		if (CADLibrary::FImportParameters::bGDisableCADKernelTessellation)
 		{
 			TSharedRef<FOpenNurbsBRepToCoretechConverter> OpenNurbsBRepToCoretechConverter = MakeShared<FOpenNurbsBRepToCoretechConverter>(TEXT("Al2CTSharedSession"), ImportParameters);
 			CADModelConverter = OpenNurbsBRepToCoretechConverter;
@@ -833,6 +832,7 @@ private:
 
 	TArray<FString> MissingRenderMeshes;
 
+	CADLibrary::FImportParameters ImportParameters;
 	TSharedPtr<CADLibrary::ICADModelConverter> CADModelConverter;
 	TSharedPtr<IOpenNurbsBRepConverter> OpenNurbsBRepConverter;
 
@@ -2965,8 +2965,7 @@ bool FOpenNurbsTranslatorImpl::TranslateBRep(ON_Brep* Brep, const ON_3dmObjectAt
 	if (OpenNurbsOptions.Geometry == EDatasmithOpenNurbsBrepTessellatedSource::UseUnrealNurbsTessellation)
 	{
 		// Ref. visitBRep
-		const FDatasmithOpenNurbsOptions& TessellationOptions = OpenNurbsOptions;
-		CADModelConverter->SetImportParameters(TessellationOptions.ChordTolerance, TessellationOptions.MaxEdgeLength, TessellationOptions.NormalTolerance, (CADLibrary::EStitchingTechnique) TessellationOptions.StitchingTechnique, false);
+		CADModelConverter->SetImportParameters(OpenNurbsOptions.ChordTolerance, OpenNurbsOptions.MaxEdgeLength, OpenNurbsOptions.NormalTolerance, (CADLibrary::EStitchingTechnique)OpenNurbsOptions.StitchingTechnique);
 
 		CADModelConverter->InitializeProcess(CADModelConverter->GetMetricUnit());
 
diff --git a/Engine/Plugins/Enterprise/DatasmithCADImporter/Source/DatasmithPLMXMLTranslator/Private/DatasmithPlmXmlImporter.cpp b/Engine/Plugins/Enterprise/DatasmithCADImporter/Source/DatasmithPLMXMLTranslator/Private/DatasmithPlmXmlImporter.cpp
--- a/Engine/Plugins/Enterprise/DatasmithCADImporter/Source/DatasmithPLMXMLTranslator/Private/DatasmithPlmXmlImporter.cpp
+++ b/Engine/Plugins/Enterprise/DatasmithCADImporter/Source/DatasmithPLMXMLTranslator/Private/DatasmithPlmXmlImporter.cpp
@@ -240,7 +240,7 @@ namespace PlmXml
 		FPlmXmlMeshLoaderWithDatasmithDispatcher(TSharedRef<IDatasmithScene> InDatasmithScene, FDatasmithTessellationOptions& InTessellationOptions)
 			: DatasmithScene(InDatasmithScene)
 			, TessellationOptions(InTessellationOptions)
-		
+			, ImportParameters()
 		{
 			FCADToolsModule& CADToolsModule = FCADToolsModule::Get();
 			
@@ -248,12 +248,7 @@ namespace PlmXml
 			IFileManager::Get().MakeDirectory(*CacheDir);
 
 			// Setup of import parameters for DatasmithDispatcher copied from FDatasmithCADTranslator's setup
-			ImportParameters.MetricUnit = 0.001;
-			ImportParameters.ScaleFactor = 0.1;
-			ImportParameters.ChordTolerance = TessellationOptions.ChordTolerance;
-			ImportParameters.MaxEdgeLength = TessellationOptions.MaxEdgeLength;
-			ImportParameters.MaxNormalAngle = TessellationOptions.NormalTolerance;
-			ImportParameters.StitchingTechnique = (CADLibrary::EStitchingTechnique) TessellationOptions.StitchingTechnique;
+			ImportParameters.SetTesselationParameters(TessellationOptions.ChordTolerance, TessellationOptions.MaxEdgeLength, TessellationOptions.NormalTolerance, (CADLibrary::EStitchingTechnique) TessellationOptions.StitchingTechnique);
 
 			DatasmithDispatcher = MakeUnique<DatasmithDispatcher::FDatasmithDispatcher>(ImportParameters, CacheDir, FPlatformMisc::NumberOfCores(), CADFileToUEFileMap, CADFileToUEGeomMap);
 		}
@@ -305,7 +300,7 @@ namespace PlmXml
 			{
 				OutMeshPayload.LodMeshes.Add(MoveTemp(Mesh.GetValue()));
 
-				if (ImportParameters.bDisableCADKernelTessellation)
+				if (CADLibrary::FImportParameters::bGDisableCADKernelTessellation)
 				{
 					CoreTechSurface::AddSurfaceDataForMesh(MeshElement->GetFile(), ImportParameters, MeshParameters, TessellationOptions, OutMeshPayload);
 				}
diff --git a/Engine/Plugins/Enterprise/DatasmithCADImporter/Source/DatasmithWireTranslator/Private/DatasmithWireTranslator.cpp b/Engine/Plugins/Enterprise/DatasmithCADImporter/Source/DatasmithWireTranslator/Private/DatasmithWireTranslator.cpp
--- a/Engine/Plugins/Enterprise/DatasmithCADImporter/Source/DatasmithWireTranslator/Private/DatasmithWireTranslator.cpp
+++ b/Engine/Plugins/Enterprise/DatasmithCADImporter/Source/DatasmithWireTranslator/Private/DatasmithWireTranslator.cpp
@@ -134,6 +134,7 @@ public:
 		, ArchiveWireVersion(0)
 		, FileLength(0)
 		, NumCRCErrors(0)
+		, ImportParameters(0.01, 1)
 	{
 		// Set ProductName, ProductVersion in DatasmithScene for Analytics purpose
 		// application_name is something like "Catia V5"
@@ -143,11 +144,7 @@ public:
 		DatasmithScene->SetProductName(TEXT("Alias Tools"));
 		DatasmithScene->SetProductVersion(TEXT("Alias 2022"));
 
-		CADLibrary::FImportParameters ImportParameters;
-		ImportParameters.MetricUnit = 0.01;
-		ImportParameters.ScaleFactor = 1;
-
-		if(CADLibrary::bGDisableCADKernelTessellation)
+		if(FImportParameters::bGDisableCADKernelTessellation)
 		{
 			TSharedRef<FAliasModelToCoretechConverter> AliasToCoretechConverter = MakeShared<FAliasModelToCoretechConverter>(TEXT("Al2CTSharedSession"), ImportParameters);
 			CADModelConverter = AliasToCoretechConverter;
@@ -266,8 +263,6 @@ private:
 	FString OutputPath;
 	FString SceneFullPath;
 
-	FDatasmithTessellationOptions TessellationOptions;
-
 	// Hash value of the scene file used to check if the file has been modified for re-import
 	uint32 SceneFileHash;
 
@@ -295,6 +290,8 @@ private:
 	// If > 0, then the archive is corrupt.
 	int32 NumCRCErrors;
 
+	CADLibrary::FImportParameters ImportParameters;
+
 	TSharedPtr<ICADModelConverter> CADModelConverter;
 	TSharedPtr<IAliasBRepConverter> AliasBRepConverter;
 
@@ -302,7 +299,7 @@ private:
 
 void FWireTranslatorImpl::SetTessellationOptions(const FDatasmithTessellationOptions& Options)
 {
-	TessellationOptions = Options;
+	ImportParameters.SetTesselationParameters(Options.ChordTolerance, Options.MaxEdgeLength, Options.NormalTolerance, (CADLibrary::EStitchingTechnique) Options.StitchingTechnique);
 	SceneFileHash = HashCombine(Options.GetHash(), GetSceneFileHash(SceneFullPath, SceneName));
 }
 
@@ -316,8 +313,6 @@ bool FWireTranslatorImpl::Read()
 		return false;
 	}
 
-	CADModelConverter->SetImportParameters(TessellationOptions.ChordTolerance, TessellationOptions.MaxEdgeLength, TessellationOptions.NormalTolerance, (CADLibrary::EStitchingTechnique)TessellationOptions.StitchingTechnique, true);
-
 	AlRetrieveOptions options;
 	AlUniverse::retrieveOptions(options);
 
@@ -1558,7 +1553,7 @@ void FWireTranslatorImpl::AddNodeInBodyGroup(TSharedPtr<AlDagNode>& DagNode, con
 
 bool FWireTranslatorImpl::RecurseDagForLeaves(const TSharedPtr<AlDagNode>& FirstDagNode, const FDagNodeInfo& ParentInfo)
 {
-	if (TessellationOptions.StitchingTechnique != EDatasmithCADStitchingTechnique::StitchingSew)
+	if (ImportParameters.GetStitchingTechnique() != EStitchingTechnique::StitchingSew)
 	{
 		return RecurseDagForLeavesNoMerge(FirstDagNode, ParentInfo);
 	}
@@ -1801,7 +1796,7 @@ TOptional<FMeshDescription> FWireTranslatorImpl::MeshDagNodeWithExternalMesher(T
 		// All actors of a Alias symmetric layer are defined in the world Reference i.e. they have identity transform. So Mesh actor has to be defined in the world reference. 
 		ObjectReference = EAliasObjectReference::WorldReference;
 	}
-	else if (TessellationOptions.StitchingTechnique == EDatasmithCADStitchingTechnique::StitchingSew)
+	else if (ImportParameters.GetStitchingTechnique() == EStitchingTechnique::StitchingSew)
 	{
 		// In the case of StitchingSew, AlDagNode children of a GroupNode are merged together. To be merged, they have to be defined in the reference of parent GroupNode.
 		ObjectReference = EAliasObjectReference::ParentReference;
@@ -1832,7 +1827,7 @@ TOptional<FMeshDescription> FWireTranslatorImpl::GetMeshOfShellNode(AlDagNode& D
 		AlMatrix4x4 AlMatrix;
 		DagNode.inverseGlobalTransformationMatrix(AlMatrix);
 		// TODO: the best way, should be to don't have to apply inverse global transform to the generated mesh
-		TSharedPtr<AlDagNode> TesselatedNode = OpenModelUtils::TesselateDagLeaf(DagNode, ETesselatorType::Fast, TessellationOptions.ChordTolerance);
+		TSharedPtr<AlDagNode> TesselatedNode = OpenModelUtils::TesselateDagLeaf(DagNode, ETesselatorType::Fast, ImportParameters.GetChordTolerance());
 		if (TesselatedNode.IsValid())
 		{
 			// Get the meshes from the dag nodes. Note that removing the mesh's DAG.
diff --git a/Engine/Plugins/Enterprise/DatasmithCADImporter/Source/ParametricSurface/Public/CADModelConverter.h b/Engine/Plugins/Enterprise/DatasmithCADImporter/Source/ParametricSurface/Public/CADModelConverter.h
--- a/Engine/Plugins/Enterprise/DatasmithCADImporter/Source/ParametricSurface/Public/CADModelConverter.h
+++ b/Engine/Plugins/Enterprise/DatasmithCADImporter/Source/ParametricSurface/Public/CADModelConverter.h
@@ -12,7 +12,7 @@ struct FDatasmithTessellationOptions;
 
 namespace CADLibrary
 {
-struct FImportParameters;
+class FImportParameters;
 struct FMeshParameters;
 
 /**
@@ -56,9 +56,8 @@ public:
 	 * @param StitchingTechnique : CAD topology correction technique
 	 * @param bScaleUVMap : Scale the UV map to a world unit.
 	 */
-	virtual void SetImportParameters(double ChordTolerance, double MaxEdgeLength, double NormalTolerance, CADLibrary::EStitchingTechnique StitchingTechnique, bool bScaleUVMap) = 0;
+	virtual void SetImportParameters(double ChordTolerance, double MaxEdgeLength, double NormalTolerance, CADLibrary::EStitchingTechnique StitchingTechnique) = 0;
 
-	//virtual void SetScaleFactor(double NewScaleValue) = 0;
 	virtual void SetMetricUnit(double NewMetricUnit) = 0;
 	virtual double GetScaleFactor() const = 0;
 	virtual double GetMetricUnit() const = 0;
diff --git a/Engine/Plugins/Enterprise/DatasmithCADImporter/Source/ParametricSurface/Public/ParametricSurfaceTranslator.h b/Engine/Plugins/Enterprise/DatasmithCADImporter/Source/ParametricSurface/Public/ParametricSurfaceTranslator.h
--- a/Engine/Plugins/Enterprise/DatasmithCADImporter/Source/ParametricSurface/Public/ParametricSurfaceTranslator.h
+++ b/Engine/Plugins/Enterprise/DatasmithCADImporter/Source/ParametricSurface/Public/ParametricSurfaceTranslator.h
@@ -14,7 +14,6 @@ class UDatasmithCommonTessellationOptions;
 class PARAMETRICSURFACE_API FParametricSurfaceTranslator : public IDatasmithTranslator
 {
 public:
-	virtual ~FParametricSurfaceTranslator() {}
 
 	// Begin IDatasmithTranslator overrides
 	virtual void GetSceneImportOptions(TArray<TStrongObjectPtr<UDatasmithOptionsBase>>& Options) override;
diff --git a/Engine/Plugins/Enterprise/DatasmithCADImporter/Source/ParametricSurfaceExtension/Private/DataprepTessellationOperation.h b/Engine/Plugins/Enterprise/DatasmithCADImporter/Source/ParametricSurfaceExtension/Private/DataprepTessellationOperation.h
--- a/Engine/Plugins/Enterprise/DatasmithCADImporter/Source/ParametricSurfaceExtension/Private/DataprepTessellationOperation.h
+++ b/Engine/Plugins/Enterprise/DatasmithCADImporter/Source/ParametricSurfaceExtension/Private/DataprepTessellationOperation.h
@@ -12,7 +12,7 @@
 
 namespace CADLibrary
 {
-	struct FImportParameters;
+	class FImportParameters;
 }
 struct FMeshDescription;
 class UParametricSurfaceData;
diff --git a/Engine/Source/Programs/Enterprise/Datasmith/DatasmithCADWorker/Private/DatasmithCADWorkerImpl.h b/Engine/Source/Programs/Enterprise/Datasmith/DatasmithCADWorker/Private/DatasmithCADWorkerImpl.h
--- a/Engine/Source/Programs/Enterprise/Datasmith/DatasmithCADWorker/Private/DatasmithCADWorkerImpl.h
+++ b/Engine/Source/Programs/Enterprise/Datasmith/DatasmithCADWorker/Private/DatasmithCADWorkerImpl.h
@@ -11,7 +11,7 @@
 
 
 struct FFileStatData;
-struct FImportParameters;
+class FImportParameters;
 
 class FDatasmithCADWorkerImpl
 {
diff --git a/Engine/Source/Programs/InterchangeWorker/Private/InterchangeWorkerImpl.h b/Engine/Source/Programs/InterchangeWorker/Private/InterchangeWorkerImpl.h
--- a/Engine/Source/Programs/InterchangeWorker/Private/InterchangeWorkerImpl.h
+++ b/Engine/Source/Programs/InterchangeWorker/Private/InterchangeWorkerImpl.h
@@ -10,7 +10,7 @@
 #include "InterchangeFbxParser.h"
 
 struct FFileStatData;
-struct FImportParameters;
+class FImportParameters;
 
 class FInterchangeWorkerImpl
 {
