Fix CL17540414:
- remove dev comments
- remove default initialization

diff --git a/Engine/Plugins/Enterprise/DatasmithCADImporter/Source/DatasmithCADTranslator/Private/DatasmithSceneGraphBuilder.cpp b/Engine/Plugins/Enterprise/DatasmithCADImporter/Source/DatasmithCADTranslator/Private/DatasmithSceneGraphBuilder.cpp
--- a/Engine/Plugins/Enterprise/DatasmithCADImporter/Source/DatasmithCADTranslator/Private/DatasmithSceneGraphBuilder.cpp
+++ b/Engine/Plugins/Enterprise/DatasmithCADImporter/Source/DatasmithCADTranslator/Private/DatasmithSceneGraphBuilder.cpp
@@ -733,20 +733,19 @@ void FDatasmithSceneBaseGraphBuilder::AddMetaData(TSharedPtr< IDatasmithActorEle
 
 bool FDatasmithSceneBaseGraphBuilder::DoesActorHaveChildrenOrIsAStaticMesh(const TSharedPtr< IDatasmithActorElement >& ActorElement)
 {
-	return true;
-	//if (ActorElement != nullptr)
-	//{
-	//	if (ActorElement->GetChildrenCount() > 0)
-	//	{
-	//		return true;
-	//	}
-	//	else if (ActorElement->IsA(EDatasmithElementType::StaticMeshActor))
-	//	{
-	//		const TSharedPtr< IDatasmithMeshActorElement >& MeshActorElement = StaticCastSharedPtr< IDatasmithMeshActorElement >(ActorElement);
-	//		return FCString::Strlen(MeshActorElement->GetStaticMeshPathName()) > 0;
-	//	}
-	//}
-	//return false;
+	if (ActorElement != nullptr)
+	{
+		if (ActorElement->GetChildrenCount() > 0)
+		{
+			return true;
+		}
+		else if (ActorElement->IsA(EDatasmithElementType::StaticMeshActor))
+		{
+			const TSharedPtr< IDatasmithMeshActorElement >& MeshActorElement = StaticCastSharedPtr< IDatasmithMeshActorElement >(ActorElement);
+			return FCString::Strlen(MeshActorElement->GetStaticMeshPathName()) > 0;
+		}
+	}
+	return false;
 }
 
 
diff --git a/Engine/Plugins/Enterprise/DatasmithCADImporter/Source/DatasmithDispatcher/Public/DatasmithCommands.h b/Engine/Plugins/Enterprise/DatasmithCADImporter/Source/DatasmithDispatcher/Public/DatasmithCommands.h
--- a/Engine/Plugins/Enterprise/DatasmithCADImporter/Source/DatasmithDispatcher/Public/DatasmithCommands.h
+++ b/Engine/Plugins/Enterprise/DatasmithCADImporter/Source/DatasmithDispatcher/Public/DatasmithCommands.h
@@ -115,7 +115,6 @@ public:
 	}
 
 	FImportParametersCommand()
-		: ImportParameters()
 	{
 	}
 
diff --git a/Engine/Plugins/Enterprise/DatasmithCADImporter/Source/DatasmithPLMXMLTranslator/Private/DatasmithPlmXmlImporter.cpp b/Engine/Plugins/Enterprise/DatasmithCADImporter/Source/DatasmithPLMXMLTranslator/Private/DatasmithPlmXmlImporter.cpp
--- a/Engine/Plugins/Enterprise/DatasmithCADImporter/Source/DatasmithPLMXMLTranslator/Private/DatasmithPlmXmlImporter.cpp
+++ b/Engine/Plugins/Enterprise/DatasmithCADImporter/Source/DatasmithPLMXMLTranslator/Private/DatasmithPlmXmlImporter.cpp
@@ -240,7 +240,6 @@ namespace PlmXml
 		FPlmXmlMeshLoaderWithDatasmithDispatcher(TSharedRef<IDatasmithScene> InDatasmithScene, FDatasmithTessellationOptions& InTessellationOptions)
 			: DatasmithScene(InDatasmithScene)
 			, TessellationOptions(InTessellationOptions)
-			, ImportParameters()
 		{
 			FCADToolsModule& CADToolsModule = FCADToolsModule::Get();
 			
