TechSoftImplementation

diff --git a/Engine/Plugins/Enterprise/DatasmithCADImporter/Source/CADInterfaces/Private/CADInterfacesModule.cpp b/Engine/Plugins/Enterprise/DatasmithCADImporter/Source/CADInterfaces/Private/CADInterfacesModule.cpp
--- a/Engine/Plugins/Enterprise/DatasmithCADImporter/Source/CADInterfaces/Private/CADInterfacesModule.cpp
+++ b/Engine/Plugins/Enterprise/DatasmithCADImporter/Source/CADInterfaces/Private/CADInterfacesModule.cpp
@@ -1,10 +1,12 @@
 // Copyright Epic Games, Inc. All Rights Reserved.
 #include "CADInterfacesModule.h"
 
+#include "CADOptions.h"
 #include "CoreTechTypes.h"
 #include "HAL/PlatformProcess.h"
 #include "Misc/Paths.h"
 #include "Modules/ModuleManager.h"
+#include "TechSoftInterface.h"
 
 #define LOCTEXT_NAMESPACE "CADInterfacesModule"
 
@@ -17,9 +19,11 @@ private:
 	virtual void ShutdownModule() override;
 
 	static void* KernelIOLibHandle;
+	static void* TechSoftLibHandle;
 };
 
 void* FCADInterfacesModule::KernelIOLibHandle = nullptr;
+void* FCADInterfacesModule::TechSoftLibHandle = nullptr;
 
 ICADInterfacesModule& ICADInterfacesModule::Get()
 {
@@ -32,9 +36,19 @@ ECADInterfaceAvailability ICADInterfacesModule::GetAvailability()
 {
 	if (FModuleManager::Get().IsModuleLoaded(CADINTERFACES_MODULE_NAME))
 	{
-		if(CADLibrary::CTKIO_InitializeKernel())
+		if (CADLibrary::GCADLibrary == TEXT("KernelIO"))
 		{
-			return ECADInterfaceAvailability::Available;
+			if (CADLibrary::CTKIO_InitializeKernel())
+			{
+				return ECADInterfaceAvailability::Available;
+			}
+		}
+		else if (CADLibrary::GCADLibrary == TEXT("TechSoft"))
+		{
+			if (CADLibrary::TECHSOFT_InitializeKernel())
+			{
+				return ECADInterfaceAvailability::Available;
+			}
 		}
 	}
 
@@ -44,32 +58,65 @@ ECADInterfaceAvailability ICADInterfacesModule::GetAvailability()
 
 void FCADInterfacesModule::StartupModule()
 {
+	// determine directory paths
+	FString CADImporterDllPath = FPaths::Combine(FPaths::EnginePluginsDir(), TEXT("Enterprise/DatasmithCADImporter"), TEXT("Binaries"), FPlatformProcess::GetBinariesSubdirectory());
+	FPlatformProcess::PushDllDirectory(*CADImporterDllPath);
+
 #if WITH_EDITOR & defined(USE_KERNEL_IO_SDK)
 	check(KernelIOLibHandle == nullptr);
 
 	FString KernelIODll = TEXT("kernel_io.dll");
-
-	// determine directory paths
-	FString CADImporterDllPath = FPaths::Combine(FPaths::EnginePluginsDir(), TEXT("Enterprise/DatasmithCADImporter"), TEXT("Binaries"), FPlatformProcess::GetBinariesSubdirectory());
-	FPlatformProcess::PushDllDirectory(*CADImporterDllPath);
 	KernelIODll = FPaths::Combine(CADImporterDllPath, KernelIODll);
 
 	if (!FPaths::FileExists(KernelIODll))
 	{
 		UE_LOG(CADInterfaces, Warning, TEXT("CoreTech module is missing. Plug-in will not be functional."));
-		return;
 	}
-
-	KernelIOLibHandle = FPlatformProcess::GetDllHandle(*KernelIODll);
-	if (KernelIOLibHandle == nullptr)
+	else
 	{
-		UE_LOG(CADInterfaces, Warning, TEXT("Failed to load required library %s. Plug-in will not be functional."), *KernelIODll);
+		KernelIOLibHandle = FPlatformProcess::GetDllHandle(*KernelIODll);
+		if (KernelIOLibHandle == nullptr)
+		{
+			UE_LOG(CADInterfaces, Warning, TEXT("Failed to load required library %s. Plug-in will not be functional."), *KernelIODll);
+		}
+		else
+		{
+			CADLibrary::InitializeCoreTechInterface();
+		}
+
+		FPlatformProcess::PopDllDirectory(*CADImporterDllPath);
 	}
+#endif
+
+#if WITH_EDITOR & defined(USE_TECHSOFT_SDK)
+	check(TechSoftLibHandle == nullptr);
+
+	FString TechSoftDllPath = FPaths::Combine(CADImporterDllPath, "TechSoft");
+	FPlatformProcess::PushDllDirectory(*TechSoftDllPath);
 
-	FPlatformProcess::PopDllDirectory(*CADImporterDllPath);
+	FString TechSoftDll = TEXT("A3DLIBS.dll");
+	TechSoftDll = FPaths::Combine(TechSoftDllPath, TechSoftDll);
+
+	if (!FPaths::FileExists(TechSoftDll))
+	{
+		UE_LOG(CADInterfaces, Warning, TEXT("TechSoft module is missing. Plug-in will not be functional."));
+	}
+	else
+	{
+		TechSoftLibHandle = FPlatformProcess::GetDllHandle(*TechSoftDll);
+		if (TechSoftLibHandle == nullptr)
+		{
+			UE_LOG(CADInterfaces, Warning, TEXT("Failed to load required library %s. Plug-in will not be functional."), *TechSoftDll);
+		}
+		else
+		{
+			CADLibrary::InitializeTechSoftInterface();
+		}
 
-	CADLibrary::InitializeCoreTechInterface();
+		FPlatformProcess::PopDllDirectory(*TechSoftDllPath);
+	}
 #endif
+
 }
 
 void FCADInterfacesModule::ShutdownModule()
diff --git a/Engine/Plugins/Enterprise/DatasmithCADImporter/Source/CADInterfaces/Private/TUniqueTechSoftObj.h b/Engine/Plugins/Enterprise/DatasmithCADImporter/Source/CADInterfaces/Private/TUniqueTechSoftObj.h
new file mode 100644
--- /dev/null
+++ b/Engine/Plugins/Enterprise/DatasmithCADImporter/Source/CADInterfaces/Private/TUniqueTechSoftObj.h
@@ -0,0 +1,187 @@
+// Copyright Epic Games, Inc. All Rights Reserved.
+#pragma once
+
+#include "CADData.h"
+
+#include "A3DSDKErrorCodes.h"
+#include "A3DSDKInitializeFunctions.h"
+
+namespace TechSoft
+{
+
+	// Single-ownership smart TeshSoft object
+	// Use this when you need to manage TechSoft object's lifetime.
+	//
+	// TechSoft give access to void pointers
+	// According to the context, the class name of the void pointer is known but the class is unknown
+	// i.e. A3DSDKTypes.h defines all type like :
+	// 	   typedef void A3DEntity;		
+	// 	   typedef void A3DAsmModelFile; ...
+	// 
+	// From a pointer, TechSoft give access to a copy of the associated structure :
+	//
+	// const A3DXXXXX* pPointer;
+	// A3DXXXXXData sData; // the structure
+	// A3D_INITIALIZE_DATA(A3DXXXXXData, sData); // initialization of the structure
+	// A3DXXXXXXGet(pPointer, &sData); // Copy of the data of the pointer in the structure
+	// ...
+	// A3DXXXXXXGet(NULL, &sData); // Free the structure
+	//
+	// A3D_INITIALIZE_DATA, and all A3DXXXXXXGet methods are TechSoft macro
+	//
+
+	template< class ObjectType, class PtrType>
+	class TUniqueTSObj
+	{
+	public:
+
+		/**
+		 * Constructor of an initialized ObjectType object
+		 * @param InGetter: Function Pointer of the A3DXXXXXXGet function.
+		 */
+		explicit TUniqueTSObj(A3DStatus (*InGetter)(const PtrType*, ObjectType*))
+		{
+			Getter = InGetter;
+			memset(&Data, 0, sizeof(Data)); 
+			Data.m_usStructSize = sizeof(Data);
+		}
+
+		/**
+		 * Constructor of an filled ObjectType object with the data of DataPtr
+		 * @param DataPtr: the pointer of the data to copy
+		 * @param InGetter: Function pointer of the A3DXXXXXXGet function.
+		 */
+		explicit TUniqueTSObj(const PtrType* DataPtr, A3DStatus (*InGetter)(const PtrType*, ObjectType*))
+			:TUniqueTSObj(InGetter)
+		{
+			Get(DataPtr);
+		}
+
+		/**
+		 * Constructor of an filled ObjectType object with the data of DataPtr
+		 * This type of structure need a specific initialization method
+		 * @param DataPtr: the pointer of the data to copy
+		 * @param InGetter: Function pointer of the A3DXXXXXXGet function.
+		 * @param Initializer: Initialization function pointer for ObjectType.
+		 */
+		explicit TUniqueTSObj(A3DStatus(*InGetter)(const PtrType*, ObjectType*), void(*Initializer)(ObjectType&))
+		{
+			Getter = InGetter;
+			Initializer(Data);
+		}
+
+		/**
+		 * Constructor of an initialized ObjectType object
+		 * This type of structure need a specific initialization method
+		 * @param DataPtr: the pointer of the data to copy
+		 * @param InGetter: Function pointer of the A3DXXXXXXGet function.
+		 * @param Initializer: Initialization function pointer for ObjectType.
+		 */
+		explicit TUniqueTSObj(const PtrType* DataPtr, A3DStatus(*InGetter)(const PtrType*, ObjectType*), void(*Initializer)(ObjectType&))
+			:TUniqueTSObj(InGetter, Initializer)
+		{
+			Get(DataPtr);
+		}
+
+		~TUniqueTSObj()
+		{
+			Getter(NULL, &Data);
+		}
+
+		/**
+		 * Fill the structure with the data of a new DataPtr
+		 */
+		A3DStatus Get(const PtrType* DataPtr)
+		{
+			if (IsValid())
+			{
+				Status = Getter(NULL, &Data);
+			}
+			else
+			{
+				Status = A3DStatus::A3D_SUCCESS;
+			}
+
+			if (!IsValid() || (DataPtr == NULL))
+			{
+				Status = A3DStatus::A3D_ERROR;
+				return Status;
+			}
+
+			Status = Getter(DataPtr, &Data);
+			return Status;
+		}
+
+		/**
+		 * Empty the structure
+		 */
+		void Reset()
+		{
+			Get(NULL);
+		}
+
+		/**
+		 * Return 
+		 *  - A3DStatus::A3D_SUCCESS if the data is filled
+		 *  - A3DStatus::A3D_ERROR if the data is empty
+		 */
+		A3DStatus GetStatus()
+		{
+			return Status;
+		}
+
+		/**
+		 * Return true if the data is filled
+		 */
+		const bool IsValid() const
+		{
+			return Status == A3DStatus::A3D_SUCCESS;
+		}
+
+		// Non-copyable
+		TUniqueTSObj(const TUniqueTSObj&) = delete;
+		TUniqueTSObj& operator=(const TUniqueTSObj&) = delete;
+
+		// Conversion methods
+
+		const ObjectType& operator*() const
+		{
+			return Data;
+		}
+
+		ObjectType& operator*()
+		{
+			return Data;
+		}
+
+		const ObjectType* operator->() const
+		{
+			check(IsValid());
+			return &Data;
+		}
+
+		ObjectType* operator->()
+		{
+			check(IsValid());
+			return &Data;
+		}
+
+		/**
+		 * Return the structure and set the status to filled
+		 * The method is used to manage structure filled by UE
+		 */
+		ObjectType* GetEmptyDataPtr()
+		{
+			Get(NULL);
+			Status = A3DStatus::A3D_SUCCESS;
+			return &Data;
+		}
+
+
+	private:
+		ObjectType Data;
+		A3DStatus Status = A3DStatus::A3D_ERROR;
+		A3DStatus (*Getter)(const PtrType*, ObjectType*);
+	};
+
+}
diff --git a/Engine/Plugins/Enterprise/DatasmithCADImporter/Source/CADInterfaces/Private/TechSoftFileParser.cpp b/Engine/Plugins/Enterprise/DatasmithCADImporter/Source/CADInterfaces/Private/TechSoftFileParser.cpp
new file mode 100644
--- /dev/null
+++ b/Engine/Plugins/Enterprise/DatasmithCADImporter/Source/CADInterfaces/Private/TechSoftFileParser.cpp
@@ -0,0 +1,1731 @@
+// Copyright Epic Games, Inc. All Rights Reserved.
+
+#include "TechSoftFileParser.h"
+
+#ifdef USE_TECHSOFT_SDK
+
+#include "TUniqueTechSoftObj.h"
+
+#undef TEXT
+#ifdef WIN32
+#define WIN32_LEAN_AND_MEAN
+#include <windows.h>
+#endif
+
+#pragma warning(push)
+#pragma warning(disable:4505)
+#pragma warning(disable:4191)
+#include "A3DSDKIncludes.h"
+#pragma warning(pop)
+
+// redefine UE TEXT
+#undef TEXT
+#if !defined(TEXT) && !UE_BUILD_DOCS
+#if PLATFORM_TCHAR_IS_CHAR16
+#define TEXT_PASTE(x) u ## x
+#else
+#define TEXT_PASTE(x) L ## x
+#endif
+#define TEXT(x) TEXT_PASTE(x)
+#endif
+
+
+namespace TechSoftFileParserImpl
+{
+	// Specific initialization function used by TUniqueTechSoftObj
+
+	void InitializeData(A3DAsmPartDefinitionData& sData)
+	{
+		A3D_INITIALIZE_DATA(A3DAsmPartDefinitionData, sData);
+	}
+
+	void InitializeData(A3DRiDirectionData& sData)
+	{
+		A3D_INITIALIZE_DATA(A3DRiDirectionData, sData);
+	}
+
+	void InitializeData(A3DTopoUniqueVertexData& sData)
+	{
+		A3D_INITIALIZE_DATA(A3DTopoUniqueVertexData, sData);
+	}
+
+	void InitializeData(A3DTopoWireEdgeData& sData)
+	{
+		A3D_INITIALIZE_DATA(A3DTopoWireEdgeData, sData);
+	}
+
+	void InitializeData(A3DTopoEdgeData& sData)
+	{
+		A3D_INITIALIZE_DATA(A3DTopoEdgeData, sData);
+	}
+
+	void InitializeData(A3DTopoFaceData& sData)
+	{
+		A3D_INITIALIZE_DATA(A3DTopoFaceData, sData);
+	}
+
+	void InitializeData(A3DTopoBrepDataData& sData)
+	{
+		A3D_INITIALIZE_DATA(A3DTopoBrepDataData, sData);
+	}
+
+
+	// Functions to clean metadata
+
+	inline void RemoveUnwantedChar(FString& StringToClean, TCHAR UnwantedChar)
+	{
+		FString NewString;
+		NewString.Reserve(StringToClean.Len());
+		for (const TCHAR& Char : StringToClean)
+		{
+			if (Char != UnwantedChar)
+			{
+				NewString.AppendChar(Char);
+			}
+		}
+		StringToClean = MoveTemp(NewString);
+	}
+
+	FString GetOriginalName(const FString& Name)
+	{
+		int32 Index;
+		if (Name.FindLastChar(TEXT('['), Index))
+		{
+			return Name.Left(Index);
+		}
+		return Name;
+	}
+
+	FString Get3dxmlOriginalName(const FString& Name)
+	{
+		int32 Index;
+		if (Name.FindChar(TEXT('('), Index))
+		{
+			FString NewName = Name.Right(Index + 1);
+			if (Name.FindLastChar(TEXT(')'), Index))
+			{
+				NewName = NewName.Left(Index);
+			}
+			return NewName;
+		}
+		return Name;
+	}
+
+	FString GetSWOriginalName(const FString& Name)
+	{
+		int32 Position;
+		if (Name.FindLastChar(TEXT('-'), Position))
+		{
+			FString NewName = Name.Left(Position) + TEXT("<") + Name.Right(Position + 1) + TEXT(">");
+			return NewName;
+		}
+		return Name;
+	}
+
+	FString GetJTOriginalName(const FString& Name)
+	{
+		int32 Position;
+		if (Name.FindLastChar(TEXT(';'), Position))
+		{
+			FString NewName = Name.Left(Position);
+			if (Name.FindLastChar(TEXT('.'), Position))
+			{
+				NewName = NewName.Left(Position);
+			}
+			return NewName;
+		}
+		return Name;
+	}
+
+	FString GetCatiaBodyOriginalName(const FString& Name)
+	{
+		if (Name.Left(14) == TEXT("MechanicalTool"))
+		{
+			FString NewName = L"Body" + Name.Right(14);
+			return NewName;
+		}
+		return Name;
+	}
+
+	FString GetCreoBodyOriginalName(const FString& Name)
+	{
+		int32 Position;
+		if (Name.FindLastChar(TEXT('.'), Position))
+		{
+			FString NewName = Name.Left(Position);
+			return NewName;
+		}
+		return Name;
+	}
+
+
+	// Functions used in traverse model process
+	bool IsUnloadedPrototype(const A3DAsmProductOccurrenceData& PrototypeData)
+	{
+		bool bIsUnloaded = true;
+
+		if (PrototypeData.m_uiPOccurrencesSize > 0)
+		{
+			bIsUnloaded = false;
+		}
+		else if (PrototypeData.m_uiEntityReferenceSize > 0)
+		{
+			bIsUnloaded = false;
+		}
+		else if (PrototypeData.m_pPart)
+		{
+			bIsUnloaded = false;
+		}
+		else if (PrototypeData.m_pPrototype)
+		{
+			TechSoft::TUniqueTSObj<A3DAsmProductOccurrenceData, A3DAsmProductOccurrence> SubPrototypeData(PrototypeData.m_pPrototype, A3DAsmProductOccurrenceGet);
+			if (SubPrototypeData.IsValid())
+			{
+				bIsUnloaded = TechSoftFileParserImpl::IsUnloadedPrototype(*SubPrototypeData);
+			}
+		}
+
+		return bIsUnloaded;
+	}
+
+	void TraverseAttribute(const A3DMiscAttributeData& AttributeData, TMap<FString, FString>& OutMetaData)
+	{
+		FString AttributeName;
+		if (AttributeData.m_bTitleIsInt)
+		{
+			A3DUns32 UnsignedValue = 0;
+			memcpy(&UnsignedValue, AttributeData.m_pcTitle, sizeof(A3DUns32));
+			AttributeName = FString::Printf(TEXT("%u"), UnsignedValue);
+		}
+		else if (AttributeData.m_pcTitle && AttributeData.m_pcTitle[0] != '\0')
+		{
+			AttributeName = UTF8_TO_TCHAR(AttributeData.m_pcTitle);
+		}
+
+		for (A3DUns32 Index = 0; Index < AttributeData.m_uiSize; ++Index)
+		{
+			FString AttributeValue;
+			switch (AttributeData.m_asSingleAttributesData[Index].m_eType)
+			{
+			case kA3DModellerAttributeTypeTime:
+			case kA3DModellerAttributeTypeInt:
+			{
+				A3DInt32 Value;
+				memcpy(&Value, AttributeData.m_asSingleAttributesData[Index].m_pcData, sizeof(A3DInt32));
+				AttributeValue = FString::Printf(TEXT("%d"), Value);
+				break;
+			}
+
+			case kA3DModellerAttributeTypeReal:
+			{
+				A3DDouble Value;
+				memcpy(&Value, AttributeData.m_asSingleAttributesData[Index].m_pcData, sizeof(A3DDouble));
+				AttributeValue = FString::Printf(TEXT("%f"), Value);
+				break;
+			}
+
+			case kA3DModellerAttributeTypeString:
+			{
+				if (AttributeData.m_asSingleAttributesData[Index].m_pcData && AttributeData.m_asSingleAttributesData[Index].m_pcData[0] != '\0')
+				{
+					AttributeValue = UTF8_TO_TCHAR(AttributeData.m_asSingleAttributesData[Index].m_pcData);
+				}
+				break;
+			}
+
+			default:
+				break;
+			}
+
+			if (AttributeName.Len())
+			{
+				if (Index)
+				{
+					OutMetaData.Emplace(FString::Printf(TEXT("%s_%u"), *AttributeName, (int32)Index), AttributeValue);
+				}
+				else
+				{
+					OutMetaData.Emplace(AttributeName, AttributeValue);
+				}
+			}
+		}
+	}
+
+	void SetIOOption(A3DImport& Importer)
+	{
+		// A3DRWParamsGeneralData Importer.m_sGeneral
+		Importer.m_sLoadData.m_sGeneral.m_bReadSolids = A3D_TRUE;
+		Importer.m_sLoadData.m_sGeneral.m_bReadSurfaces = A3D_TRUE;
+		Importer.m_sLoadData.m_sGeneral.m_bReadWireframes = A3D_FALSE;
+		Importer.m_sLoadData.m_sGeneral.m_bReadPmis = A3D_FALSE;
+		Importer.m_sLoadData.m_sGeneral.m_bReadAttributes = A3D_TRUE;
+		Importer.m_sLoadData.m_sGeneral.m_bReadHiddenObjects = A3D_FALSE;
+		Importer.m_sLoadData.m_sGeneral.m_bReadConstructionAndReferences = A3D_FALSE;
+		Importer.m_sLoadData.m_sGeneral.m_bReadActiveFilter = A3D_FALSE;
+		Importer.m_sLoadData.m_sGeneral.m_eReadingMode2D3D = kA3DRead_3D;
+
+		Importer.m_sLoadData.m_sGeneral.m_eReadGeomTessMode = kA3DReadGeomOnly;
+		Importer.m_sLoadData.m_sGeneral.m_eDefaultUnit;
+		Importer.m_sLoadData.m_sGeneral.m_bReadFeature = A3D_FALSE;
+
+		Importer.m_sLoadData.m_sGeneral.m_bReadConstraints = A3D_FALSE;
+		Importer.m_sLoadData.m_sGeneral.m_iNbMultiProcess = 1;
+
+		Importer.m_sLoadData.m_sIncremental.m_bLoadNoDependencies = true;
+		Importer.m_sLoadData.m_sIncremental.m_bLoadStructureOnly = false;
+	}
+
+	void UpdateIOOptionAccordingToFormat(const CADLibrary::ECADFormat Format, A3DImport& Importer)
+	{
+		switch (Format)
+		{
+		case CADLibrary::ECADFormat::IGES:
+		{
+			break;
+		}
+
+		case CADLibrary::ECADFormat::CATIA:
+		{
+			break;
+		}
+
+		case CADLibrary::ECADFormat::INVENTOR:
+		case CADLibrary::ECADFormat::CATIA_3DXML:
+		{
+			Importer.m_sLoadData.m_sIncremental.m_bLoadNoDependencies = false;
+			break;
+		}
+
+		default:
+			break;
+		};
+	}
+
+
+}
+
+
+
+namespace CADLibrary
+{
+
+	ECADParsingResult FTechSoftFileParser::Process()
+	{
+		TechSoftInterface.Reset();
+		TechSoftInterface = GetTechSoftInterface();
+
+		// Process the file
+		ECADParsingResult Result = ECADParsingResult::ProcessOk;
+
+		const FFileDescriptor& File = CADFileData.GetCADFileDescription();
+
+		if (File.GetPathOfFileToLoad().IsEmpty())
+		{
+			return ECADParsingResult::FileNotFound;
+		}
+
+		A3DImport Import(*File.GetPathOfFileToLoad());
+
+		TechSoftFileParserImpl::SetIOOption(Import);
+
+		// Add specific options according to format
+		ECADFormat Format = File.GetFileFormat();
+		TechSoftFileParserImpl::UpdateIOOptionAccordingToFormat(Format, Import);
+
+		A3DStatus IRet = TechSoftInterface->Import(Import);
+		if (IRet != A3D_SUCCESS && IRet != A3D_LOAD_MULTI_MODELS_CADFILE && IRet != A3D_LOAD_MISSING_COMPONENTS)
+		{
+			return ECADParsingResult::ProcessFailed;
+		}
+
+		// save the file for the next load
+		if (CADFileData.IsCacheDefined())
+		{
+			FString CacheFilePath = CADFileData.GetCADCachePath();
+			if (CacheFilePath != File.GetPathOfFileToLoad())
+			{
+				// todo
+			}
+		}
+
+		if (CADFileData.GetImportParameters().GetStitchingTechnique() != StitchingNone && !FImportParameters::bGDisableCADKernelTessellation)
+		{
+			// todo
+			//TechSoftInterface->Repair(CADFileData.GetStitchingTechnique());
+		}
+
+		ReserveCADFileData();
+
+		ReadMaterialsAndColors();
+
+		return TraverseModel(TechSoftInterface->GetModelFile());
+	}
+
+	void FTechSoftFileParser::ReserveCADFileData()
+	{
+		CountUnderModel(TechSoftInterface->GetModelFile());
+
+		CADFileData.ReserveBodyMeshes(ComponentCount[EComponentType::Body]);
+
+		FArchiveSceneGraph& SceneGraphArchive = CADFileData.GetSceneGraphArchive();
+		SceneGraphArchive.Bodies.Reserve(ComponentCount[EComponentType::Body]);
+		SceneGraphArchive.Components.Reserve(ComponentCount[EComponentType::Reference]);
+		SceneGraphArchive.UnloadedComponents.Reserve(ComponentCount[EComponentType::Reference]);
+		SceneGraphArchive.ExternalReferences.Reserve(ComponentCount[EComponentType::Reference]);
+		SceneGraphArchive.Instances.Reserve(ComponentCount[EComponentType::Occurrence]);
+
+		SceneGraphArchive.CADIdToBodyIndex.Reserve(ComponentCount[EComponentType::Body]);
+		SceneGraphArchive.CADIdToComponentIndex.Reserve(ComponentCount[EComponentType::Reference]);
+		SceneGraphArchive.CADIdToUnloadedComponentIndex.Reserve(ComponentCount[EComponentType::Reference]);
+		SceneGraphArchive.CADIdToInstanceIndex.Reserve(ComponentCount[EComponentType::Occurrence]);
+
+		uint32 MaterialNum = CountMaterial() + CountColor();
+		SceneGraphArchive.MaterialHIdToMaterial.Reserve(MaterialNum);
+	}
+
+	ECADParsingResult FTechSoftFileParser::TraverseModel(const A3DAsmModelFile* ModelFile)
+	{
+		TechSoft::TUniqueTSObj<A3DAsmModelFileData, A3DAsmModelFile> ModelFileData(ModelFile, A3DAsmModelFileGet);
+		if (!ModelFileData.IsValid())
+		{
+			return ECADParsingResult::ProcessFailed;
+		}
+
+		ModellerType = (EModellerType)ModelFileData->m_eModellerType;
+		FileUnit = ModelFileData->m_dUnit;
+
+		FEntityMetaData MetaData;
+		TraverseMetaData(ModelFile, MetaData);
+		TraverseSpecificMetaData(ModelFile, MetaData);
+
+		for (uint32 Index = 0; Index < ModelFileData->m_uiPOccurrencesSize; ++Index)
+		{
+			if (IsConfigurationSet(ModelFileData->m_ppPOccurrences[Index]))
+			{
+				//TraverseConfigurationSet(ModelFileData->m_ppPOccurrences[Index], MetaData);
+			}
+			else
+			{
+				TraverseReference(ModelFileData->m_ppPOccurrences[Index]);
+			}
+		}
+
+		return ECADParsingResult::ProcessOk;
+	}
+
+	void FTechSoftFileParser::TraverseReference(const A3DAsmProductOccurrence* ReferencePtr)
+	{
+		FEntityMetaData MetaData;
+		TraverseMetaData(ReferencePtr, MetaData);
+
+		if (MetaData.bRemoved || !MetaData.bShow)
+		{
+			return;
+		}
+
+		TraverseSpecificMetaData(ReferencePtr, MetaData);
+		DefineEntityName(MetaData.MetaData, EComponentType::Reference);
+
+		TraverseMaterialProperties(ReferencePtr);
+
+		int32 ComponentId = LastEntityId++;
+		int32 Index = CADFileData.AddComponent(ComponentId);
+		FArchiveComponent& Component = CADFileData.GetComponentAt(Index);
+		Component.MetaData = MoveTemp(MetaData.MetaData);
+
+		TechSoft::TUniqueTSObj<A3DAsmProductOccurrenceData, A3DAsmProductOccurrence> ReferenceData(ReferencePtr, A3DAsmProductOccurrenceGet);
+		if (!ReferenceData.IsValid())
+		{
+			return;
+		}
+
+		for (uint32 OccurenceIndex = 0; OccurenceIndex < ReferenceData->m_uiPOccurrencesSize; ++OccurenceIndex)
+		{
+			int32 ChildrenId = TraverseOccurrence(ReferenceData->m_ppPOccurrences[OccurenceIndex]);
+			Component.Children.Add(ChildrenId);
+		}
+
+		if (ReferenceData->m_pPart)
+		{
+			int32 ChildrenId = TraversePartDefinition(ReferenceData->m_pPart);
+			Component.Children.Add(ChildrenId);
+		}
+	}
+
+
+	FArchiveInstance& FTechSoftFileParser::AddInstance(FEntityMetaData& InstanceMetaData)
+	{
+		FCadId InstanceId = LastEntityId++;
+		int32 InstanceIndex = CADFileData.AddInstance(InstanceId);
+		FArchiveInstance& Instance = CADFileData.GetInstanceAt(InstanceIndex);
+		Instance.MetaData = MoveTemp(InstanceMetaData.MetaData);
+		return Instance;
+	}
+
+	FArchiveComponent& FTechSoftFileParser::AddComponent(FEntityMetaData& ComponentMetaData, FArchiveInstance& Instance)
+	{
+		FCadId ComponentId = LastEntityId++;
+		int32 ComponentIndex = CADFileData.AddComponent(ComponentId);
+		FArchiveComponent& Prototype = CADFileData.GetComponentAt(ComponentIndex);
+		Prototype.MetaData = MoveTemp(ComponentMetaData.MetaData);
+
+		Instance.ReferenceNodeId = ComponentId;
+
+		return Prototype;
+	}
+
+	FArchiveUnloadedComponent& FTechSoftFileParser::AddUnloadedComponent(FEntityMetaData& ComponentMetaData, FArchiveInstance& Instance)
+	{
+		FCadId ComponentId = LastEntityId++;
+		int32 ComponentIndex = CADFileData.AddUnloadedComponent(ComponentId);
+		FArchiveUnloadedComponent& Component = CADFileData.GetUnloadedComponentAt(ComponentIndex);
+		Component.MetaData = MoveTemp(ComponentMetaData.MetaData);
+
+		Instance.bIsExternalReference = true;
+		Instance.ReferenceNodeId = ComponentId;
+
+		Instance.ExternalReference = ComponentMetaData.ExternalFile;
+
+		ECADFormat Format = Instance.ExternalReference.GetFileFormat();
+		if (Format == ECADFormat::SOLIDWORKS /*|| CADFileData.FileFormat() == ECADFormat::JT*/)
+		{
+			if (FString* ConfigurationName = ComponentMetaData.MetaData.Find(TEXT("ConfigurationName")))
+			{
+				Instance.ExternalReference.SetConfiguration(*ConfigurationName);
+			}
+		}
+
+		CADFileData.AddExternalRef(Instance.ExternalReference);
+
+		return Component;
+	}
+
+	FArchiveComponent& FTechSoftFileParser::AddOccurence(FEntityMetaData& InstanceMetaData, FCadId& OutComponentId)
+	{
+		FArchiveInstance& Instance = AddInstance(InstanceMetaData);
+		OutComponentId = Instance.ObjectId;
+		FEntityMetaData ReferenceMetaData;
+		FArchiveComponent& Prototype = AddComponent(ReferenceMetaData, Instance);
+		return Prototype;
+	}
+
+	FArchiveComponent& FTechSoftFileParser::AddOccurence(FEntityMetaData& InstanceMetaData, FEntityMetaData& ReferenceMetaData, FCadId& OutComponentId)
+	{
+		FArchiveInstance& Instance = AddInstance(InstanceMetaData);
+		OutComponentId = Instance.ObjectId;
+		FArchiveComponent& Prototype = AddComponent(ReferenceMetaData, Instance);
+		return Prototype;
+	}
+
+	FArchiveBody& FTechSoftFileParser::AddBody(FEntityMetaData& BodyMetaData)
+	{
+		FCadId BodyId = LastEntityId++;
+		int32 BodyIndex = CADFileData.AddBody(BodyId);
+		FArchiveBody& Body = CADFileData.GetBodyAt(BodyIndex);
+		Body.MetaData = MoveTemp(BodyMetaData.MetaData);
+
+		return Body;
+	}
+
+	FCadId FTechSoftFileParser::TraverseOccurrence(const A3DAsmProductOccurrence* OccurrencePtr)
+	{
+		FEntityMetaData InstanceMetaData;
+		TraverseMetaData(OccurrencePtr, InstanceMetaData);
+
+		if (InstanceMetaData.bRemoved || !InstanceMetaData.bShow)
+		{
+			return 0;
+		}
+
+		TraverseSpecificMetaData(OccurrencePtr, InstanceMetaData);
+		DefineEntityName(InstanceMetaData.MetaData, EComponentType::Occurrence);
+
+		TraverseMaterialProperties(OccurrencePtr);
+		TraverseLayer(OccurrencePtr);
+
+		TechSoft::TUniqueTSObj<A3DAsmProductOccurrenceData, A3DAsmProductOccurrence> OccurrenceData(OccurrencePtr, A3DAsmProductOccurrenceGet);
+		if (!OccurrenceData.IsValid())
+		{
+			return 0;
+		}
+
+		TraverseTransformation(OccurrenceData->m_pLocation);
+
+		FArchiveInstance& Instance = AddInstance(InstanceMetaData);
+
+		FEntityMetaData PrototypeMetaData;
+		bool bIsUnloadedPrototype = false;
+		if (OccurrenceData->m_pPrototype)
+		{
+			TraversePrototype(OccurrenceData->m_pPrototype, PrototypeMetaData);
+		}
+
+		if (PrototypeMetaData.bUnloaded)
+		{
+			FArchiveUnloadedComponent& Prototype = AddUnloadedComponent(PrototypeMetaData, Instance);
+		}
+		else
+		{
+			FArchiveComponent& Component = AddComponent(PrototypeMetaData, Instance);
+
+			for (uint32 Index = 0; Index < OccurrenceData->m_uiPOccurrencesSize; ++Index)
+			{
+				int32 ChildrenId = TraverseOccurrence(OccurrenceData->m_ppPOccurrences[Index]);
+				Component.Children.Add(ChildrenId);
+			}
+
+			if (OccurrenceData->m_pPart)
+			{
+				int32 ChildrenId = TraversePartDefinition(OccurrenceData->m_pPart);
+				Component.Children.Add(ChildrenId);
+			}
+		}
+
+		return Instance.ObjectId;
+	}
+
+	FFileDescriptor FTechSoftFileParser::GetOccurrenceFileName(const A3DAsmProductOccurrence* OccurrencePtr)
+	{
+		TechSoft::TUniqueTSObj<A3DAsmProductOccurrenceData, A3DAsmProductOccurrence> OccurrenceData(OccurrencePtr, A3DAsmProductOccurrenceGet);
+		if (!OccurrenceData.IsValid())
+		{
+			return FFileDescriptor();
+		}
+
+		if (OccurrenceData->m_pPrototype)
+		{
+			return GetOccurrenceFileName(OccurrenceData->m_pPrototype);
+		}
+
+		A3DUTF8Char* FilePathUTF8Ptr = nullptr;
+		int Ret = A3DAsmProductOccurrenceGetFilePathName(OccurrencePtr, &FilePathUTF8Ptr);
+		if (Ret == A3D_SUCCESS && FilePathUTF8Ptr)
+		{
+			FString OccurrenceFilePath = UTF8_TO_TCHAR(FilePathUTF8Ptr);
+			A3DAsmProductOccurrenceGetFilePathName(NULL, &FilePathUTF8Ptr);
+			FString OccurrenceFileName = FPaths::GetCleanFilename(OccurrenceFilePath);
+			if (OccurrenceFileName != CADFileData.GetCADFileDescription().GetFileName())
+			{
+				return FFileDescriptor(*OccurrenceFilePath, nullptr, *CADFileData.GetCADFileDescription().GetRootFolder());
+			}
+		}
+
+		FilePathUTF8Ptr = nullptr;
+		Ret = A3DAsmProductOccurrenceGetOriginalFilePathName(OccurrencePtr, &FilePathUTF8Ptr);
+		if (Ret == A3D_SUCCESS && FilePathUTF8Ptr)
+		{
+			FString OriginalFilePath = UTF8_TO_TCHAR(FilePathUTF8Ptr);
+			A3DAsmProductOccurrenceGetFilePathName(NULL, &FilePathUTF8Ptr);
+			FString OriginalFileName = FPaths::GetCleanFilename(OriginalFilePath);
+			if (OriginalFileName != CADFileData.GetCADFileDescription().GetFileName())
+			{
+				return FFileDescriptor(*OriginalFilePath, nullptr, *CADFileData.GetCADFileDescription().GetRootFolder());
+			}
+		}
+
+		return FFileDescriptor();
+	}
+
+	void FTechSoftFileParser::TraversePrototype(const A3DAsmProductOccurrence* PrototypePtr, FEntityMetaData& PrototypeMetaData)
+	{
+		TechSoft::TUniqueTSObj<A3DAsmProductOccurrenceData, A3DAsmProductOccurrence> PrototypeData(PrototypePtr, A3DAsmProductOccurrenceGet);
+		if (!PrototypeData.IsValid())
+		{
+			return;
+		}
+
+		TraverseMetaData(PrototypePtr, PrototypeMetaData);
+		TraverseSpecificMetaData(PrototypePtr, PrototypeMetaData);
+		DefineEntityName(PrototypeMetaData.MetaData, EComponentType::Reference);
+
+		TraverseMaterialProperties(PrototypePtr);
+
+		TraverseTransformation(PrototypeData->m_pLocation);
+
+		PrototypeMetaData.bUnloaded = TechSoftFileParserImpl::IsUnloadedPrototype(*PrototypeData);
+		if (PrototypeMetaData.bUnloaded)
+		{
+			PrototypeMetaData.ExternalFile = GetOccurrenceFileName(PrototypePtr);
+		}
+	}
+
+	FCadId FTechSoftFileParser::TraversePartDefinition(const A3DAsmPartDefinition* PartDefinitionPtr)
+	{
+		FEntityMetaData PartMetaData;
+		TraverseMetaData(PartDefinitionPtr, PartMetaData);
+
+		if (PartMetaData.bRemoved || !PartMetaData.bShow)
+		{
+			return 0;
+		}
+
+		TraverseSpecificMetaData(PartDefinitionPtr, PartMetaData);
+		DefineEntityName(PartMetaData.MetaData, EComponentType::Occurrence);
+
+		TraverseMaterialProperties(PartDefinitionPtr);
+		TraverseLayer(PartDefinitionPtr);
+
+		TechSoft::TUniqueTSObj<A3DAsmPartDefinitionData, A3DAsmPartDefinition> PartData(PartDefinitionPtr, A3DAsmPartDefinitionGet, TechSoftFileParserImpl::InitializeData);
+		if (!PartData.IsValid())
+		{
+			return 0;
+		}
+
+		FCadId PartId = 0;
+		FArchiveComponent& Part = AddOccurence(PartMetaData, PartId);
+
+		for (unsigned int Index = 0; Index < PartData->m_uiRepItemsSize; ++Index)
+		{
+			int32 ChildId = TraverseRepresentationItem(PartData->m_ppRepItems[Index]);
+			Part.Children.Add(ChildId);
+		}
+
+		return PartId;
+	}
+
+	FCadId FTechSoftFileParser::TraverseRepresentationItem(A3DRiRepresentationItem* RepresentationItem)
+	{
+		A3DEEntityType Type;
+		A3DEntityGetType(RepresentationItem, &Type);
+
+		switch (Type)
+		{
+		case kA3DTypeRiSet:
+			return TraverseRepresentationSet(RepresentationItem);
+		case kA3DTypeRiBrepModel:
+			return TraverseBRepModel(RepresentationItem);
+		case kA3DTypeRiPolyBrepModel:
+			return TraversePolyBRepModel(RepresentationItem);
+		default:
+			break;
+		}
+		return 0;
+	}
+
+	FCadId FTechSoftFileParser::TraverseRepresentationSet(const A3DRiSet* RepresentationSetPtr)
+	{
+		TechSoft::TUniqueTSObj<A3DRiSetData, A3DRiSet> RepresentationSetData(RepresentationSetPtr, A3DRiSetGet);
+		if (!RepresentationSetData.IsValid())
+		{
+			return 0;
+		}
+
+		FEntityMetaData RepresentationSetMetaData;
+		TraverseMetaData(RepresentationSetPtr, RepresentationSetMetaData);
+
+		if (RepresentationSetMetaData.bRemoved || !RepresentationSetMetaData.bShow)
+		{
+			return 0;
+		}
+
+		TraverseMaterialProperties(RepresentationSetPtr);
+
+		FCadId RepresentationSetId = 0;
+		FArchiveComponent& RepresentationSet = AddOccurence(RepresentationSetMetaData, RepresentationSetId);
+
+		for (A3DUns32 ui = 0; ui < RepresentationSetData->m_uiRepItemsSize; ++ui)
+		{
+			int32 ChildId = TraverseRepresentationItem(RepresentationSetData->m_ppRepItems[ui]);
+			RepresentationSet.Children.Add(ChildId);
+		}
+		return RepresentationSetId;
+	}
+
+	FCadId FTechSoftFileParser::TraverseBRepModel(A3DRiBrepModel* BRepModelPtr)
+	{
+		TechSoft::TUniqueTSObj<A3DRiBrepModelData, A3DRiBrepModel> BodyData(BRepModelPtr, A3DRiBrepModelGet);
+		if (!BodyData.IsValid())
+		{
+			return 0;
+		}
+
+		FEntityMetaData BRepMetaData;
+		TraverseMetaData(BRepModelPtr, BRepMetaData);
+		if (!BRepMetaData.bShow || BRepMetaData.bRemoved)
+		{
+			return 0;
+		}
+
+		TraverseSpecificMetaData(BRepModelPtr, BRepMetaData);
+		TraverseMaterialProperties(BRepModelPtr);
+
+		FArchiveBody& Body = AddBody(BRepMetaData);
+
+		TraverseRepresentationContent(BRepModelPtr, Body);
+		if (FImportParameters::bGDisableCADKernelTessellation)
+		{
+			MeshRepresentationWithTechSoft(BRepModelPtr, Body);
+		}
+		else
+		{
+			// Mesh with CADKernel
+		}
+		return Body.ObjectId;
+	}
+
+	void FTechSoftFileParser::TraverseRepresentationContent(const A3DRiRepresentationItem* RepresentationItemPtr, FArchiveBody& Body)
+	{
+		TechSoft::TUniqueTSObj<A3DRiRepresentationItemData, A3DRiRepresentationItem> RepresentationItemData(RepresentationItemPtr, A3DRiRepresentationItemGet);
+		if (!RepresentationItemData.IsValid())
+		{
+			return;
+		}
+
+		if (RepresentationItemData->m_pCoordinateSystem)
+		{
+			TraverseCoordinateSystem(RepresentationItemData->m_pCoordinateSystem);
+		}
+
+		if (RepresentationItemData->m_pTessBase)
+		{
+			TraverseTessellationBase(RepresentationItemData->m_pTessBase, Body);
+		}
+	}
+
+	FCadId FTechSoftFileParser::TraversePolyBRepModel(const A3DRiPolyBrepModel* PolygonalPtr)
+	{
+		TechSoft::TUniqueTSObj<A3DRiPolyBrepModelData, A3DRiPolyBrepModel> BodyData(PolygonalPtr, A3DRiPolyBrepModelGet);
+		if (!BodyData.IsValid())
+		{
+			return 0;
+		}
+
+		FEntityMetaData BRepMetaData;
+		TraverseMetaData(PolygonalPtr, BRepMetaData);
+		if (!BRepMetaData.bShow || BRepMetaData.bRemoved)
+		{
+			return 0;
+		}
+
+		TraverseSpecificMetaData(PolygonalPtr, BRepMetaData);
+		TraverseMaterialProperties(PolygonalPtr);
+
+		FArchiveBody& Body = AddBody(BRepMetaData);
+		TraverseRepresentationContent(PolygonalPtr, Body);
+
+		return Body.ObjectId;
+	}
+
+	void FTechSoftFileParser::TraverseMetaData(const A3DEntity* Entity, FEntityMetaData& OutMetaData)
+	{
+		TechSoft::TUniqueTSObj<A3DRootBaseData, A3DEntity> MetaData(Entity, A3DRootBaseGet);
+		if (MetaData.IsValid())
+		{
+			if (false && MetaData->m_uiPersistentId > 0)
+			{
+				FString PersistentId = FString::FromInt(MetaData->m_uiPersistentId);
+				OutMetaData.MetaData.Emplace(TEXT("PersistentId"), PersistentId);
+			}
+
+			if (false && MetaData->m_pcPersistentId)
+			{
+				FString PersistentUuid = UTF8_TO_TCHAR(MetaData->m_pcPersistentId);
+				TechSoftFileParserImpl::RemoveUnwantedChar(PersistentUuid, TEXT('-'));
+				OutMetaData.MetaData.Emplace(TEXT("UUID"), PersistentUuid);
+			}
+
+			if (MetaData->m_pcName && MetaData->m_pcName[0] != '\0')
+			{
+				OutMetaData.MetaData.Find(TEXT("TechSoftName"));
+			}
+
+			TechSoft::TUniqueTSObj<A3DMiscAttributeData, A3DMiscAttribute> AttributeData(A3DMiscAttributeGet);
+			for (A3DUns32 Index = 0; Index < MetaData->m_uiSize; ++Index)
+			{
+				AttributeData.Get(MetaData->m_ppAttributes[Index]);
+				if (AttributeData.IsValid())
+				{
+					TechSoftFileParserImpl::TraverseAttribute(*AttributeData, OutMetaData.MetaData);
+				}
+			}
+		}
+
+		if (A3DEntityIsBaseWithGraphicsType(Entity))
+		{
+			TechSoft::TUniqueTSObj<A3DRootBaseWithGraphicsData, A3DEntity> MetaDataWithGraphics(Entity, A3DRootBaseWithGraphicsGet);
+			if (MetaDataWithGraphics.IsValid())
+			{
+				if (MetaDataWithGraphics->m_pGraphics != NULL)
+				{
+					FEntityBehaviour GraphicsBehaviour;
+					TraverseGraphics(MetaDataWithGraphics->m_pGraphics, GraphicsBehaviour);
+					OutMetaData.bRemoved = GraphicsBehaviour.bRemoved;
+					OutMetaData.bShow = GraphicsBehaviour.bShow;
+				}
+			}
+		}
+	}
+
+	FString FTechSoftFileParser::DefineEntityName(TMap<FString, FString>& MetaData, EComponentType EntityType)
+	{
+		if (MetaData.IsEmpty())
+		{
+			return FString();
+		}
+
+		ECADFormat Format = CADFileData.GetCADFileDescription().GetFileFormat();
+
+		switch (EntityType)
+		{
+		case EComponentType::Occurrence:
+		{
+			FString* Name = MetaData.Find(TEXT("Name"));
+			if (Name)
+			{
+				return *Name;
+			}
+
+			FString* InstanceName = MetaData.Find(TEXT("InstanceName"));
+			if (InstanceName)
+			{
+				MetaData.Emplace(TEXT("Name"), *InstanceName);
+				return *InstanceName;
+			}
+
+			FString* OriginalName = MetaData.Find(TEXT("OriginalName"));
+			FString* TechSoftName = MetaData.Find(TEXT("TechSoftName"));
+
+			if (!OriginalName && TechSoftName)
+			{
+				FString PartNumberStr = TechSoftFileParserImpl::GetOriginalName(*TechSoftName);
+
+				switch (Format)
+				{
+				case ECADFormat::CATIA_3DXML:
+					PartNumberStr = TechSoftFileParserImpl::Get3dxmlOriginalName(PartNumberStr);
+					break;
+
+				case ECADFormat::SOLIDWORKS:
+					PartNumberStr = TechSoftFileParserImpl::GetSWOriginalName(PartNumberStr);
+					break;
+
+				case ECADFormat::JT:
+					PartNumberStr = TechSoftFileParserImpl::GetJTOriginalName(PartNumberStr);
+					break;
+
+				default:
+					break;
+				}
+
+				MetaData.Emplace(TEXT("PartNumber"), PartNumberStr);
+			}
+
+			FString* PartNumber = MetaData.Find(TEXT("PartNumber"));
+			if (PartNumber)
+			{
+				MetaData.Emplace(TEXT("Name"), *PartNumber);
+			}
+			else if (OriginalName)
+			{
+				MetaData.Emplace(TEXT("Name"), *OriginalName);
+			}
+			return MetaData[TEXT("Name")];
+		}
+		break;
+
+		default:
+		{
+			FString* Name = MetaData.Find(TEXT("Name"));
+			FString* OriginalName = MetaData.Find(TEXT("OriginalName"));
+			FString* PartNumber = MetaData.Find(TEXT("PartNumber"));
+			FString* TechSoftName = MetaData.Find(TEXT("TechSoftName"));
+
+			if (!OriginalName && !PartNumber && TechSoftName)
+			{
+				FString OriginalNameStr = TechSoftFileParserImpl::GetOriginalName(*TechSoftName);
+
+				if (EntityType == EComponentType::Body)
+				{
+					switch (Format)
+					{
+					case ECADFormat::CATIA:
+						OriginalNameStr = TechSoftFileParserImpl::GetCatiaBodyOriginalName(OriginalNameStr);
+						break;
+
+					case ECADFormat::CREO:
+						OriginalNameStr = TechSoftFileParserImpl::GetCreoBodyOriginalName(OriginalNameStr);
+						break;
+
+					default:
+						break;
+					}
+				}
+
+				MetaData.Emplace(TEXT("OriginalName"), OriginalNameStr);
+				OriginalName = MetaData.Find(TEXT("OriginalName"));
+			}
+
+			if (PartNumber)
+			{
+				MetaData.Emplace(TEXT("Name"), *PartNumber);
+			}
+			else if (OriginalName)
+			{
+				MetaData.Emplace(TEXT("Name"), *OriginalName);
+			}
+
+			FString OutName = MetaData[TEXT("Name")];
+			if (!OriginalName)
+			{
+				MetaData.Emplace(TEXT("OriginalName"), OutName);
+			}
+			return OutName;
+		}
+		break;
+		}
+
+		return FString();
+	}
+
+	void FTechSoftFileParser::TraverseSpecificMetaData(const A3DAsmProductOccurrence* Occurrence, FEntityMetaData& OutMetaData)
+	{
+		//----------- Export Specific information per CAD format -----------
+		switch (ModellerType)
+		{
+		case kA3DModellerSlw:
+		{
+			TechSoft::TUniqueTSObj<A3DAsmProductOccurrenceDataSLW, A3DAsmProductOccurrence> SolidWorksSpecificData(Occurrence, A3DAsmProductOccurrenceGetSLW);
+			if (SolidWorksSpecificData.IsValid())
+			{
+				if (SolidWorksSpecificData->m_psCfgName)
+				{
+					FString ConfigurationName = UTF8_TO_TCHAR(SolidWorksSpecificData->m_psCfgName);
+					OutMetaData.MetaData.Emplace(TEXT("ConfigurationName"), ConfigurationName);
+					FString ConfigurationIndex = FString::FromInt(SolidWorksSpecificData->m_iIndexCfg);
+					OutMetaData.MetaData.Emplace(TEXT("ConfigurationIndex"), ConfigurationIndex);
+				}
+			}
+			break;
+		}
+		case kA3DModellerUnigraphics:
+		{
+			TechSoft::TUniqueTSObj<A3DAsmProductOccurrenceDataUg, A3DAsmProductOccurrence> UnigraphicsSpecificData(Occurrence, A3DAsmProductOccurrenceGetUg);
+			if (UnigraphicsSpecificData.IsValid())
+			{
+				if (UnigraphicsSpecificData->m_psPartUID)
+				{
+					FString PartUID = UTF8_TO_TCHAR(UnigraphicsSpecificData->m_psPartUID);
+					OutMetaData.MetaData.Emplace(TEXT("UnigraphicsPartUID"), PartUID);
+				}
+				if (UnigraphicsSpecificData->m_psInstanceFileName)
+				{
+					FString InstanceFileName = UTF8_TO_TCHAR(UnigraphicsSpecificData->m_psInstanceFileName);
+					OutMetaData.MetaData.Emplace(TEXT("UnigraphicsPartUID"), InstanceFileName);
+				}
+
+				if (UnigraphicsSpecificData->m_uiInstanceTag)
+				{
+					FString InstanceTag = FString::FromInt(UnigraphicsSpecificData->m_uiInstanceTag);
+					OutMetaData.MetaData.Emplace(TEXT("UnigraphicsInstanceTag"), InstanceTag);
+				}
+			}
+			break;
+		}
+
+		case kA3DModellerCatiaV5:
+		{
+			TechSoft::TUniqueTSObj<A3DAsmProductOccurrenceDataCV5, A3DAsmProductOccurrence> CatiaV5SpecificData(Occurrence, A3DAsmProductOccurrenceGetCV5);
+			if (CatiaV5SpecificData.IsValid())
+			{
+				if (CatiaV5SpecificData->m_psVersion)
+				{
+					FString Version = UTF8_TO_TCHAR(CatiaV5SpecificData->m_psVersion);
+					OutMetaData.MetaData.Emplace(L"CatiaVersion", Version);
+				}
+
+				if (CatiaV5SpecificData->m_psPartNumber)
+				{
+					FString PartNumber = UTF8_TO_TCHAR(CatiaV5SpecificData->m_psPartNumber);
+					OutMetaData.MetaData.Emplace(L"CatiaPartNumber", PartNumber);
+				}
+			}
+			break;
+		}
+
+		default:
+			break;
+		}
+	}
+
+	void FTechSoftFileParser::TraverseGraphics(const A3DGraphics* Graphics, FEntityBehaviour& GraphicsBehaviour)
+	{
+		// TODO
+
+		TechSoft::TUniqueTSObj<A3DGraphicsData, A3DGraphics> GraphicsData(Graphics, A3DGraphicsGet);
+		if (!GraphicsData.IsValid())
+		{
+			return;
+		}
+
+		GraphicsBehaviour.bFatherHeritColor = GraphicsData->m_usBehaviour & kA3DGraphicsFatherHeritColor;
+		GraphicsBehaviour.bFatherHeritLayer = GraphicsData->m_usBehaviour & kA3DGraphicsFatherHeritLayer;
+		GraphicsBehaviour.bFatherHeritLinePattern = GraphicsData->m_usBehaviour & kA3DGraphicsFatherHeritLinePattern;
+		GraphicsBehaviour.bFatherHeritLineWidth = GraphicsData->m_usBehaviour & kA3DGraphicsFatherHeritLineWidth;
+		GraphicsBehaviour.bFatherHeritShow = GraphicsData->m_usBehaviour & kA3DGraphicsFatherHeritShow;
+		GraphicsBehaviour.bFatherHeritTransparency = GraphicsData->m_usBehaviour & kA3DGraphicsFatherHeritTransparency;
+		GraphicsBehaviour.bRemoved = GraphicsData->m_usBehaviour & kA3DGraphicsRemoved;
+		GraphicsBehaviour.bShow = GraphicsData->m_usBehaviour & kA3DGraphicsShow;
+		GraphicsBehaviour.bSonHeritColor = GraphicsData->m_usBehaviour & kA3DGraphicsSonHeritColor;
+		GraphicsBehaviour.bSonHeritLayer = GraphicsData->m_usBehaviour & kA3DGraphicsSonHeritLayer;
+		GraphicsBehaviour.bSonHeritLinePattern = GraphicsData->m_usBehaviour & kA3DGraphicsSonHeritLinePattern;
+		GraphicsBehaviour.bSonHeritLineWidth = GraphicsData->m_usBehaviour & kA3DGraphicsSonHeritLineWidth;
+		GraphicsBehaviour.bSonHeritShow = GraphicsData->m_usBehaviour & kA3DGraphicsSonHeritShow;
+		GraphicsBehaviour.bSonHeritTransparency = GraphicsData->m_usBehaviour & kA3DGraphicsSonHeritTransparency;
+	}
+
+	void FTechSoftFileParser::TraverseMaterialProperties(const A3DEntity* Entity)
+	{
+		// TODO
+
+		TechSoft::TUniqueTSObj<A3DMiscMaterialPropertiesData, A3DEntity> MaterialPropertiesData(Entity, A3DMiscGetMaterialProperties);
+		if (!MaterialPropertiesData.IsValid())
+		{
+			return;
+		}
+
+		//xml->SetDoubleAttribute("m_dDensity", MaterialPropertiesData.m_dDensity);
+		//xml->SetAttribute("m_pcMaterialName", MaterialPropertiesData.m_pcMaterialName ? MaterialPropertiesData.m_pcMaterialName : "NULL");
+		//xml->SetAttribute("m_ePhysicType", MaterialPropertiesData.m_ePhysicType);
+		switch (MaterialPropertiesData->m_ePhysicType)
+		{
+		case A3DPhysicType_None:
+		case A3DPhysicType_Fiber:
+		case A3DPhysicType_HoneyComb:
+		case A3DPhysicType_Isotropic:
+		case A3DPhysicType_Orthotropic2D:
+		case A3DPhysicType_Orthotropic3D:
+		case A3DPhysicType_Anisotropic:
+
+		default:
+			break;
+		}
+	}
+
+	void FTechSoftFileParser::TraverseTransformation3D(const A3DMiscCartesianTransformation* CartesianTransformation)
+	{
+		// TODO
+
+		if (CartesianTransformation)
+		{
+			TechSoft::TUniqueTSObj<A3DMiscCartesianTransformationData, A3DEntity> CartesianTransformationData(CartesianTransformation, A3DMiscCartesianTransformationGet);
+			if (CartesianTransformationData.IsValid())
+			{
+				//trsf->SetAttribute("m_ucBehaviour", (int)(CartesianTransformationData.m_ucBehaviour));
+				//"m_sOrigin", CartesianTransformationData.m_sOrigin, trsf);
+				//"m_sXVector", CartesianTransformationData.m_sXVector, trsf);
+				//"m_sYVector", CartesianTransformationData.m_sYVector, trsf);
+				//"m_sScale", CartesianTransformationData.m_sScale, trsf);
+			}
+		}
+	}
+
+	void FTechSoftFileParser::TraverseGeneralTransformation(const A3DMiscGeneralTransformation* GeneralTransformation)
+	{
+		// TODO
+
+		TechSoft::TUniqueTSObj<A3DMiscGeneralTransformationData, A3DMiscGeneralTransformation> GeneralTransformationData(GeneralTransformation, A3DMiscGeneralTransformationGet);
+		if (GeneralTransformationData.IsValid())
+		{
+			//traverseDoubles("m_adCoeff", 16, GeneralTransformationData.m_adCoeff, trsf);
+		}
+	}
+
+	void FTechSoftFileParser::TraverseTransformation(const A3DMiscCartesianTransformation* Transformation3D)
+	{
+		if (Transformation3D == NULL)
+		{
+			return;
+		}
+
+		A3DEEntityType Type = kA3DTypeUnknown;
+		A3DEntityGetType(Transformation3D, &Type);
+
+		if (Type == kA3DTypeMiscCartesianTransformation)
+		{
+			TraverseTransformation3D(Transformation3D);
+		}
+		else if (Type == kA3DTypeMiscGeneralTransformation)
+		{
+			TraverseGeneralTransformation(Transformation3D);
+		}
+	}
+
+	void FTechSoftFileParser::TraverseCoordinateSystem(const A3DRiCoordinateSystem* CoordinateSystem)
+	{
+		TechSoft::TUniqueTSObj<A3DRiCoordinateSystemData, A3DRiCoordinateSystem> CoordinateSystemData(CoordinateSystem, A3DRiCoordinateSystemGet);
+		if (CoordinateSystemData.IsValid())
+		{
+			TraverseTransformation3D(CoordinateSystemData->m_pTransformation);
+		}
+	}
+
+	void FTechSoftFileParser::TraverseLayer(const A3DAsmProductOccurrence* Occurrence)
+	{
+		// TODO
+
+		A3DUns32 LayerCount = 0;
+		A3DAsmLayer* AsmLayer = 0;
+		if (A3DAsmProductOccurrenceGetLayerList(Occurrence, &LayerCount, &AsmLayer) == A3D_SUCCESS)
+		{
+			if (LayerCount)
+			{
+				for (A3DUns32 Index = 0; Index < LayerCount; ++Index)
+				{
+					//Layer->SetAttribute("Name", asLayers[Index].m_pcLayerName ? asLayers[Index].m_pcLayerName : "null");
+					//Layer->SetAttribute("Layer", asLayers[Index].m_usLayer);
+				}
+			}
+			A3DAsmProductOccurrenceGetLayerList(NULL, &LayerCount, &AsmLayer);
+		}
+	}
+
+	bool FTechSoftFileParser::IsConfigurationSet(const A3DAsmProductOccurrence* Occurrence)
+	{
+		TechSoft::TUniqueTSObj<A3DAsmProductOccurrenceData, A3DAsmProductOccurrence> OccurrenceData(Occurrence, A3DAsmProductOccurrenceGet);
+		if (!OccurrenceData.IsValid())
+		{
+			return false;
+		}
+
+		bool bIsConfiguration = false;
+		if (OccurrenceData->m_uiPOccurrencesSize)
+		{
+			TechSoft::TUniqueTSObj<A3DAsmProductOccurrenceData, A3DAsmProductOccurrence> ChildData(A3DAsmProductOccurrenceGet);
+			for (uint32 Index = 0; Index < OccurrenceData->m_uiPOccurrencesSize; ++Index)
+			{
+				if (ChildData.Get(OccurrenceData->m_ppPOccurrences[Index]) == A3D_SUCCESS)
+				{
+					if (ChildData->m_uiProductFlags & A3D_PRODUCT_FLAG_CONFIG)
+					{
+						bIsConfiguration = true;
+					}
+					break;
+				}
+			}
+		}
+		return bIsConfiguration;
+	}
+
+	void FTechSoftFileParser::CountUnderModel(const A3DAsmModelFile* AsmModel)
+	{
+		TechSoft::TUniqueTSObj<A3DAsmModelFileData, A3DAsmModelFile> ModelFileData(AsmModel, A3DAsmModelFileGet);
+		if (!ModelFileData.IsValid())
+		{
+			return;
+		}
+
+		ComponentCount[EComponentType::Occurrence] ++;
+
+		for (uint32 Index = 0; Index < ModelFileData->m_uiPOccurrencesSize; ++Index)
+		{
+			//if (IsConfigurationSet(ModelFileData.m_ppPOccurrences[Index]))
+			//{
+			//	TraverseConfigurationSet(ModelFileData.m_ppPOccurrences[Index], MetaData);
+			//}
+			//else
+			{
+				CountUnderOccurrence(ModelFileData->m_ppPOccurrences[Index]);
+			}
+		}
+	}
+
+	void FTechSoftFileParser::CountUnderOccurrence(const A3DAsmProductOccurrence* Occurrence)
+	{
+		TechSoft::TUniqueTSObj<A3DAsmProductOccurrenceData, A3DAsmProductOccurrence> OccurrenceData(Occurrence, A3DAsmProductOccurrenceGet);
+		if (!OccurrenceData.IsValid())
+		{
+			return;
+		}
+
+		ComponentCount[EComponentType::Occurrence] ++;
+
+		if (OccurrenceData->m_pPrototype)
+		{
+			CountUnderPrototype(OccurrenceData->m_pPrototype);
+		}
+
+		for (uint32 Index = 0; Index < OccurrenceData->m_uiPOccurrencesSize; ++Index)
+		{
+			CountUnderOccurrence(OccurrenceData->m_ppPOccurrences[Index]);
+		}
+
+		if (OccurrenceData->m_pPart)
+		{
+			CountUnderPartDefinition(OccurrenceData->m_pPart);
+		}
+	}
+
+	void FTechSoftFileParser::CountUnderPrototype(const A3DAsmProductOccurrence* Prototype)
+	{
+		if (PrototypeCounted.Contains(Prototype))
+		{
+			return;
+		}
+		PrototypeCounted.Add(Prototype);
+
+		TechSoft::TUniqueTSObj<A3DAsmProductOccurrenceData, A3DAsmProductOccurrence> PrototypeData(Prototype, A3DAsmProductOccurrenceGet);
+		if (!PrototypeData.IsValid())
+		{
+			return;
+		}
+
+		ComponentCount[EComponentType::Reference] ++;
+
+		if (PrototypeData->m_pPrototype)
+		{
+			CountUnderPrototype(PrototypeData->m_pPrototype);
+		}
+	}
+
+	void FTechSoftFileParser::CountUnderPartDefinition(const A3DAsmPartDefinition* PartDefinition)
+	{
+		TechSoft::TUniqueTSObj<A3DAsmPartDefinitionData, A3DAsmPartDefinition> PartData(PartDefinition, A3DAsmPartDefinitionGet);
+		if (!PartData.IsValid())
+		{
+			return;
+		}
+
+		for (unsigned int Index = 0; Index < PartData->m_uiRepItemsSize; ++Index)
+		{
+			CountUnderRepresentation(PartData->m_ppRepItems[Index]);
+		}
+
+	}
+
+	void FTechSoftFileParser::CountUnderRepresentation(const A3DRiRepresentationItem* RepresentationItem)
+	{
+		A3DEEntityType Type;
+		A3DEntityGetType(RepresentationItem, &Type);
+
+		switch (Type)
+		{
+		case kA3DTypeRiSet:
+			CountUnderRepresentationSet(RepresentationItem);
+			break;
+		case kA3DTypeRiBrepModel:
+		case kA3DTypeRiPolyBrepModel:
+			ComponentCount[EComponentType::Body] ++;
+			break;
+
+		default:
+			break;
+		}
+	}
+
+	void FTechSoftFileParser::CountUnderRepresentationSet(const A3DRiSet* RepresentationSet)
+	{
+		TechSoft::TUniqueTSObj<A3DRiSetData, A3DRiSet> RepresentationSetData(RepresentationSet, A3DRiSetGet);
+		if (!RepresentationSetData.IsValid())
+		{
+			return;
+		}
+		ComponentCount[EComponentType::Body] += RepresentationSetData->m_uiRepItemsSize;
+	}
+
+	int32 FTechSoftFileParser::CountMaterial()
+	{
+		// TODO
+
+		return 0;
+	}
+
+	int32 FTechSoftFileParser::CountColor()
+	{
+		// TODO
+
+		return 0;
+	}
+
+	void FTechSoftFileParser::ReadMaterialsAndColors()
+	{
+		// TODO
+	}
+
+	void FTechSoftFileParser::MeshRepresentationWithTechSoft(A3DRiRepresentationItem* RepresentationItemPtr, FArchiveBody& Body)
+	{
+		TechSoft::TUniqueTSObj<A3DRiRepresentationItemData, A3DEntity> RepresentationItemData(A3DRiRepresentationItemGet);
+
+		// TUniqueTechSoftObj does not work in this case
+		A3DRWParamsTessellationData TessellationParameters;
+		A3D_INITIALIZE_DATA(A3DRWParamsTessellationData, TessellationParameters);
+
+		if (true)
+		{
+			TessellationParameters.m_eTessellationLevelOfDetail = kA3DTessLODMedium;		/*!< Enum to specify predefined values for some following members. */
+		}
+		else
+		{
+			// TODO
+			// Check unity conversion
+			ensure(false);
+
+			const FImportParameters& ImportParameters = CADFileData.GetImportParameters();
+			////A3DRWParamsTessellationData /*!< The tessellation reading parameters. */
+			TessellationParameters.m_eTessellationLevelOfDetail = kA3DTessLODUserDefined;		/*!< Enum to specify predefined values for some following members. */
+			TessellationParameters.m_bUseHeightInsteadOfRatio = A3D_TRUE;
+			TessellationParameters.m_dMaxChordHeight = ImportParameters.GetChordTolerance();
+			TessellationParameters.m_dAngleToleranceDeg = ImportParameters.GetMaxNormalAngle();
+			TessellationParameters.m_dMaximalTriangleEdgeLength = 0; //ImportParameters.MaxEdgeLength;
+
+			TessellationParameters.m_bAccurateTessellation = A3D_FALSE;  // A3D_FALSE' indicates the tessellation is set for visualization
+			TessellationParameters.m_bAccurateTessellationWithGrid = A3D_FALSE; /*!< Enable accurate tessellation with faces inner points on a grid.*/
+			TessellationParameters.m_dAccurateTessellationWithGridMaximumStitchLength = 0; 	/*!< Maximal grid stitch length. Disabled if value is 0. Be careful, a too small value can generate a huge tessellation. */
+		}
+
+		TessellationParameters.m_bKeepUVPoints = A3D_TRUE; /*!< Keep parametric points as texture points. */
+
+		// Get the tesselation
+		A3DStatus Status = A3DRiRepresentationItemComputeTessellation(RepresentationItemPtr, &TessellationParameters);
+		Status = A3DRiRepresentationItemGet(RepresentationItemPtr, RepresentationItemData.GetEmptyDataPtr());
+		TraverseTessellationBase(RepresentationItemData->m_pTessBase, Body);
+	}
+
+	void FTechSoftFileParser::TraverseTessellationBase(const A3DTessBase* Tessellation, FArchiveBody& Body)
+	{
+		A3DEEntityType eType;
+		if (A3DEntityGetType(Tessellation, &eType) == A3D_SUCCESS)
+		{
+			switch (eType)
+			{
+			case kA3DTypeTess3D:
+				TraverseTessellation3D(Tessellation, Body);
+				break;
+			case kA3DTypeTess3DWire:
+			case kA3DTypeTessMarkup:
+			default:
+				break;
+			}
+		}
+	}
+
+	void AddFaceTriangleWithUniqueNormal(FTessellationData& Tessellation, const A3DTess3DData& Tessellation3DData, long& InOutStartIndex, unsigned long InTrinangleCount)
+	{
+		TechSoftFileParserImpl::Reserve(Tessellation, InTrinangleCount, /*bWithTexture*/ false);
+
+		unsigned long Index = InOutStartIndex;
+
+		int32 FaceIndex[3] = { 0, 0, 0 };
+		int32 NormalIndex[3] = { 0, 0, 0 };
+		int32 VertexIndex = 0;
+
+		NormalIndex[0] = Tessellation3DData.m_puiTriangulatedIndexes[Index++];
+		NormalIndex[1] = NormalIndex[0];
+		NormalIndex[2] = NormalIndex[0];
+
+		// Get Triangles
+		for (unsigned long TriangleIndex = 0; TriangleIndex < InTrinangleCount; TriangleIndex++)
+		{
+			FaceIndex[0] = Tessellation3DData.m_puiTriangulatedIndexes[Index++];
+			FaceIndex[1] = Tessellation3DData.m_puiTriangulatedIndexes[Index++];
+			FaceIndex[2] = Tessellation3DData.m_puiTriangulatedIndexes[Index++];
+
+			if (!TechSoftFileParserImpl::AddFace(FaceIndex, Tessellation, VertexIndex))
+			{
+				continue;
+			}
+
+			TechSoftFileParserImpl::AddNormals(Tessellation3DData.m_pdNormals, NormalIndex, Tessellation.NormalArray);
+		}
+
+		InOutStartIndex = Index;
+	}
+
+	void AddFaceTriangleWithUniqueNormalAndTexture(FTessellationData& Tessellation, const A3DTess3DData& Tessellation3DData, long& InOutStartIndex, unsigned long InTrinangleCount, long TextureCount)
+	{
+		TechSoftFileParserImpl::Reserve(Tessellation, InTrinangleCount, /*bWithTexture*/ true);
+
+		unsigned long Index = InOutStartIndex;
+
+		int32 FaceIndex[3] = { 0, 0, 0 };
+		int32 NormalIndex[3] = { 0, 0, 0 };
+		int32 TextureIndex[3] = { 0, 0, 0 };
+		int32 VertexIndex = 0;
+
+		NormalIndex[0] = Tessellation3DData.m_puiTriangulatedIndexes[Index++];
+		NormalIndex[1] = NormalIndex[0];
+		NormalIndex[2] = NormalIndex[0];
+
+		// Get Triangles
+		for (unsigned long TriangleIndex = 0; TriangleIndex < InTrinangleCount; TriangleIndex++)
+		{
+			TextureIndex[0] = Tessellation3DData.m_puiTriangulatedIndexes[Index];
+			Index += TextureCount;
+			FaceIndex[0] = Tessellation3DData.m_puiTriangulatedIndexes[Index++];
+			TextureIndex[1] = Tessellation3DData.m_puiTriangulatedIndexes[Index];
+			Index += TextureCount;
+			FaceIndex[1] = Tessellation3DData.m_puiTriangulatedIndexes[Index++];
+			TextureIndex[2] = Tessellation3DData.m_puiTriangulatedIndexes[Index];
+			Index += TextureCount;
+			FaceIndex[2] = Tessellation3DData.m_puiTriangulatedIndexes[Index++];
+
+			if (!TechSoftFileParserImpl::AddFace(FaceIndex, Tessellation, VertexIndex))
+			{
+				continue;
+			}
+
+			TechSoftFileParserImpl::AddNormals(Tessellation3DData.m_pdNormals, NormalIndex, Tessellation.NormalArray);
+			TechSoftFileParserImpl::AddTextureCoordinates(Tessellation3DData.m_pdTextureCoords, TextureIndex, Tessellation.TexCoordArray);
+		}
+
+		InOutStartIndex = Index;
+	}
+
+	void AddFaceTriangle(FTessellationData& Tessellation, const A3DTess3DData& Tessellation3DData, long& InOutStartIndex, unsigned long InTrinangleCount)
+	{
+		TechSoftFileParserImpl::Reserve(Tessellation, InTrinangleCount, /*bWithTexture*/ false);
+
+		unsigned long Index = InOutStartIndex;
+
+		int32 FaceIndex[3] = { 0, 0, 0 };
+		int32 NormalIndex[3] = { 0, 0, 0 };
+		int32 VertexIndex = 0;
+
+		// Get Triangles
+		for (unsigned long TriangleIndex = 0; TriangleIndex < InTrinangleCount; TriangleIndex++)
+		{
+			NormalIndex[0] = Tessellation3DData.m_puiTriangulatedIndexes[Index++];
+			FaceIndex[0] = Tessellation3DData.m_puiTriangulatedIndexes[Index++];
+			NormalIndex[1] = Tessellation3DData.m_puiTriangulatedIndexes[Index++];
+			FaceIndex[1] = Tessellation3DData.m_puiTriangulatedIndexes[Index++];
+			NormalIndex[2] = Tessellation3DData.m_puiTriangulatedIndexes[Index++];
+			FaceIndex[2] = Tessellation3DData.m_puiTriangulatedIndexes[Index++];
+
+			if (!TechSoftFileParserImpl::AddFace(FaceIndex, Tessellation, VertexIndex))
+			{
+				continue;
+			}
+
+			TechSoftFileParserImpl::AddNormals(Tessellation3DData.m_pdNormals, NormalIndex, Tessellation.NormalArray);
+		}
+
+		InOutStartIndex = Index;
+	}
+
+	void AddFaceTriangleWithTexture(FTessellationData& Tessellation, const A3DTess3DData& Tessellation3DData, long& InOutStartIndex, unsigned long InTrinangleCount, long TextureCount)
+	{
+		TechSoftFileParserImpl::Reserve(Tessellation, InTrinangleCount, /*bWithTexture*/ true);
+
+		unsigned long Index = InOutStartIndex;
+
+		int32 FaceIndex[3] = { 0, 0, 0 };
+		int32 NormalIndex[3] = { 0, 0, 0 };
+		int32 TextureIndex[3] = { 0, 0, 0 };
+		int32 VertexIndex = 0;
+
+		// Get Triangles
+		for (unsigned long TriangleIndex = 0; TriangleIndex < InTrinangleCount; TriangleIndex++)
+		{
+			NormalIndex[0] = Tessellation3DData.m_puiTriangulatedIndexes[Index++];
+			TextureIndex[0] = Tessellation3DData.m_puiTriangulatedIndexes[Index];
+			Index += TextureCount;
+			FaceIndex[0] = Tessellation3DData.m_puiTriangulatedIndexes[Index++];
+			NormalIndex[1] = Tessellation3DData.m_puiTriangulatedIndexes[Index++];
+			TextureIndex[1] = Tessellation3DData.m_puiTriangulatedIndexes[Index];
+			Index += TextureCount;
+			FaceIndex[1] = Tessellation3DData.m_puiTriangulatedIndexes[Index++];
+			NormalIndex[2] = Tessellation3DData.m_puiTriangulatedIndexes[Index++];
+			TextureIndex[2] = Tessellation3DData.m_puiTriangulatedIndexes[Index];
+			Index += TextureCount;
+			FaceIndex[2] = Tessellation3DData.m_puiTriangulatedIndexes[Index++];
+
+			if (!TechSoftFileParserImpl::AddFace(FaceIndex, Tessellation, VertexIndex))
+			{
+				continue;
+			}
+
+			TechSoftFileParserImpl::AddNormals(Tessellation3DData.m_pdNormals, NormalIndex, Tessellation.NormalArray);
+			TechSoftFileParserImpl::AddTextureCoordinates(Tessellation3DData.m_pdTextureCoords, TextureIndex, Tessellation.TexCoordArray);
+		}
+
+		InOutStartIndex = Index;
+	}
+
+	// TODO import fan and strip meshes
+#ifdef UNUSED
+	double ComputeFaceAreaTriangleFan(double* InCoordinates, unsigned long* InTriangulatedIndexes, long& InOutStartIndex, unsigned long InTrinangleCount, bool bInOneNormal, long TextureCount, int& TriangleCount)
+	{
+		double Area = 0;
+
+		unsigned long OffSet = TextureCount + (bInOneNormal ? 0 : 1) + 1;
+		unsigned long Index = InOutStartIndex + 2 + TextureCount;
+
+		int Index0 = InTriangulatedIndexes[Index];
+		double* Point0 = InCoordinates + (InTriangulatedIndexes[Index]);
+		Index += OffSet;
+		int Index1 = InTriangulatedIndexes[Index];
+		double* Point1 = InCoordinates + (InTriangulatedIndexes[Index]);
+		Index += OffSet;
+
+		for (unsigned long TriangleIndex = 2; TriangleIndex < InTrinangleCount; TriangleIndex++)
+		{
+			double* Point2 = InCoordinates + (InTriangulatedIndexes[Index]);
+			int Index2 = InTriangulatedIndexes[Index];
+			if (Index0 == Index1 || Index0 == Index2 || Index1 == Index2)
+			{
+				TriangleCount--;
+			}
+			else
+			{
+				Area += FaceArea(Point0, Point1, Point2);
+			}
+			Index += OffSet;
+			Point1 = Point2;
+			Index1 = Index2;
+		}
+		InOutStartIndex = Index - OffSet;
+
+		return Area;
+	}
+
+	double ComputeFaceAreaTriangleStrip(double* InCoordinates, unsigned long* InTriangulatedIndexes, long& InOutStartIndex, unsigned long InTrinangleCount, bool bInOneNormal, long TextureCount, int& TriangleCount)
+	{
+		double Area = 0;
+
+		unsigned long OffSet = TextureCount + (bInOneNormal ? 0 : 1) + 1;
+		unsigned long Index = InOutStartIndex + 2 + TextureCount;
+
+		int Index0 = InTriangulatedIndexes[Index];
+		double* Point0 = InCoordinates + (InTriangulatedIndexes[Index]);
+		Index += OffSet;
+		int Index1 = InTriangulatedIndexes[Index];
+		double* Point1 = InCoordinates + (InTriangulatedIndexes[Index]);
+		Index += OffSet;
+
+		for (unsigned long TriangleIndex = 2; TriangleIndex < InTrinangleCount; TriangleIndex++)
+		{
+			double* Point2 = InCoordinates + (InTriangulatedIndexes[Index]);
+			int Index2 = InTriangulatedIndexes[Index];
+			if (Index0 == Index1 || Index0 == Index2 || Index1 == Index2)
+			{
+				TriangleCount--;
+			}
+			else
+			{
+				Area += FaceArea(Point0, Point1, Point2);
+			}
+
+			Index += OffSet;
+			Point0 = Point1;
+			Point1 = Point2;
+			Index0 = Index1;
+			Index1 = Index2;
+		}
+		InOutStartIndex = Index - OffSet;
+
+		return Area;
+	}
+#endif
+
+	void FTechSoftFileParser::TraverseTessellation3D(const A3DTess3D* TessellationPtr, FArchiveBody& Body)
+	{
+		FBodyMesh& BodyMesh = CADFileData.AddBodyMesh(Body.ObjectId, Body);
+
+		const int TessellationFaceDataWithTriangle = 0x2222;
+		const int TessellationFaceDataWithFan = 0x4444;
+		const int TessellationFaceDataWithStrip = 0x8888;
+		const int TessellationFaceDataWithOneNormal = 0xE0E0;
+
+		// Coordinates
+		TechSoft::TUniqueTSObj<A3DTessBaseData, A3DTess3D> TessellationBaseData(TessellationPtr, A3DTessBaseGet);
+		{
+			if (TessellationBaseData.IsValid() && TessellationBaseData->m_uiCoordSize > 0)
+			{
+				int32 VertexCount = TessellationBaseData->m_uiCoordSize / 3;
+				BodyMesh.VertexArray.Reserve(VertexCount);
+
+				double* Coordinates = TessellationBaseData->m_pdCoords;
+				for (unsigned int Index = 0; Index < TessellationBaseData->m_uiCoordSize; ++Index)
+				{
+					Coordinates[Index] *= FileUnit;
+				}
+
+				for (unsigned int Index = 0; Index < TessellationBaseData->m_uiCoordSize; Index += 3)
+				{
+					BodyMesh.VertexArray.Emplace(Coordinates[Index], Coordinates[Index + 1], Coordinates[Index + 2]);
+				}
+			}
+			else
+			{
+				// No vertex, no mesh...
+				return;
+			}
+		}
+
+		TechSoft::TUniqueTSObj<A3DTess3DData, A3DTess3D> Tessellation3DData(TessellationPtr, A3DTess3DGet);
+		if (Tessellation3DData.IsValid())
+		{
+			for (unsigned int Index = 0; Index < Tessellation3DData->m_uiFaceTessSize; ++Index)
+			{
+				const A3DTessFaceData& FaceTessData = Tessellation3DData->m_psFaceTessData[Index];
+				FTessellationData& Tessellation = BodyMesh.Faces.Emplace_GetRef();
+
+				unsigned int UsedEntitiesFlags = FaceTessData.m_usUsedEntitiesFlags;
+				long StartTriangulated = FaceTessData.m_uiStartTriangulated;
+
+				unsigned int FaceSetIndex = 0;
+				// Triangles
+				if (UsedEntitiesFlags & TessellationFaceDataWithTriangle)
+				{
+					bool bTessellationWithOneNormal = (UsedEntitiesFlags & TessellationFaceDataWithOneNormal);
+					if (bTessellationWithOneNormal)
+					{
+						if (FaceTessData.m_uiTextureCoordIndexesSize)
+						{
+							AddFaceTriangleWithUniqueNormalAndTexture(Tessellation, *Tessellation3DData, StartTriangulated, FaceTessData.m_puiSizesTriangulated[0], FaceTessData.m_uiTextureCoordIndexesSize);
+						}
+						else
+						{
+							AddFaceTriangleWithUniqueNormal(Tessellation, *Tessellation3DData, StartTriangulated, FaceTessData.m_puiSizesTriangulated[0]);
+						}
+					}
+					else
+					{
+						if (FaceTessData.m_uiTextureCoordIndexesSize)
+						{
+							AddFaceTriangleWithTexture(Tessellation, *Tessellation3DData, StartTriangulated, FaceTessData.m_puiSizesTriangulated[0], FaceTessData.m_uiTextureCoordIndexesSize);
+						}
+						else
+						{
+							AddFaceTriangle(Tessellation, *Tessellation3DData, StartTriangulated, FaceTessData.m_puiSizesTriangulated[0]);
+						}
+					}
+					FaceSetIndex++;
+				}
+
+				// Fans TODO
+				if (FaceTessData.m_uiSizesTriangulatedSize > FaceSetIndex)
+				{
+					if (UsedEntitiesFlags & TessellationFaceDataWithFan)
+					{
+						unsigned int LastFanIndex = 1 + FaceSetIndex + FaceTessData.m_puiSizesTriangulated[FaceSetIndex];
+						FaceSetIndex++;
+						for (; FaceSetIndex < LastFanIndex; FaceSetIndex++)
+						{
+							bool bTessellationWithOneNormal = (UsedEntitiesFlags & TessellationFaceDataWithOneNormal) && (FaceTessData.m_puiSizesTriangulated[FaceSetIndex] & kA3DTessFaceDataNormalSingle);
+							A3DUns32 FanSize = (FaceTessData.m_puiSizesTriangulated[FaceSetIndex] & kA3DTessFaceDataNormalMask);
+							//TriangleCount += (FanSize - 2);
+							//BodyMeshMetaData.MeshArea += ComputeFaceAreaTriangleFan(TessellationBaseData.m_pdCoords, Tessellation3DData->m_puiTriangulatedIndexes, StartTriangulated, FanSize, bTessellationWithOneNormal, FaceTessData.m_uiTextureCoordIndexesSize, TriangleCount);
+						}
+					}
+				}
+
+				// Strip TODO
+				if (FaceTessData.m_uiSizesTriangulatedSize > FaceSetIndex)
+				{
+					FaceSetIndex++;
+					for (; FaceSetIndex < FaceTessData.m_uiSizesTriangulatedSize; FaceSetIndex++)
+					{
+						bool bTessellationWithOneNormal = (UsedEntitiesFlags & TessellationFaceDataWithOneNormal) && (FaceTessData.m_puiSizesTriangulated[FaceSetIndex] & kA3DTessFaceDataNormalSingle);
+						A3DUns32 StripSize = (FaceTessData.m_puiSizesTriangulated[FaceSetIndex] & kA3DTessFaceDataNormalMask);
+						//TriangleCount += (StripSize - 2);
+						//BodyMeshMetaData.MeshArea += ComputeFaceAreaTriangleStrip(TessellationBaseData.m_pdCoords, Tessellation3DData->m_puiTriangulatedIndexes, StartTriangulated, StripSize, bTessellationWithOneNormal, FaceTessData.m_uiTextureCoordIndexesSize, TriangleCount);
+					}
+				}
+
+			}
+		}
+	}
+
+}
+
+#endif
\ No newline at end of file
diff --git a/Engine/Plugins/Enterprise/DatasmithCADImporter/Source/CADInterfaces/Private/TechSoftFileParser.h b/Engine/Plugins/Enterprise/DatasmithCADImporter/Source/CADInterfaces/Private/TechSoftFileParser.h
new file mode 100644
--- /dev/null
+++ b/Engine/Plugins/Enterprise/DatasmithCADImporter/Source/CADInterfaces/Private/TechSoftFileParser.h
@@ -0,0 +1,281 @@
+// Copyright Epic Games, Inc. All Rights Reserved.
+#pragma once
+
+#include "CoreMinimal.h"
+
+#include "CADData.h"
+#include "CADFileData.h"
+#include "CADFileParser.h"
+#include "CADOptions.h"
+#include "CADSceneGraph.h"
+
+#include "TechSoftInterface.h"
+
+typedef void A3DAsmModelFile;
+typedef void A3DAsmPartDefinition;
+typedef void A3DAsmProductOccurrence;
+typedef void A3DEntity;
+typedef void A3DGraphics;
+typedef void A3DMiscAttribute;
+typedef void A3DMiscCartesianTransformation;
+typedef void A3DMiscCartesianTransformation;
+typedef void A3DMiscGeneralTransformation;
+typedef void A3DRiBrepModel;
+typedef void A3DRiCoordinateSystem;
+typedef void A3DRiPolyBrepModel;
+typedef void A3DRiRepresentationItem;
+typedef void A3DRiSet;
+typedef void A3DTess3D;
+typedef void A3DTessBase;
+typedef void A3DTopoBrepData;
+
+namespace CADLibrary
+{
+	class FEntityData;
+	class FComponentData;
+
+	enum EComponentType : uint32
+	{
+		Reference = 0,
+		Occurrence,
+		UnloadedComponent,
+		Body,
+		Undefined,
+		LastType
+	};
+
+	enum EModellerType : uint32
+	{
+		kA3DModellerUnknown = 0,              /*!< User modeller. */
+		kA3DModellerCatia = 2,                /*!< CATIA modeller. */
+		kA3DModellerCatiaV5 = 3,              /*!< CATIA V5 modeller. */
+		kA3DModellerCadds = 4,                /*!< CADDS modeller. */
+		kA3DModellerUnigraphics = 5,          /*!< Unigraphics modeller. */
+		kA3DModellerParasolid = 6,            /*!< Parasolid modeller. */
+		kA3DModellerEuclid = 7,               /*!< Euclid modeller. */
+		kA3DModellerIges = 9,                 /*!< IGES modeller. */
+		kA3DModellerUnisurf = 10,             /*!< Unisurf modeller. */
+		kA3DModellerVda = 11,                 /*!< VDA modeller. */
+		kA3DModellerStl = 12,                 /*!< STL modeller. */
+		kA3DModellerWrl = 13,                 /*!< WRL modeller. */
+		kA3DModellerDxf = 14,                 /*!< DXF modeller. */
+		kA3DModellerAcis = 15,                /*!< ACIS modeller. */
+		kA3DModellerProE = 16,                /*!< Pro/E modeller. */
+		kA3DModellerStep = 18,                /*!< STEP modeller. */
+		kA3DModellerIdeas = 19,               /*!< I-DEAS modeller. */
+		kA3DModellerJt = 20,                  /*!< JT modeller. */
+		kA3DModellerSlw = 22,                 /*!< SolidWorks modeller. */
+		kA3DModellerCgr = 23,                 /*!< CGR modeller. */
+		kA3DModellerPrc = 24,                 /*!< PRC modeller. */
+		kA3DModellerXvl = 25,                 /*!< XVL modeller. */
+		kA3DModellerHpgl = 26,                /*!< HPGL modeller. */
+		kA3DModellerTopSolid = 27,            /*!< TopSolid modeller. */
+		kA3DModellerOneSpaceDesigner = 28,    /*!< OneSpace designer modeller. */
+		kA3DModeller3dxml = 29,               /*!< 3DXML modeller. */
+		kA3DModellerInventor = 30,            /*!< Inventor modeller. */
+		kA3DModellerPostScript = 31,          /*!< Postscript modeller. */
+		kA3DModellerPDF = 32,                 /*!< PDF modeller. */
+		kA3DModellerU3D = 33,                 /*!< U3D modeller. */
+		kA3DModellerIFC = 34,                 /*!< IFC modeller. */
+		kA3DModellerDWG = 35,                 /*!< DWG modeller. */
+		kA3DModellerDWF = 36,                 /*!< DWF modeller. */
+		kA3DModellerSE = 37,                  /*!< SolidEdge modeller. */
+		kA3DModellerOBJ = 38,                 /*!< OBJ modeller. */
+		kA3DModellerKMZ = 39,                 /*!< KMZ modeller. */
+		kA3DModellerDAE = 40,                 /*!< COLLADA modeller. */
+		kA3DModeller3DS = 41,                 /*!< 3DS modeller. */
+		kA3DModellerRhino = 43,               /*!< Rhino modeller. */
+		kA3DModellerXML = 44,                 /*!< XML modeller. */
+		kA3DModeller3mf = 45,                 /*!< 3MF modeller. */
+		kA3DModellerScs = 46,                 /*!< SCS modeller. */
+		kA3DModeller3dHtml = 47,              /*!< 3DHTML modeller. */
+		kA3DModellerHsf = 48,                 /*!< Hsf modeller. */
+		kA3DModellerGltf = 49,                /*!< GL modeller. */
+		kA3DModellerRevit = 50,               /*!< Revit modeller. */
+		kA3DModellerFBX = 51,                 /*!< FBX modeller. */
+		kA3DModellerStepXML = 52,             /*!< StepXML modeller. */
+		kA3DModellerPLMXML = 53,              /*!< PLMXML modeller. */
+		kA3DModellerNavisworks = 54,			/*!< For Future Use: Navisworks modeller. */
+		kA3DModellerLast
+	};
+
+	struct FEntityBehaviour
+	{
+		bool bFatherHeritColor = false;
+		bool bFatherHeritLayer = false;
+		bool bFatherHeritLinePattern = false;
+		bool bFatherHeritLineWidth = false;
+		bool bFatherHeritShow = false;
+		bool bFatherHeritTransparency = false;
+		bool bRemoved = false;
+		bool bShow = true;
+		bool bSonHeritColor = false;
+		bool bSonHeritLayer = false;
+		bool bSonHeritLinePattern = false;
+		bool bSonHeritLineWidth = false;
+		bool bSonHeritShow = false;
+		bool bSonHeritTransparency = false;
+	};
+
+	struct FEntityMetaData
+	{
+		TMap<FString, FString> MetaData;
+		bool bRemoved = false;
+		bool bShow = true;
+		bool bUnloaded = false;
+		FFileDescriptor ExternalFile;
+	};
+
+	class CADINTERFACES_API FTechSoftFileParser : public ICADFileParser
+	{
+	public:
+
+		/**
+		 * @param InCADData TODO
+		 * @param EnginePluginsPath Full Path of EnginePlugins. Mandatory to set KernelIO to import DWG, or DGN files
+		 */
+		FTechSoftFileParser(FCADFileData& InCADData, const FString& EnginePluginsPath = TEXT(""))
+			: CADFileData(InCADData)
+		{
+		}
+
+#ifndef USE_TECHSOFT_SDK
+		virtual ECADParsingResult Process() override
+		{
+			return ECADParsingResult::ProcessFailed;
+		}
+#else
+		virtual ECADParsingResult Process() override;
+
+	private:
+	
+		// Methods to parse a first time the file to count the number of components
+		// Needed to reserve CADFileData
+		// Start with CountUnderModel
+		void CountUnderModel(const A3DAsmModelFile* AsmModel);
+		void CountUnderOccurrence(const A3DAsmProductOccurrence* Occurrence);
+		void CountUnderPrototype(const A3DAsmProductOccurrence* Prototype);
+		void CountUnderPartDefinition(const A3DAsmPartDefinition* PartDefinition);
+		void CountUnderRepresentation(const A3DRiRepresentationItem* RepresentationItem);
+		void CountUnderRepresentationSet(const A3DRiSet* RepresentationSet);
+
+		void ReserveCADFileData();
+
+		// Materials and colors
+		int32 CountMaterial();
+		int32 CountColor();
+		void ReadMaterialsAndColors();
+
+		// Traverse ASM tree by starting from the model
+		ECADParsingResult TraverseModel(const A3DAsmModelFile* AsmModel);
+		void TraverseReference(const A3DAsmProductOccurrence* Reference);
+		bool IsConfigurationSet(const A3DAsmProductOccurrence* Occurrence);
+		// TODO for SW //void TraverseReferenceFromConfiguration(const A3DAsmProductOccurrence* Reference, FEntityData& ParentMetaData);
+		// TODO for SW //void TraverseConfigurationSet(const A3DAsmProductOccurrence* ConfigurationSet, FEntityData& ParentMetaData);
+		// TODO for SW //void TraverseConfiguration(const A3DAsmProductOccurrence* OccurrConfigurationence, FEntityData& ParentMetaData);
+		FCadId TraverseOccurrence(const A3DAsmProductOccurrence* Occurrence);
+		void TraversePrototype(const A3DAsmProductOccurrence* Prototype, FEntityMetaData& MetaData);
+		FCadId TraversePartDefinition(const A3DAsmPartDefinition* PartDefinition);
+		FCadId TraverseRepresentationSet(const A3DRiSet* pSet);
+		FCadId TraverseRepresentationItem(A3DRiRepresentationItem* RepresentationItem);
+		FCadId TraverseBRepModel(A3DRiBrepModel* BrepModel);
+		FCadId TraversePolyBRepModel(const A3DRiPolyBrepModel* PolygonalBrepModel);
+
+		// Tessellation methods
+		void MeshRepresentationWithTechSoft(A3DRiRepresentationItem* RepresentationItem, FArchiveBody& Body);
+		void TraverseRepresentationContent(const A3DRiRepresentationItem* RepresentationItem, FArchiveBody& Body);
+		void TraverseTessellationBase(const A3DTessBase* Tessellation, FArchiveBody& Body);
+		void TraverseTessellation3D(const A3DTess3D* Tessellation, FArchiveBody& Body);
+
+		// MetaData
+		void TraverseMetaData(const A3DEntity* Entity, FEntityMetaData& OutMetaData);
+		void TraverseSpecificMetaData(const A3DAsmProductOccurrence* Occurrence, FEntityMetaData& OutMetaData);
+		FString DefineEntityName(TMap<FString, FString>& OutMetaData, EComponentType EntityType);
+
+		// Graphic properties
+		void TraverseGraphics(const A3DGraphics* Graphics, FEntityBehaviour& GraphicsBehaviour);
+		void TraverseMaterialProperties(const A3DEntity* Entity);
+		void TraverseLayer(const A3DAsmProductOccurrence* Occurrence);
+
+		// Transform
+		void TraverseCoordinateSystem(const A3DRiCoordinateSystem* CoordinateSystem);
+		void TraverseTransformation(const A3DMiscCartesianTransformation* Transformation3d);
+		void TraverseGeneralTransformation(const A3DMiscGeneralTransformation* GeneralTransformation);
+		void TraverseTransformation3D(const A3DMiscCartesianTransformation* CartesianTransformation);
+
+		FFileDescriptor GetOccurrenceFileName(const A3DAsmProductOccurrence* OccurrencePtr);
+
+		// Archive methodes
+		FArchiveInstance& AddInstance(FEntityMetaData& InstanceMetaData);
+		FArchiveComponent& AddComponent(FEntityMetaData& ComponentMetaData, FArchiveInstance& Instance);
+		FArchiveUnloadedComponent& AddUnloadedComponent(FEntityMetaData& ComponentMetaData, FArchiveInstance& Instance);
+		FArchiveComponent& AddOccurence(FEntityMetaData& InstanceMetaData, FEntityMetaData& ReferenceMetaData, FCadId& OutComponentId);
+		FArchiveComponent& AddOccurence(FEntityMetaData& InstanceMetaData, FCadId& OutComponentId);
+		FArchiveBody& AddBody(FEntityMetaData& BodyMetaData);
+#endif
+
+	private:
+
+		TSet<const A3DAsmProductOccurrence*> PrototypeCounted;
+		uint32 ComponentCount[EComponentType::LastType] = {0};
+
+		FCADFileData& CADFileData;
+		TSharedPtr<ITechSoftInterface> TechSoftInterface;
+
+		EModellerType ModellerType;
+		double FileUnit = 1;
+		FCadId LastEntityId = 1;
+	};
+} // ns CADLibrary
+
+namespace TechSoftFileParserImpl
+{
+	// Methodes used by TraverseTessellation3D
+	
+	typedef double A3DDouble;
+	inline bool AddFace(int32 FaceIndex[3], CADLibrary::FTessellationData& Tessellation, int32& VertexIndex)
+	{
+		if (FaceIndex[0] == FaceIndex[1] || FaceIndex[0] == FaceIndex[2] || FaceIndex[1] == FaceIndex[2])
+		{
+			return false;
+		}
+
+		for (int32 Index = 0; Index < 3; ++Index)
+		{
+			FaceIndex[Index] /= 3;
+			Tessellation.VertexIndices.Add(VertexIndex++);
+		}
+		Tessellation.PositionIndices.Append(FaceIndex, 3);
+		return true;
+	};
+
+	inline void AddNormals(const A3DDouble* Normals, const int32 Indices[3], TArray<FVector>& TessellationNormals)
+	{
+		for (int32 Index = 0; Index < 3; ++Index)
+		{
+			int32 NormalIndex = Indices[Index];
+			TessellationNormals.Emplace(Normals[NormalIndex], Normals[NormalIndex + 1], Normals[NormalIndex + 2]);
+		}
+	};
+
+	inline void AddTextureCoordinates(const A3DDouble* TextureCoords, const int32 Indices[3], TArray<FVector2D>& TessellationTextures)
+	{
+		for (int32 Index = 0; Index < 3; ++Index)
+		{
+			int32 TextureIndex = Indices[Index];
+			TessellationTextures.Emplace(TextureCoords[TextureIndex], TextureCoords[TextureIndex + 1]);
+		}
+	};
+
+	inline void Reserve(CADLibrary::FTessellationData& Tessellation, int32 InTrinangleCount, bool bWithTexture)
+	{
+		Tessellation.PositionIndices.Reserve(3 * InTrinangleCount);
+		Tessellation.VertexIndices.Reserve(3 * InTrinangleCount);
+		Tessellation.NormalArray.Reserve(3 * InTrinangleCount);
+		if (bWithTexture)
+		{
+			Tessellation.TexCoordArray.Reserve(3 * InTrinangleCount);
+		}
+	};
+
+}
\ No newline at end of file
diff --git a/Engine/Plugins/Enterprise/DatasmithCADImporter/Source/CADInterfaces/Private/TechSoftInterface.cpp b/Engine/Plugins/Enterprise/DatasmithCADImporter/Source/CADInterfaces/Private/TechSoftInterface.cpp
new file mode 100644
--- /dev/null
+++ b/Engine/Plugins/Enterprise/DatasmithCADImporter/Source/CADInterfaces/Private/TechSoftInterface.cpp
@@ -0,0 +1,32 @@
+// Copyright Epic Games, Inc. All Rights Reserved.
+#include "TechSoftInterface.h"
+
+#include "TechSoftInterfaceImpl.h"
+
+
+namespace CADLibrary
+{
+	TSharedPtr<ITechSoftInterface> TechSoftInterface;
+
+	TSharedPtr<ITechSoftInterface>& GetTechSoftInterface()
+	{
+		return TechSoftInterface;
+	}
+
+	void InitializeTechSoftInterface()
+	{
+#ifdef USE_TECHSOFT_SDK
+		TechSoftInterface = MakeShared<CADLibrary::FTechSoftInterfaceImpl>();
+#endif
+	}
+
+	bool TECHSOFT_InitializeKernel(const TCHAR* InEnginePluginsPath)
+	{
+		if (!TechSoftInterface.IsValid())
+		{
+			return false;
+		}
+
+		return TechSoftInterface->InitializeKernel(InEnginePluginsPath);
+	}
+}
diff --git a/Engine/Plugins/Enterprise/DatasmithCADImporter/Source/CADInterfaces/Private/TechSoftInterfaceImpl.cpp b/Engine/Plugins/Enterprise/DatasmithCADImporter/Source/CADInterfaces/Private/TechSoftInterfaceImpl.cpp
new file mode 100644
--- /dev/null
+++ b/Engine/Plugins/Enterprise/DatasmithCADImporter/Source/CADInterfaces/Private/TechSoftInterfaceImpl.cpp
@@ -0,0 +1,54 @@
+// Copyright Epic Games, Inc. All Rights Reserved.
+#ifdef USE_TECHSOFT_SDK
+
+#define INITIALIZE_A3D_API
+
+#include "TechSoftInterfaceImpl.h"
+
+#include "CADInterfacesModule.h"
+#include "Misc/Paths.h"
+
+
+namespace CADLibrary
+{
+	bool FTechSoftInterfaceImpl::InitializeKernel(const TCHAR* InEnginePluginsPath)
+	{
+		if (bIsInitialize)
+		{
+			return true;
+		}
+
+		FString EnginePluginsPath(InEnginePluginsPath);
+		if (EnginePluginsPath.IsEmpty())
+		{
+			EnginePluginsPath = FPaths::EnginePluginsDir();
+		}
+
+		FString TechSoftDllPath = FPaths::Combine(EnginePluginsPath, TEXT("Enterprise/DatasmithCADImporter"), TEXT("Binaries"), FPlatformProcess::GetBinariesSubdirectory(), "TechSoft");
+		TechSoftDllPath = FPaths::ConvertRelativePathToFull(TechSoftDllPath);
+		ExchangeLoader = MakeUnique<A3DSDKHOOPSExchangeLoader>(*TechSoftDllPath);
+
+		const A3DStatus IRet = ExchangeLoader->m_eSDKStatus;
+		if (IRet != A3D_SUCCESS)
+		{
+			UE_LOG(CADInterfaces, Warning, TEXT("Failed to load required library in %s. Plug-in will not be functional."), *TechSoftDllPath);
+		}
+		else
+		{
+			bIsInitialize = true;
+		}
+		return bIsInitialize;
+	}
+
+	A3DStatus FTechSoftInterfaceImpl::Import(const A3DImport& Import)
+	{
+		return ExchangeLoader->Import(Import);
+	}
+
+	A3DAsmModelFile* FTechSoftInterfaceImpl::GetModelFile()
+	{
+		return ExchangeLoader->m_psModelFile;
+	}
+
+}
+#endif // USE_TECHSOFT_SDK
diff --git a/Engine/Plugins/Enterprise/DatasmithCADImporter/Source/CADInterfaces/Private/TechSoftInterfaceImpl.h b/Engine/Plugins/Enterprise/DatasmithCADImporter/Source/CADInterfaces/Private/TechSoftInterfaceImpl.h
new file mode 100644
--- /dev/null
+++ b/Engine/Plugins/Enterprise/DatasmithCADImporter/Source/CADInterfaces/Private/TechSoftInterfaceImpl.h
@@ -0,0 +1,72 @@
+// Copyright Epic Games, Inc. All Rights Reserved.
+#pragma once
+
+#include "CoreMinimal.h"
+
+#ifdef USE_TECHSOFT_SDK
+#include "CADData.h"
+#include "CADOptions.h"
+#include "CADSceneGraph.h"
+#include "TechSoftInterface.h"
+
+// undefine UE TEXT to include windows.h
+#undef TEXT
+#ifdef WIN32
+#define WIN32_LEAN_AND_MEAN
+#include <windows.h>
+#endif
+
+// redefine UE TEXT
+#undef TEXT
+#if !defined(TEXT) && !UE_BUILD_DOCS
+#if PLATFORM_TCHAR_IS_CHAR16
+#define TEXT_PASTE(x) u ## x
+#else
+#define TEXT_PASTE(x) L ## x
+#endif
+#define TEXT(x) TEXT_PASTE(x)
+#endif
+
+
+#pragma warning(push)
+#pragma warning(disable:4505)
+#pragma warning(disable:4191)
+#include "A3DSDKIncludes.h"
+#pragma warning(pop)
+
+namespace CADLibrary
+{
+	class CADINTERFACES_API FTechSoftInterfaceImpl : public ITechSoftInterface
+	{
+	public:
+		FTechSoftInterfaceImpl(bool bSetExternal = false)
+			: bIsExternal(bSetExternal)
+		{}
+
+		virtual bool InitializeKernel(const TCHAR* EnginePluginsPath) override;
+
+		virtual bool IsExternal() override
+		{
+			return bIsExternal;
+		}
+
+		virtual void SetExternal(bool Value) override
+		{
+			bIsExternal = Value;
+		}
+
+		virtual A3DStatus Import(const A3DImport& Import) override;
+		virtual A3DAsmModelFile* GetModelFile() override;
+
+	private:
+
+		bool bIsExternal = false;
+		bool bIsInitialize = false;
+		TUniquePtr<A3DSDKHOOPSExchangeLoader> ExchangeLoader;
+
+	};
+}
+
+#endif // USE_TECHSOFT_SDK
+
+
diff --git a/Engine/Plugins/Enterprise/DatasmithCADImporter/Source/CADInterfaces/Public/TechSoftInterface.h b/Engine/Plugins/Enterprise/DatasmithCADImporter/Source/CADInterfaces/Public/TechSoftInterface.h
new file mode 100644
--- /dev/null
+++ b/Engine/Plugins/Enterprise/DatasmithCADImporter/Source/CADInterfaces/Public/TechSoftInterface.h
@@ -0,0 +1,41 @@
+// Copyright Epic Games, Inc. All Rights Reserved.
+#pragma once
+
+#include "CADData.h"
+
+#ifdef USE_TECHSOFT_SDK
+#include "A3DSDKErrorCodes.h"
+#endif
+
+class A3DImport;
+typedef void A3DAsmModelFile;
+
+namespace CADLibrary
+{
+	CADINTERFACES_API void InitializeTechSoftInterface();
+	CADINTERFACES_API bool TECHSOFT_InitializeKernel(const TCHAR* = TEXT(""));
+
+	CADINTERFACES_API class ITechSoftInterface
+	{
+	public:
+
+		virtual ~ITechSoftInterface() = default;
+
+		/*
+		* Returns true if the object has been created outside of the memory pool of the running process
+		* This is the case when the object has been created by the DatasmithRuntime plugin
+		*/
+		virtual bool IsExternal() = 0;
+		virtual void SetExternal(bool Value) = 0;
+
+		virtual bool InitializeKernel(const TCHAR* = TEXT("")) = 0;
+
+#ifdef USE_TECHSOFT_SDK
+		virtual A3DStatus Import(const A3DImport& Import) = 0;
+		virtual A3DAsmModelFile* GetModelFile() = 0;
+#endif
+	};
+
+	CADINTERFACES_API TSharedPtr<ITechSoftInterface>& GetTechSoftInterface();
+}
+
