Fix crash with ds.CADTranslator.EnableCADCache=false

diff --git a/Engine/Plugins/Enterprise/DatasmithCADImporter/Source/CADInterfaces/Private/CADFileReader.cpp b/Engine/Plugins/Enterprise/DatasmithCADImporter/Source/CADInterfaces/Private/CADFileReader.cpp
--- a/Engine/Plugins/Enterprise/DatasmithCADImporter/Source/CADInterfaces/Private/CADFileReader.cpp
+++ b/Engine/Plugins/Enterprise/DatasmithCADImporter/Source/CADInterfaces/Private/CADFileReader.cpp
@@ -142,10 +142,13 @@ namespace CADLibrary
 		// Process the file
 		ECADParsingResult Result = CADParser->Process();
 
-		if (ECADParsingResult::ProcessOk == Result)
+		if (FImportParameters::bGEnableCADCache)
 		{
-			CADFileData.ExportSceneGraphFile();
-			CADFileData.ExportMeshArchiveFile();
+			if (ECADParsingResult::ProcessOk == Result)
+			{
+				CADFileData.ExportSceneGraphFile();
+				CADFileData.ExportMeshArchiveFile();
+			}
 		}
 
 		return Result;
