Fix Compile with #ifdef CORETECHBRIDGE_DEBUG

diff --git a/Engine/Plugins/Enterprise/DatasmithCADImporter/Source/CADInterfaces/Private/CoreTechFileParserImpl.cpp b/Engine/Plugins/Enterprise/DatasmithCADImporter/Source/CADInterfaces/Private/CoreTechFileParserImpl.cpp
--- a/Engine/Plugins/Enterprise/DatasmithCADImporter/Source/CADInterfaces/Private/CoreTechFileParserImpl.cpp
+++ b/Engine/Plugins/Enterprise/DatasmithCADImporter/Source/CADInterfaces/Private/CoreTechFileParserImpl.cpp
@@ -707,8 +707,8 @@ namespace CADLibrary
 			CADKernelModel->Add(CADKernelBody);
 
 #ifdef CORETECHBRIDGE_DEBUG
-			FString FolderName = FPaths::GetCleanFilename(FileDescription.MainCadFilePath);
-			CADKernelSession->SaveDatabase(*FPaths::Combine(CADFileData.CachePath, TEXT("CADKernel"), FolderName, FString::Printf(TEXT("%06d_"), BodyIndex++) + FileDescription.Name + TEXT(".ugeom")));
+			FString FolderName = FPaths::GetCleanFilename(FileDescription.GetFileName());
+			CADKernelSession->SaveDatabase(*FPaths::Combine(CADFileData.GetCADCachePath(), TEXT("CADKernel"), FolderName, FString::Printf(TEXT("%06d_"), BodyIndex++) + FileDescription.GetFileName() + TEXT(".ugeom")));
 #endif
 
 			// Save Body in CADKernelArchive file for re-tessellation
