Fix ds.CADTranslator.DisableCADKernelTessellation CVar

diff --git a/Engine/Plugins/Enterprise/DatasmithCADImporter/Source/CADInterfaces/Private/CoreTechFileParserImpl.cpp b/Engine/Plugins/Enterprise/DatasmithCADImporter/Source/CADInterfaces/Private/CoreTechFileParserImpl.cpp
--- a/Engine/Plugins/Enterprise/DatasmithCADImporter/Source/CADInterfaces/Private/CoreTechFileParserImpl.cpp
+++ b/Engine/Plugins/Enterprise/DatasmithCADImporter/Source/CADInterfaces/Private/CoreTechFileParserImpl.cpp
@@ -280,7 +280,7 @@ namespace CADLibrary
 
 		CoreTechFileParserUtils::AddFaceIdAttribut(MainId);
 
-		if (ImportParameters.GetStitchingTechnique() != StitchingNone && !FImportParameters::bGDisableCADKernelTessellation)
+		if (ImportParameters.GetStitchingTechnique() != StitchingNone && FImportParameters::bGDisableCADKernelTessellation)
 		{
 			CADLibrary::CTKIO_Repair(MainId, ImportParameters.GetStitchingTechnique(), 10.);
 		}
@@ -473,7 +473,7 @@ namespace CADLibrary
 			CT_FLAGS BodyProperties;
 			CT_BODY_IO::AskProperties(BodyId, BodyProperties);
 
-			if (!FImportParameters::bGDisableCADKernelTessellation || !(BodyProperties & CT_BODY_PROP_EXACT))
+			if (FImportParameters::bGDisableCADKernelTessellation || !(BodyProperties & CT_BODY_PROP_EXACT))
 			{
 				if (ReadKioBody(BodyId, ComponentId, DefaultMaterialHash, false))
 				{
@@ -664,7 +664,7 @@ namespace CADLibrary
 	}
 
 #ifdef CORETECHBRIDGE_DEBUG
-	static int32 BodyIndex = 0;
+	static int32 CoretechBridgeBodyIndex = 0;
 #endif
 	bool FCoreTechFileParser::ReadBody(CT_OBJECT_ID BodyId, CT_OBJECT_ID ParentId, uint32 DefaultMaterialHash, bool bNeedRepair)
 	{
@@ -708,7 +708,7 @@ namespace CADLibrary
 
 #ifdef CORETECHBRIDGE_DEBUG
 			FString FolderName = FPaths::GetCleanFilename(FileDescription.GetFileName());
-			CADKernelSession->SaveDatabase(*FPaths::Combine(CADFileData.GetCADCachePath(), TEXT("CADKernel"), FolderName, FString::Printf(TEXT("%06d_"), BodyIndex++) + FileDescription.GetFileName() + TEXT(".ugeom")));
+			CADKernelSession->SaveDatabase(*FPaths::Combine(CADFileData.GetCachePath(), TEXT("CADKernel"), FolderName, FString::Printf(TEXT("%06d_"), CoretechBridgeBodyIndex++) + FileDescription.GetFileName() + TEXT(".ugeom")));
 #endif
 
 			// Save Body in CADKernelArchive file for re-tessellation
