Fix Wire retesselation: Retesselation parameters were not transferred to CADModelConverter.
OpenNurbs : Remove duplicated information i.e. FDatasmithOpenNurbsOptions OpenNurbsOptions vs FImportParameters

diff --git a/Engine/Plugins/Enterprise/DatasmithCADImporter/Source/DatasmithOpenNurbsTranslator/Private/DatasmithOpenNurbsTranslator.cpp b/Engine/Plugins/Enterprise/DatasmithCADImporter/Source/DatasmithOpenNurbsTranslator/Private/DatasmithOpenNurbsTranslator.cpp
--- a/Engine/Plugins/Enterprise/DatasmithCADImporter/Source/DatasmithOpenNurbsTranslator/Private/DatasmithOpenNurbsTranslator.cpp
+++ b/Engine/Plugins/Enterprise/DatasmithCADImporter/Source/DatasmithOpenNurbsTranslator/Private/DatasmithOpenNurbsTranslator.cpp
@@ -640,14 +640,13 @@ public:
 		, ArchiveOpenNurbsVersion(0)
 		, FileLength(0)
 		, NumCRCErrors(0)
-		, ImportParameters(0.001, 0.1, FDatasmithUtils::EModelCoordSystem::ZUp_RightHanded_FBXLegacy)
 	{
-
 		if (!TranslationCache.IsValid())
 		{
 			TranslationCache = MakeShared<FTranslationCache>();
 		}
 
+		CADLibrary::FImportParameters ImportParameters(0.001, 0.1, FDatasmithUtils::EModelCoordSystem::ZUp_RightHanded_FBXLegacy);
 		ImportParameters.SwitchOffUVMapScaling();
 
 		if (CADLibrary::FImportParameters::bGDisableCADKernelTessellation)
@@ -832,7 +831,6 @@ private:
 
 	TArray<FString> MissingRenderMeshes;
 
-	CADLibrary::FImportParameters ImportParameters;
 	TSharedPtr<CADLibrary::ICADModelConverter> CADModelConverter;
 	TSharedPtr<IOpenNurbsBRepConverter> OpenNurbsBRepConverter;
 
diff --git a/Engine/Plugins/Enterprise/DatasmithCADImporter/Source/DatasmithWireTranslator/Private/DatasmithWireTranslator.cpp b/Engine/Plugins/Enterprise/DatasmithCADImporter/Source/DatasmithWireTranslator/Private/DatasmithWireTranslator.cpp
--- a/Engine/Plugins/Enterprise/DatasmithCADImporter/Source/DatasmithWireTranslator/Private/DatasmithWireTranslator.cpp
+++ b/Engine/Plugins/Enterprise/DatasmithCADImporter/Source/DatasmithWireTranslator/Private/DatasmithWireTranslator.cpp
@@ -134,7 +134,6 @@ public:
 		, ArchiveWireVersion(0)
 		, FileLength(0)
 		, NumCRCErrors(0)
-		, ImportParameters(0.01, 1)
 	{
 		// Set ProductName, ProductVersion in DatasmithScene for Analytics purpose
 		// application_name is something like "Catia V5"
@@ -144,6 +143,7 @@ public:
 		DatasmithScene->SetProductName(TEXT("Alias Tools"));
 		DatasmithScene->SetProductVersion(TEXT("Alias 2022"));
 
+		FImportParameters ImportParameters(0.01, 1);
 		if(FImportParameters::bGDisableCADKernelTessellation)
 		{
 			TSharedRef<FAliasModelToCoretechConverter> AliasToCoretechConverter = MakeShared<FAliasModelToCoretechConverter>(TEXT("Al2CTSharedSession"), ImportParameters);
@@ -290,7 +290,7 @@ private:
 	// If > 0, then the archive is corrupt.
 	int32 NumCRCErrors;
 
-	CADLibrary::FImportParameters ImportParameters;
+	FDatasmithTessellationOptions TessellationOptions;
 
 	TSharedPtr<ICADModelConverter> CADModelConverter;
 	TSharedPtr<IAliasBRepConverter> AliasBRepConverter;
@@ -299,7 +299,8 @@ private:
 
 void FWireTranslatorImpl::SetTessellationOptions(const FDatasmithTessellationOptions& Options)
 {
-	ImportParameters.SetTesselationParameters(Options.ChordTolerance, Options.MaxEdgeLength, Options.NormalTolerance, (CADLibrary::EStitchingTechnique) Options.StitchingTechnique);
+	TessellationOptions = Options;
+	CADModelConverter->SetImportParameters(Options.ChordTolerance, Options.MaxEdgeLength, Options.NormalTolerance, (CADLibrary::EStitchingTechnique) Options.StitchingTechnique);
 	SceneFileHash = HashCombine(Options.GetHash(), GetSceneFileHash(SceneFullPath, SceneName));
 }
 
@@ -1553,7 +1554,7 @@ void FWireTranslatorImpl::AddNodeInBodyGroup(TSharedPtr<AlDagNode>& DagNode, con
 
 bool FWireTranslatorImpl::RecurseDagForLeaves(const TSharedPtr<AlDagNode>& FirstDagNode, const FDagNodeInfo& ParentInfo)
 {
-	if (ImportParameters.GetStitchingTechnique() != EStitchingTechnique::StitchingSew)
+	if (TessellationOptions.StitchingTechnique != EDatasmithCADStitchingTechnique::StitchingSew)
 	{
 		return RecurseDagForLeavesNoMerge(FirstDagNode, ParentInfo);
 	}
@@ -1796,7 +1797,7 @@ TOptional<FMeshDescription> FWireTranslatorImpl::MeshDagNodeWithExternalMesher(T
 		// All actors of a Alias symmetric layer are defined in the world Reference i.e. they have identity transform. So Mesh actor has to be defined in the world reference. 
 		ObjectReference = EAliasObjectReference::WorldReference;
 	}
-	else if (ImportParameters.GetStitchingTechnique() == EStitchingTechnique::StitchingSew)
+	else if (TessellationOptions.StitchingTechnique != EDatasmithCADStitchingTechnique::StitchingSew)
 	{
 		// In the case of StitchingSew, AlDagNode children of a GroupNode are merged together. To be merged, they have to be defined in the reference of parent GroupNode.
 		ObjectReference = EAliasObjectReference::ParentReference;
@@ -1827,7 +1828,7 @@ TOptional<FMeshDescription> FWireTranslatorImpl::GetMeshOfShellNode(AlDagNode& D
 		AlMatrix4x4 AlMatrix;
 		DagNode.inverseGlobalTransformationMatrix(AlMatrix);
 		// TODO: the best way, should be to don't have to apply inverse global transform to the generated mesh
-		TSharedPtr<AlDagNode> TesselatedNode = OpenModelUtils::TesselateDagLeaf(DagNode, ETesselatorType::Fast, ImportParameters.GetChordTolerance());
+		TSharedPtr<AlDagNode> TesselatedNode = OpenModelUtils::TesselateDagLeaf(DagNode, ETesselatorType::Fast, TessellationOptions.ChordTolerance);
 		if (TesselatedNode.IsValid())
 		{
 			// Get the meshes from the dag nodes. Note that removing the mesh's DAG.
@@ -1839,7 +1840,6 @@ TOptional<FMeshDescription> FWireTranslatorImpl::GetMeshOfShellNode(AlDagNode& D
 	return TOptional<FMeshDescription>();
 }
 
-
 TOptional<FMeshDescription> FWireTranslatorImpl::GetMeshOfShellBody(TSharedRef<BodyData> Body, TSharedRef<IDatasmithMeshElement> MeshElement, CADLibrary::FMeshParameters& MeshParameters)
 {
 	TOptional<FMeshDescription> UEMesh = MeshDagNodeWithExternalMesher(Body, MeshElement, MeshParameters);
