CADLibrary refactoring

diff --git a/Engine/Plugins/Enterprise/DatasmithCADImporter/Source/CADInterfaces/Private/CoreTechFileParserImpl.cpp b/Engine/Plugins/Enterprise/DatasmithCADImporter/Source/CADInterfaces/Private/CoreTechFileParserImpl.cpp
--- a/Engine/Plugins/Enterprise/DatasmithCADImporter/Source/CADInterfaces/Private/CoreTechFileParserImpl.cpp
+++ b/Engine/Plugins/Enterprise/DatasmithCADImporter/Source/CADInterfaces/Private/CoreTechFileParserImpl.cpp
@@ -43,12 +43,9 @@ namespace CADLibrary
 	namespace CoreTechFileParserUtils
 	{
 		FString AsFString(CT_STR CtName);
-		uint32 GetSceneFileHash(const uint32 InSGHash, const FImportParameters& ImportParam);
-		uint32 GetGeomFileHash(const uint32 InSGHash, const FImportParameters& ImportParam);
 		void AddFaceIdAttribut(CT_OBJECT_ID NodeId);
 		void GetInstancesAndBodies(CT_OBJECT_ID InComponentId, TArray<CT_OBJECT_ID>& OutInstances, TArray<CT_OBJECT_ID>& OutBodies);
 		void GetCTObjectDisplayDataIds(CT_OBJECT_ID ObjectID, FObjectDisplayDataId& Material);
-		uint32 GetStaticMeshUuid(const TCHAR* OutSgFile, const int32 BodyId);
 		bool GetColor(uint32 ColorHash, FColor& OutColor);
 		bool GetMaterial(uint32 MaterialID, FCADMaterial& OutMaterial);
 		int32 GetIntegerParameterDataValue(CT_OBJECT_ID NodeId, const TCHAR* InMetaDataName);
@@ -614,16 +611,13 @@ namespace CADLibrary
 
 		ReadNodeMetaData(BodyId, Body.MetaData);
 
-		FBodyMesh& BodyMesh = CADFileData.AddBodyMesh(BodyId);
+		FBodyMesh& BodyMesh = CADFileData.AddBodyMesh(BodyId, Body);
 
 		if (uint32 MaterialHash = GetObjectMaterial(Body))
 		{
 			DefaultMaterialHash = MaterialHash;
 		}
 
-		Body.MeshActorName = CoreTechFileParserUtils::GetStaticMeshUuid(*CADFileData.GetSceneGraphFileName(), BodyId);
-		BodyMesh.MeshActorName = Body.MeshActorName;
-
 		CT_FLAGS BodyProperties;
 		CT_BODY_IO::AskProperties(BodyId, BodyProperties);
 
@@ -686,16 +680,13 @@ namespace CADLibrary
 		FArchiveBody& ArchiveBody = CADFileData.GetBodyAt(Index);
 		ReadNodeMetaData(BodyId, ArchiveBody.MetaData);
 
-		FBodyMesh& BodyMesh = CADFileData.AddBodyMesh(BodyId);
+		FBodyMesh& BodyMesh = CADFileData.AddBodyMesh(BodyId, ArchiveBody);
 
 		if (uint32 MaterialHash = GetObjectMaterial(ArchiveBody))
 		{
 			DefaultMaterialHash = MaterialHash;
 		}
 
-		ArchiveBody.MeshActorName = CoreTechFileParserUtils::GetStaticMeshUuid(*CADFileData.GetSceneGraphFileName(), BodyId);
-		BodyMesh.MeshActorName = ArchiveBody.MeshActorName;
-
 		{
 			TSharedRef<CADKernel::FSession> CADKernelSession = MakeShared<CADKernel::FSession>(0.00001 / CADFileData.GetImportParameters().GetMetricUnit());
 
@@ -1681,13 +1672,6 @@ namespace CADLibrary
 			return true;
 		}
 
-		uint32 GetStaticMeshUuid(const TCHAR* OutSgFile, const int32 BodyId)
-		{
-			uint32 BodyUUID = ::GetTypeHash(OutSgFile);
-			BodyUUID = HashCombine(BodyUUID, ::GetTypeHash(BodyId));
-			return BodyUUID;
-		}
-
 		// For each face, adds an integer parameter representing the id of the face to avoid re-indentation of faces in sub CT file
 		// Used by Re-tessellation Rule to Skip Deleted Surfaces
 		void AddFaceIdAttribut(CT_OBJECT_ID NodeId)
diff --git a/Engine/Plugins/Enterprise/DatasmithCADImporter/Source/CADInterfaces/Public/CADFileData.h b/Engine/Plugins/Enterprise/DatasmithCADImporter/Source/CADInterfaces/Public/CADFileData.h
--- a/Engine/Plugins/Enterprise/DatasmithCADImporter/Source/CADInterfaces/Public/CADFileData.h
+++ b/Engine/Plugins/Enterprise/DatasmithCADImporter/Source/CADInterfaces/Public/CADFileData.h
@@ -85,6 +85,9 @@ namespace CADLibrary
 			return FString();
 		}
 
+		/**
+		 * @return the path of CAD cache file
+		 */
 		const FString GetCADCachePath() const
 		{
 			if (IsCacheDefined())
@@ -95,6 +98,15 @@ namespace CADLibrary
 			return FString();
 		}
 
+		const FString GetCachePath() const
+		{
+			if (IsCacheDefined())
+			{
+				return CachePath;
+			}
+			return FString();
+		}
+
 		bool IsCacheDefined() const
 		{
 			return bIsCacheDefined;
@@ -223,13 +235,6 @@ namespace CADLibrary
 			return BodyMesh;
 		}
 
-		FBodyMesh& AddBodyMesh(FCadId BodyId)
-		{
-			FBodyMesh& BodyMesh = BodyMeshes.Emplace_GetRef(BodyId);
-			BodyMesh.MeshActorName = GetStaticMeshHash(BodyId);
-			return BodyMesh;
-		}
-
 		void ExportMeshArchiveFile()
 		{
 			FString MeshArchiveFilePath = GetMeshArchiveFilePath();
