[Alias] Uses the local installed version of Alias

diff --git a/Engine/Plugins/Enterprise/DatasmithCADImporter/DatasmithCADImporter.uplugin b/Engine/Plugins/Enterprise/DatasmithCADImporter/DatasmithCADImporter.uplugin
--- a/Engine/Plugins/Enterprise/DatasmithCADImporter/DatasmithCADImporter.uplugin
+++ b/Engine/Plugins/Enterprise/DatasmithCADImporter/DatasmithCADImporter.uplugin
@@ -66,7 +66,39 @@
 			]
 		},
 		{
-			"Name": "DatasmithWireTranslator",
+			"Name": "DatasmithWireTranslator2020",
+			"Type": "Editor",
+			"LoadingPhase": "Default",
+			"PlatformAllowList": [
+				"Win64"
+			]
+		},
+		{
+			"Name": "DatasmithWireTranslator2021_3",
+			"Type": "Editor",
+			"LoadingPhase": "Default",
+			"PlatformAllowList": [
+				"Win64"
+			]
+		},
+		{
+			"Name": "DatasmithWireTranslator2022",
+			"Type": "Editor",
+			"LoadingPhase": "Default",
+			"PlatformAllowList": [
+				"Win64"
+			]
+		},
+		{
+			"Name": "DatasmithWireTranslator2022_1",
+			"Type": "Editor",
+			"LoadingPhase": "Default",
+			"PlatformAllowList": [
+				"Win64"
+			]
+		},
+		{
+			"Name": "DatasmithWireTranslator2022_2",
 			"Type": "Editor",
 			"LoadingPhase": "Default",
 			"PlatformAllowList": [
@@ -105,9 +137,9 @@
 				"Win64"
 			]
 		},
-    	{
+		{
 			"Name": "DatasmithPLMXMLTranslator",
-    		"Type": "Editor",
+			"Type": "Editor",
 			"LoadingPhase": "Default",
 			"PlatformAllowList": [
 				"Win64"
diff --git a/Engine/Plugins/Enterprise/DatasmithCADImporter/Source/DatasmithWireTranslator/DatasmithWireTranslator.Build.cs b/Engine/Plugins/Enterprise/DatasmithCADImporter/Source/DatasmithWireTranslator/DatasmithWireTranslator.Build.cs
deleted file mode 100644
diff --git a/Engine/Plugins/Enterprise/DatasmithCADImporter/Source/DatasmithWireTranslator/DatasmithWireTranslator2020.Build.cs b/Engine/Plugins/Enterprise/DatasmithCADImporter/Source/DatasmithWireTranslator/DatasmithWireTranslator2020.Build.cs
new file mode 100644
--- /dev/null
+++ b/Engine/Plugins/Enterprise/DatasmithCADImporter/Source/DatasmithWireTranslator/DatasmithWireTranslator2020.Build.cs
@@ -0,0 +1,82 @@
+// Copyright Epic Games, Inc. All Rights Reserved.
+
+using System.IO;
+using UnrealBuildTool;
+
+public abstract class DatasmithWireTranslatorBase : ModuleRules
+{
+	public DatasmithWireTranslatorBase(ReadOnlyTargetRules Target) : base(Target)
+	{
+		PublicDependencyModuleNames.AddRange(
+			new string[]
+			{
+				"Core",
+				"CoreUObject",
+			}
+		);
+
+		PrivateDependencyModuleNames.AddRange(
+			new string[]
+			{
+				"CADInterfaces",
+				"CADKernel",
+				"CADKernelSurface",
+				"CADLibrary",
+				"CADTools",
+				"CoreTechSurface",
+				"DatasmithContent",
+				"DatasmithCore",
+				"DatasmithTranslator",
+				"Engine",
+				"MeshDescription",
+				"ParametricSurface",
+				"StaticMeshDescription",
+			}
+		);
+
+		if (Target.bBuildEditor == true)
+		{
+			PrivateDependencyModuleNames.AddRange(
+				new string[]
+				{
+ 					"MessageLog",
+ 					"UnrealEd",
+				}
+			);
+		}
+
+		PublicDefinitions.Add(GetAliasDefinition());
+
+		if (System.Type.GetType(GetAliasVersion()) != null)
+		{
+			PrivateDependencyModuleNames.Add(GetAliasVersion());
+		}
+
+		if (System.Type.GetType("CoreTech") != null)
+		{
+			PrivateDependencyModuleNames.Add("CoreTech");
+		}
+	}
+
+	public abstract string GetAliasVersion();
+	public abstract string GetAliasDefinition();
+
+}
+
+public class DatasmithWireTranslator2020 : DatasmithWireTranslatorBase
+{
+	public DatasmithWireTranslator2020(ReadOnlyTargetRules Target) 
+		: base(Target)
+	{
+	}
+
+	public override string GetAliasVersion()
+	{
+		return "OpenModel2020";
+	}
+	
+	public override string GetAliasDefinition()
+	{
+		return "OPEN_MODEL_2020";
+	}
+}
diff --git a/Engine/Plugins/Enterprise/DatasmithCADImporter/Source/DatasmithWireTranslator/DatasmithWireTranslator2021_3.Build.cs b/Engine/Plugins/Enterprise/DatasmithCADImporter/Source/DatasmithWireTranslator/DatasmithWireTranslator2021_3.Build.cs
new file mode 100644
--- /dev/null
+++ b/Engine/Plugins/Enterprise/DatasmithCADImporter/Source/DatasmithWireTranslator/DatasmithWireTranslator2021_3.Build.cs
@@ -0,0 +1,22 @@
+// Copyright Epic Games, Inc. All Rights Reserved.
+
+using System.IO;
+using UnrealBuildTool;
+
+public class DatasmithWireTranslator2021_3 : DatasmithWireTranslatorBase
+{
+	public DatasmithWireTranslator2021_3(ReadOnlyTargetRules Target) 
+		: base(Target)
+	{
+	}
+
+	public override string GetAliasVersion()
+	{
+		return "OpenModel2021_3";
+	}
+	
+	public override string GetAliasDefinition()
+	{
+		return "OPEN_MODEL_2021_3";
+	}
+}
diff --git a/Engine/Plugins/Enterprise/DatasmithCADImporter/Source/DatasmithWireTranslator/DatasmithWireTranslator2022.Build.cs b/Engine/Plugins/Enterprise/DatasmithCADImporter/Source/DatasmithWireTranslator/DatasmithWireTranslator2022.Build.cs
new file mode 100644
--- /dev/null
+++ b/Engine/Plugins/Enterprise/DatasmithCADImporter/Source/DatasmithWireTranslator/DatasmithWireTranslator2022.Build.cs
@@ -0,0 +1,22 @@
+// Copyright Epic Games, Inc. All Rights Reserved.
+
+using System.IO;
+using UnrealBuildTool;
+
+public class DatasmithWireTranslator2022 : DatasmithWireTranslatorBase
+{
+	public DatasmithWireTranslator2022(ReadOnlyTargetRules Target) 
+		: base(Target)
+	{
+	}
+
+	public override string GetAliasVersion()
+	{
+		return "OpenModel2022";
+	}
+	
+	public override string GetAliasDefinition()
+	{
+	return "OPEN_MODEL_2022";
+	}
+}
\ No newline at end of file
diff --git a/Engine/Plugins/Enterprise/DatasmithCADImporter/Source/DatasmithWireTranslator/DatasmithWireTranslator2022_1.Build.cs b/Engine/Plugins/Enterprise/DatasmithCADImporter/Source/DatasmithWireTranslator/DatasmithWireTranslator2022_1.Build.cs
new file mode 100644
--- /dev/null
+++ b/Engine/Plugins/Enterprise/DatasmithCADImporter/Source/DatasmithWireTranslator/DatasmithWireTranslator2022_1.Build.cs
@@ -0,0 +1,22 @@
+// Copyright Epic Games, Inc. All Rights Reserved.
+
+using System.IO;
+using UnrealBuildTool;
+
+public class DatasmithWireTranslator2022_1 : DatasmithWireTranslatorBase
+{
+	public DatasmithWireTranslator2022_1(ReadOnlyTargetRules Target) 
+		: base(Target)
+	{
+	}
+
+	public override string GetAliasVersion()
+	{
+		return "OpenModel2022_1";
+	}
+	
+	public override string GetAliasDefinition()
+	{
+		return "OPEN_MODEL_2022_1";
+	}
+}
\ No newline at end of file
diff --git a/Engine/Plugins/Enterprise/DatasmithCADImporter/Source/DatasmithWireTranslator/DatasmithWireTranslator2022_2.Build.cs b/Engine/Plugins/Enterprise/DatasmithCADImporter/Source/DatasmithWireTranslator/DatasmithWireTranslator2022_2.Build.cs
new file mode 100644
--- /dev/null
+++ b/Engine/Plugins/Enterprise/DatasmithCADImporter/Source/DatasmithWireTranslator/DatasmithWireTranslator2022_2.Build.cs
@@ -0,0 +1,22 @@
+// Copyright Epic Games, Inc. All Rights Reserved.
+
+using System.IO;
+using UnrealBuildTool;
+
+public class DatasmithWireTranslator2022_2 : DatasmithWireTranslatorBase
+{
+	public DatasmithWireTranslator2022_2(ReadOnlyTargetRules Target) 
+		: base(Target)
+	{
+	}
+
+	public override string GetAliasVersion()
+	{
+		return "OpenModel2022_2";
+	}
+
+	public override string GetAliasDefinition()
+	{
+		return "OPEN_MODEL_2022_2";
+	}
+}
\ No newline at end of file
diff --git a/Engine/Plugins/Enterprise/DatasmithCADImporter/Source/DatasmithWireTranslator/Private/DatasmithWireTranslator.cpp b/Engine/Plugins/Enterprise/DatasmithCADImporter/Source/DatasmithWireTranslator/Private/DatasmithWireTranslator.cpp
--- a/Engine/Plugins/Enterprise/DatasmithCADImporter/Source/DatasmithWireTranslator/Private/DatasmithWireTranslator.cpp
+++ b/Engine/Plugins/Enterprise/DatasmithCADImporter/Source/DatasmithWireTranslator/Private/DatasmithWireTranslator.cpp
@@ -65,6 +65,40 @@ DEFINE_LOG_CATEGORY_STATIC(LogDatasmithWireTranslator, Log, All);
 
 #ifdef USE_OPENMODEL
 
+const uint64 LibAliasNext_Version = 17881307937833405;
+const uint64 LibAlias2022_2_0_Version = 7881307937833405;
+const uint64 LibAlias2022_1_0_Version = 7881303642865885;
+const uint64 LibAlias2022_0_1_Version = 7881299347964005;
+const uint64 LibAlias2021_3_2_Version = 7599833027117059;
+const uint64 LibAlias2021_3_1_Version = 7599824433840131;
+const uint64 LibAlias2021_3_0_Version = 7599824424206339;
+const uint64 LibAlias2021_Version = 7599824377020416;
+const uint64 LibAlias2020_Version = 7318349414924288;
+const uint64 LibAlias2019_Version = 7300000000000000;
+
+#if defined(OPEN_MODEL_2020)
+const uint64 LibAliasVersionMin = LibAlias2019_Version;
+const uint64 LibAliasVersionMax = LibAlias2021_3_0_Version;
+const FString AliasVersionChar = TEXT("AliasStudio 2020, Model files");
+#elif defined(OPEN_MODEL_2021_3)
+const uint64 LibAliasVersionMin = LibAlias2021_3_0_Version;
+const uint64 LibAliasVersionMax = LibAlias2022_0_1_Version;
+const FString  AliasVersionChar = TEXT("AliasStudio 2021.3, Model files");
+#elif defined(OPEN_MODEL_2022)
+const uint64 LibAliasVersionMin = LibAlias2022_0_1_Version;
+const uint64 LibAliasVersionMax = LibAlias2022_1_0_Version;
+const FString AliasVersionChar = TEXT("AliasStudio 2022, Model files");
+#elif defined(OPEN_MODEL_2022_1)
+const uint64 LibAliasVersionMin = LibAlias2022_1_0_Version;
+const uint64 LibAliasVersionMax = LibAlias2022_2_0_Version;
+const FString AliasVersionChar = TEXT("AliasStudio 2022.1, Model files");
+#elif defined(OPEN_MODEL_2022_2)
+const uint64 LibAliasVersionMin = LibAlias2022_2_0_Version;
+const uint64 LibAliasVersionMax = LibAliasNext_Version;
+const FString AliasVersionChar = TEXT("AliasStudio 2022.2, Model files");
+#endif
+
+
 using namespace CADLibrary;
 
 class BodyData
@@ -136,12 +170,57 @@ public:
 		, NumCRCErrors(0)
 	{
 		// Set ProductName, ProductVersion in DatasmithScene for Analytics purpose
-		// application_name is something like "Catia V5"
+		uint64 AliasFileVersion = FPlatformMisc::GetFileVersion(TEXT("libalias_api.dll"));
+
 		DatasmithScene->SetHost(TEXT("Alias"));
 		DatasmithScene->SetVendor(TEXT("Autodesk"));
-		DatasmithScene->SetExporterSDKVersion(TEXT("2022"));
 		DatasmithScene->SetProductName(TEXT("Alias Tools"));
-		DatasmithScene->SetProductVersion(TEXT("Alias 2022"));
+
+		if (AliasFileVersion < LibAlias2020_Version)
+		{
+			DatasmithScene->SetExporterSDKVersion(TEXT("2019"));
+			DatasmithScene->SetProductVersion(TEXT("Alias 2019"));
+		}
+		else if (AliasFileVersion < LibAlias2021_Version)
+		{
+			DatasmithScene->SetExporterSDKVersion(TEXT("2020"));
+			DatasmithScene->SetProductVersion(TEXT("Alias 2020"));
+		}
+		else if (AliasFileVersion < LibAlias2021_3_0_Version)
+		{
+			DatasmithScene->SetExporterSDKVersion(TEXT("2021.0"));
+			DatasmithScene->SetProductVersion(TEXT("Alias 2021.0"));
+		}
+		else if (AliasFileVersion < LibAlias2021_3_1_Version)
+		{
+			DatasmithScene->SetExporterSDKVersion(TEXT("2021.3.0"));
+			DatasmithScene->SetProductVersion(TEXT("Alias 2021.3.0"));
+		}
+		else if (AliasFileVersion < LibAlias2021_3_2_Version)
+		{
+			DatasmithScene->SetExporterSDKVersion(TEXT("2021.3.1"));
+			DatasmithScene->SetProductVersion(TEXT("Alias 2021.3.1"));
+		}
+		else if (AliasFileVersion < LibAlias2022_0_1_Version)
+		{
+			DatasmithScene->SetExporterSDKVersion(TEXT("2021.3.2"));
+			DatasmithScene->SetProductVersion(TEXT("Alias 2021.3.2"));
+		}
+		else if (AliasFileVersion < LibAlias2022_1_0_Version)
+		{
+			DatasmithScene->SetExporterSDKVersion(TEXT("2022"));
+			DatasmithScene->SetProductVersion(TEXT("Alias 2022"));
+		}
+		else if (AliasFileVersion < LibAlias2022_2_0_Version)
+		{
+			DatasmithScene->SetExporterSDKVersion(TEXT("2022.1"));
+			DatasmithScene->SetProductVersion(TEXT("Alias 2022.1"));
+		}
+		else
+		{
+			DatasmithScene->SetExporterSDKVersion(TEXT("2022.2"));
+			DatasmithScene->SetProductVersion(TEXT("Alias 2022.2"));
+		}
 
 		FImportParameters ImportParameters(0.01, 1);
 		if(FImportParameters::bGDisableCADKernelTessellation)
@@ -2048,7 +2127,26 @@ void FDatasmithWireTranslator::Initialize(FDatasmithTranslatorCapabilities& OutC
 #ifdef USE_OPENMODEL
 		if (FPlatformProcess::GetDllHandle(TEXT("libalias_api.dll")))
 		{
-			OutCapabilities.SupportedFileFormats.Add(FFileFormatInfo{ TEXT("wire"), TEXT("AliasStudio, Model files") });
+			// Check installed version of Alias Tools because binaries before 2021.3 are not compatible with Alias 2022
+			uint64 FileVersion = FPlatformMisc::GetFileVersion(TEXT("libalias_api.dll"));
+
+#ifdef OPEN_MODEL_2020
+			if (LibAlias2020_Version < FileVersion && FileVersion < LibAlias2021_Version)
+			{
+				UE_LOG(LogDatasmithWireTranslator, Warning, TEXT(WRONG_VERSION_TEXT));
+				OutCapabilities.bIsEnabled = false;
+				return;
+			}
+#endif
+
+			if (LibAliasVersionMin <= FileVersion && FileVersion < LibAliasVersionMax)
+			{
+				OutCapabilities.SupportedFileFormats.Add(FFileFormatInfo{ TEXT("wire"), *AliasVersionChar });
+				OutCapabilities.bIsEnabled = true;
+				return;
+			}
+
+			OutCapabilities.bIsEnabled = false;
 			return;
 		}
 #endif
@@ -2070,38 +2168,6 @@ bool FDatasmithWireTranslator::IsSourceSupported(const FDatasmithSceneSource& So
 bool FDatasmithWireTranslator::LoadScene(TSharedRef<IDatasmithScene> OutScene)
 {
 #ifdef USE_OPENMODEL
-	// Check installed version of Alias Tools because binaries before 2021.3 are not compatible with lib Alias 2022.0
-	const uint64 LibAlias2021_3Version = 7599833027117059;
-	uint64 FileVersion = FPlatformMisc::GetFileVersion(TEXT("libalias_api.dll"));
-
-	if (FileVersion < LibAlias2021_3Version)
-	{
-#if WITH_EDITOR
-		// Display message and abort import
-		FMessageLogModule& MessageLogModule = FModuleManager::LoadModuleChecked<FMessageLogModule>("MessageLog");
-
-		TSharedPtr<IMessageLogListing> LogListing = MessageLogModule.GetLogListing("DatasmithWireTranslator");
-
-		if (LogListing.IsValid())
-		{
-			LogListing->SetLabel(LOCTEXT("MessageLogging", "Datasmith Wire Translator"));
-
-			FText Message = LOCTEXT("WrongVersion", WRONG_VERSION_TEXT);
-			LogListing->AddMessage(FTokenizedMessage::Create(EMessageSeverity::Error, Message));
-
-			LogListing->NotifyIfAnyMessages(LOCTEXT("LibraryIssue", "There was an issue with the wire import"), EMessageSeverity::Info);
-		}
-		else
-		{
-			UE_LOG(LogDatasmithWireTranslator, Error, TEXT(WRONG_VERSION_TEXT));
-	}
-#else
-		UE_LOG(LogDatasmithWireTranslator, Error, TEXT(WRONG_VERSION_TEXT));
-#endif
-
-		return false;
-}
-
 	const FString& Filename = GetSource().GetSourceFile();
 
 	Translator = MakeShared<FWireTranslatorImpl>(Filename, OutScene);
diff --git a/Engine/Plugins/Enterprise/DatasmithCADImporter/Source/DatasmithWireTranslator/Public/DatasmithWireTranslatorModule.h b/Engine/Plugins/Enterprise/DatasmithCADImporter/Source/DatasmithWireTranslator/Public/DatasmithWireTranslatorModule.h
--- a/Engine/Plugins/Enterprise/DatasmithCADImporter/Source/DatasmithWireTranslator/Public/DatasmithWireTranslatorModule.h
+++ b/Engine/Plugins/Enterprise/DatasmithCADImporter/Source/DatasmithWireTranslator/Public/DatasmithWireTranslatorModule.h
@@ -6,7 +6,17 @@
 #include "Modules/ModuleInterface.h"
 #include "Modules/ModuleManager.h"
 
-#define DATASMITHWIRETRANSLATOR_MODULE_NAME TEXT("DatasmithWireTranslator")
+#if defined(OPEN_MODEL_2020)
+#define DATASMITHWIRETRANSLATOR_MODULE_NAME TEXT("DatasmithWireTranslator2020")
+#elif defined(OPEN_MODEL_2021_3)
+#define DATASMITHWIRETRANSLATOR_MODULE_NAME TEXT("DatasmithWireTranslator2021_3")
+#elif defined(OPEN_MODEL_2022)
+#define DATASMITHWIRETRANSLATOR_MODULE_NAME TEXT("DatasmithWireTranslator2022")
+#elif defined(OPEN_MODEL_2022_1)
+#define DATASMITHWIRETRANSLATOR_MODULE_NAME TEXT("DatasmithWireTranslator2022_1")
+#elif defined(OPEN_MODEL_2022_2)
+#define DATASMITHWIRETRANSLATOR_MODULE_NAME TEXT("DatasmithWireTranslator2022_2")
+#endif
 
 /**
  * Datasmith Translator for .wire files.
@@ -18,7 +28,7 @@ public:
     static FDatasmithWireTranslatorModule& Get()
     {
         return FModuleManager::LoadModuleChecked< FDatasmithWireTranslatorModule >(DATASMITHWIRETRANSLATOR_MODULE_NAME);
-    }
+	}
 
     static bool IsAvailable()
     {
