[Alias] fix version number of Alias 2019

diff --git a/Engine/Plugins/Enterprise/DatasmithCADImporter/Source/DatasmithWireTranslator/Private/DatasmithWireTranslator.cpp b/Engine/Plugins/Enterprise/DatasmithCADImporter/Source/DatasmithWireTranslator/Private/DatasmithWireTranslator.cpp
--- a/Engine/Plugins/Enterprise/DatasmithCADImporter/Source/DatasmithWireTranslator/Private/DatasmithWireTranslator.cpp
+++ b/Engine/Plugins/Enterprise/DatasmithCADImporter/Source/DatasmithWireTranslator/Private/DatasmithWireTranslator.cpp
@@ -74,7 +74,7 @@ const uint64 LibAlias2021_3_1_Version = 7599824433840131;
 const uint64 LibAlias2021_3_0_Version = 7599824424206339;
 const uint64 LibAlias2021_Version = 7599824377020416;
 const uint64 LibAlias2020_Version = 7318349414924288;
-const uint64 LibAlias2019_Version = 7300000000000000;
+const uint64 LibAlias2019_Version = 5000000000000000;
 
 #if defined(OPEN_MODEL_2020)
 const uint64 LibAliasVersionMin = LibAlias2019_Version;
