- Fix the DatasmithWireTranslator monolithic linker errors (multiple modules sharing .cpp code), by using a namespace in each version

diff --git a/Engine/Plugins/Enterprise/DatasmithCADImporter/Source/DatasmithWireTranslator/DatasmithWireTranslator2020.Build.cs b/Engine/Plugins/Enterprise/DatasmithCADImporter/Source/DatasmithWireTranslator/DatasmithWireTranslator2020.Build.cs
--- a/Engine/Plugins/Enterprise/DatasmithCADImporter/Source/DatasmithWireTranslator/DatasmithWireTranslator2020.Build.cs
+++ b/Engine/Plugins/Enterprise/DatasmithCADImporter/Source/DatasmithWireTranslator/DatasmithWireTranslator2020.Build.cs
@@ -46,6 +46,8 @@ public abstract class DatasmithWireTranslatorBase : ModuleRules
 		}
 
 		PublicDefinitions.Add(GetAliasDefinition());
+		PublicDefinitions.Add($"UE_DATASMITHWIRETRANSLATOR_NAMESPACE={GetAliasDefinition()}Namespace");
+		PublicDefinitions.Add($"UE_DATASMITHWIRETRANSLATOR_MODULE_NAME={GetType().Name}");
 
 		if (System.Type.GetType(GetAliasVersion()) != null)
 		{
diff --git a/Engine/Plugins/Enterprise/DatasmithCADImporter/Source/DatasmithWireTranslator/Private/AliasBrepConverter.h b/Engine/Plugins/Enterprise/DatasmithCADImporter/Source/DatasmithWireTranslator/Private/AliasBrepConverter.h
--- a/Engine/Plugins/Enterprise/DatasmithCADImporter/Source/DatasmithWireTranslator/Private/AliasBrepConverter.h
+++ b/Engine/Plugins/Enterprise/DatasmithCADImporter/Source/DatasmithWireTranslator/Private/AliasBrepConverter.h
@@ -4,6 +4,9 @@
 
 class AlDagNode;
 
+namespace UE_DATASMITHWIRETRANSLATOR_NAMESPACE
+{
+
 // Defined the reference in which the object has to be defined
 enum class EAliasObjectReference
 {
@@ -18,3 +21,4 @@ public:
 	virtual bool AddBRep(AlDagNode& DagNode, EAliasObjectReference ObjectReference) = 0;
 };
 
+}
\ No newline at end of file
diff --git a/Engine/Plugins/Enterprise/DatasmithCADImporter/Source/DatasmithWireTranslator/Private/AliasModelToCADKernelConverter.cpp b/Engine/Plugins/Enterprise/DatasmithCADImporter/Source/DatasmithWireTranslator/Private/AliasModelToCADKernelConverter.cpp
--- a/Engine/Plugins/Enterprise/DatasmithCADImporter/Source/DatasmithWireTranslator/Private/AliasModelToCADKernelConverter.cpp
+++ b/Engine/Plugins/Enterprise/DatasmithCADImporter/Source/DatasmithWireTranslator/Private/AliasModelToCADKernelConverter.cpp
@@ -40,13 +40,17 @@
 #include "AlTrimRegion.h"
 #include "AlTM.h"
 
+using namespace CADKernel;
+
+namespace UE_DATASMITHWIRETRANSLATOR_NAMESPACE
+{
+
 namespace AliasToCADKernelUtils
 {
+
 	template<typename Surface_T>
-	TSharedPtr<CADKernel::FSurface> AddNURBSSurface(double GeometricTolerance, Surface_T& AliasSurface, EAliasObjectReference InObjectReference, const AlMatrix4x4& InAlMatrix)
+	TSharedPtr<FSurface> AddNURBSSurface(double GeometricTolerance, Surface_T& AliasSurface, EAliasObjectReference InObjectReference, const AlMatrix4x4& InAlMatrix)
 	{
-		using namespace CADKernel;
-
 		FNurbsSurfaceHomogeneousData NURBSData;
 		NURBSData.bSwapUV= true;
 		NURBSData.bIsRational = true;
@@ -92,10 +96,8 @@ namespace AliasToCADKernelUtils
 	}
 }
 
-TSharedPtr<CADKernel::FTopologicalEdge> FAliasModelToCADKernelConverter::AddEdge(const AlTrimCurve& AliasTrimCurve, TSharedPtr<CADKernel::FSurface>& CarrierSurface)
+TSharedPtr<FTopologicalEdge> FAliasModelToCADKernelConverter::AddEdge(const AlTrimCurve& AliasTrimCurve, TSharedPtr<CADKernel::FSurface>& CarrierSurface)
 {
-	using namespace CADKernel;
-
 	FNurbsCurveData NurbsCurveData;
 
 	NurbsCurveData.Degree = AliasTrimCurve.degree();
@@ -144,10 +146,8 @@ TSharedPtr<CADKernel::FTopologicalEdge> FAliasModelToCADKernelConverter::AddEdge
 	return Edge;
 }
 
-TSharedPtr<CADKernel::FTopologicalLoop> FAliasModelToCADKernelConverter::AddLoop(const AlTrimBoundary& TrimBoundary, TSharedPtr<CADKernel::FSurface>& CarrierSurface)
+TSharedPtr<FTopologicalLoop> FAliasModelToCADKernelConverter::AddLoop(const AlTrimBoundary& TrimBoundary, TSharedPtr<CADKernel::FSurface>& CarrierSurface)
 {
-	using namespace CADKernel;
-
 	TArray<TSharedPtr<FTopologicalEdge>> Edges;
 	TArray<CADKernel::EOrientation> Directions;
 
@@ -172,8 +172,6 @@ TSharedPtr<CADKernel::FTopologicalLoop> FAliasModelToCADKernelConverter::AddLoop
 
 void FAliasModelToCADKernelConverter::LinkEdgesLoop(const AlTrimBoundary& TrimBoundary, CADKernel::FTopologicalLoop& Loop)
 {
-	using namespace CADKernel;
-
 	for (TUniquePtr<AlTrimCurve> TrimCurve(TrimBoundary.firstCurve()); TrimCurve.IsValid(); TrimCurve = TUniquePtr<AlTrimCurve>(TrimCurve->nextCurve()))
 	{
 		TSharedPtr<FTopologicalEdge>* Edge = AlEdge2CADKernelEdge.Find(TrimCurve->fSpline);
@@ -201,8 +199,6 @@ void FAliasModelToCADKernelConverter::LinkEdgesLoop(const AlTrimBoundary& TrimBo
 
 TSharedPtr<CADKernel::FTopologicalFace> FAliasModelToCADKernelConverter::AddTrimRegion(const AlTrimRegion& TrimRegion, EAliasObjectReference InObjectReference, const AlMatrix4x4& InAlMatrix, bool bInOrientation)
 {
-	using namespace CADKernel;
-
 	TSharedPtr<FSurface> Surface = AliasToCADKernelUtils::AddNURBSSurface(GeometricTolerance, TrimRegion, InObjectReference, InAlMatrix);
 	if (!Surface.IsValid())
 	{
@@ -235,8 +231,6 @@ TSharedPtr<CADKernel::FTopologicalFace> FAliasModelToCADKernelConverter::AddTrim
 
 void FAliasModelToCADKernelConverter::AddFace(const AlSurface& Surface, EAliasObjectReference InObjectReference, const AlMatrix4x4& InAlMatrix, bool bInOrientation, TSharedRef<CADKernel::FShell>& Shell)
 {
-	using namespace CADKernel;
-
 	TUniquePtr<AlTrimRegion> TrimRegion(Surface.firstTrimRegion());
 	if (TrimRegion.IsValid())
 	{
@@ -262,8 +256,6 @@ void FAliasModelToCADKernelConverter::AddFace(const AlSurface& Surface, EAliasOb
 
 void FAliasModelToCADKernelConverter::AddShell(const AlShell& InShell, EAliasObjectReference InObjectReference, const AlMatrix4x4& InAlMatrix, bool bInOrientation, TSharedRef<CADKernel::FShell>& CADKernelShell)
 {
-	using namespace CADKernel;
-
 	for(TUniquePtr<AlTrimRegion> TrimRegion(InShell.firstTrimRegion()); TrimRegion.IsValid(); TrimRegion = TUniquePtr<AlTrimRegion>(TrimRegion->nextRegion()))
 	{
 		TSharedPtr<FTopologicalFace> Face = AddTrimRegion(*TrimRegion, InObjectReference, InAlMatrix, bInOrientation);
@@ -276,8 +268,6 @@ void FAliasModelToCADKernelConverter::AddShell(const AlShell& InShell, EAliasObj
 
 bool FAliasModelToCADKernelConverter::AddBRep(AlDagNode& DagNode, EAliasObjectReference InObjectReference)
 {
-	using namespace CADKernel;
-
 	AlEdge2CADKernelEdge.Empty();
 
 	TSharedRef<FBody> CADKernelBody = FEntity::MakeShared<FBody>();
@@ -339,5 +329,6 @@ bool FAliasModelToCADKernelConverter::AddBRep(AlDagNode& DagNode, EAliasObjectRe
 	return true;
 }
 
+}
 
 #endif
\ No newline at end of file
diff --git a/Engine/Plugins/Enterprise/DatasmithCADImporter/Source/DatasmithWireTranslator/Private/AliasModelToCADKernelConverter.h b/Engine/Plugins/Enterprise/DatasmithCADImporter/Source/DatasmithWireTranslator/Private/AliasModelToCADKernelConverter.h
--- a/Engine/Plugins/Enterprise/DatasmithCADImporter/Source/DatasmithWireTranslator/Private/AliasModelToCADKernelConverter.h
+++ b/Engine/Plugins/Enterprise/DatasmithCADImporter/Source/DatasmithWireTranslator/Private/AliasModelToCADKernelConverter.h
@@ -28,6 +28,9 @@ namespace CADKernel
 	class FTopologicalLoop;
 }
 
+namespace UE_DATASMITHWIRETRANSLATOR_NAMESPACE
+{
+
 class FAliasModelToCADKernelConverter : public FCADModelToCADKernelConverterBase, public IAliasBRepConverter
 {
 
@@ -59,3 +62,4 @@ protected:
 	TMap<void*, TSharedPtr<CADKernel::FTopologicalEdge>>  AlEdge2CADKernelEdge;
 };
 
+}
\ No newline at end of file
diff --git a/Engine/Plugins/Enterprise/DatasmithCADImporter/Source/DatasmithWireTranslator/Private/AliasModelToCoretechConverter.cpp b/Engine/Plugins/Enterprise/DatasmithCADImporter/Source/DatasmithWireTranslator/Private/AliasModelToCoretechConverter.cpp
--- a/Engine/Plugins/Enterprise/DatasmithCADImporter/Source/DatasmithWireTranslator/Private/AliasModelToCoretechConverter.cpp
+++ b/Engine/Plugins/Enterprise/DatasmithCADImporter/Source/DatasmithWireTranslator/Private/AliasModelToCoretechConverter.cpp
@@ -27,6 +27,9 @@
 
 using namespace CADLibrary;
 
+namespace UE_DATASMITHWIRETRANSLATOR_NAMESPACE
+{
+
 namespace AliasToCoreTechUtils
 {
 	template<typename Surface_T>
@@ -262,4 +265,5 @@ bool FAliasModelToCoretechConverter::AddBRep(AlDagNode& DagNode, EAliasObjectRef
 	return false;
 }
 
+}
 #endif
diff --git a/Engine/Plugins/Enterprise/DatasmithCADImporter/Source/DatasmithWireTranslator/Private/AliasModelToCoretechConverter.h b/Engine/Plugins/Enterprise/DatasmithCADImporter/Source/DatasmithWireTranslator/Private/AliasModelToCoretechConverter.h
--- a/Engine/Plugins/Enterprise/DatasmithCADImporter/Source/DatasmithWireTranslator/Private/AliasModelToCoretechConverter.h
+++ b/Engine/Plugins/Enterprise/DatasmithCADImporter/Source/DatasmithWireTranslator/Private/AliasModelToCoretechConverter.h
@@ -16,6 +16,9 @@ class AlTrimCurve;
 class AlTrimRegion;
 class AlCurve;
 
+namespace UE_DATASMITHWIRETRANSLATOR_NAMESPACE
+{
+
 typedef double AlMatrix4x4[4][4];
 
 class FAliasModelToCoretechConverter : public FCADModelToCoretechConverterBase, public IAliasBRepConverter
@@ -46,3 +49,4 @@ protected:
 	TMap<void*, uint64>  AlEdge2CTEdge;
 };
 
+}
\ No newline at end of file
diff --git a/Engine/Plugins/Enterprise/DatasmithCADImporter/Source/DatasmithWireTranslator/Private/DatasmithWireTranslator.cpp b/Engine/Plugins/Enterprise/DatasmithCADImporter/Source/DatasmithWireTranslator/Private/DatasmithWireTranslator.cpp
--- a/Engine/Plugins/Enterprise/DatasmithCADImporter/Source/DatasmithWireTranslator/Private/DatasmithWireTranslator.cpp
+++ b/Engine/Plugins/Enterprise/DatasmithCADImporter/Source/DatasmithWireTranslator/Private/DatasmithWireTranslator.cpp
@@ -63,6 +63,9 @@ DEFINE_LOG_CATEGORY_STATIC(LogDatasmithWireTranslator, Log, All);
 #define WRONG_VERSION_TEXT "Unsupported version of Alias detected. Please upgrade to Alias 2021.3 (or later version)."
 #define CAD_INTERFACE_UNAVAILABLE "CAD Interface module is unavailable. Meshing will be done by Alias."
 
+namespace UE_DATASMITHWIRETRANSLATOR_NAMESPACE
+{
+
 #ifdef USE_OPENMODEL
 
 const uint64 LibAliasNext_Version = 17881307937833405;
@@ -2216,5 +2219,7 @@ void FDatasmithWireTranslator::SetSceneImportOptions(TArray<TStrongObjectPtr<UDa
 #endif
 }
 
+} // namespace
+
 #undef LOCTEXT_NAMESPACE // "DatasmithWireTranslator"
 
diff --git a/Engine/Plugins/Enterprise/DatasmithCADImporter/Source/DatasmithWireTranslator/Private/DatasmithWireTranslator.h b/Engine/Plugins/Enterprise/DatasmithCADImporter/Source/DatasmithWireTranslator/Private/DatasmithWireTranslator.h
--- a/Engine/Plugins/Enterprise/DatasmithCADImporter/Source/DatasmithWireTranslator/Private/DatasmithWireTranslator.h
+++ b/Engine/Plugins/Enterprise/DatasmithCADImporter/Source/DatasmithWireTranslator/Private/DatasmithWireTranslator.h
@@ -7,8 +7,11 @@
 #include "ParametricSurfaceTranslator.h"
 #include "UObject/ObjectMacros.h"
 
-class FWireTranslatorImpl;
 
+namespace UE_DATASMITHWIRETRANSLATOR_NAMESPACE
+{
+
+class FWireTranslatorImpl;
 
 class FDatasmithWireTranslator : public FParametricSurfaceTranslator
 {
@@ -41,3 +44,5 @@ protected:
 private:
 	TSharedPtr<FWireTranslatorImpl> Translator;
 };
+
+}
\ No newline at end of file
diff --git a/Engine/Plugins/Enterprise/DatasmithCADImporter/Source/DatasmithWireTranslator/Private/DatasmithWireTranslatorModule.cpp b/Engine/Plugins/Enterprise/DatasmithCADImporter/Source/DatasmithWireTranslator/Private/DatasmithWireTranslatorModule.cpp
--- a/Engine/Plugins/Enterprise/DatasmithCADImporter/Source/DatasmithWireTranslator/Private/DatasmithWireTranslatorModule.cpp
+++ b/Engine/Plugins/Enterprise/DatasmithCADImporter/Source/DatasmithWireTranslator/Private/DatasmithWireTranslatorModule.cpp
@@ -9,7 +9,12 @@
 
 #include "DatasmithWireTranslator.h"
 
-IMPLEMENT_MODULE(FDatasmithWireTranslatorModule, DatasmithWireTranslator);
+// need this macro wrapper to expand the UE_DATASMITHWIRETRANSLATOR_MODULE_NAME macro and not create symbols with "UE_DATASMITHWIRETRANSLATOR_MODULE_NAME" in the token
+#define IMPLEMENT_MODULE_WRAPPER(ModuleName) IMPLEMENT_MODULE(UE_DATASMITHWIRETRANSLATOR_NAMESPACE::FDatasmithWireTranslatorModule, ModuleName);
+IMPLEMENT_MODULE_WRAPPER(UE_DATASMITHWIRETRANSLATOR_MODULE_NAME)
+
+namespace UE_DATASMITHWIRETRANSLATOR_NAMESPACE
+{
 
 void FDatasmithWireTranslatorModule::StartupModule()
 {
@@ -29,3 +34,5 @@ FString FDatasmithWireTranslatorModule::GetTempDir() const
 {
 	return TempDir;
 }
+
+}
\ No newline at end of file
diff --git a/Engine/Plugins/Enterprise/DatasmithCADImporter/Source/DatasmithWireTranslator/Private/OpenModelUtils.cpp b/Engine/Plugins/Enterprise/DatasmithCADImporter/Source/DatasmithWireTranslator/Private/OpenModelUtils.cpp
--- a/Engine/Plugins/Enterprise/DatasmithCADImporter/Source/DatasmithWireTranslator/Private/OpenModelUtils.cpp
+++ b/Engine/Plugins/Enterprise/DatasmithCADImporter/Source/DatasmithWireTranslator/Private/OpenModelUtils.cpp
@@ -17,6 +17,9 @@
 #include "StaticMeshAttributes.h"
 #include "StaticMeshOperations.h"
 
+namespace UE_DATASMITHWIRETRANSLATOR_NAMESPACE
+{
+
 using namespace OpenModelUtils;
 
 const TCHAR* OpenModelUtils::AlObjectTypeToString(AlObjectType type)
@@ -1833,6 +1836,8 @@ TSharedPtr<AlDagNode> OpenModelUtils::TesselateDagLeaf(const AlDagNode& DagLeaf,
 	}
 }
 
+}
+
 #endif
 
 
diff --git a/Engine/Plugins/Enterprise/DatasmithCADImporter/Source/DatasmithWireTranslator/Private/OpenModelUtils.h b/Engine/Plugins/Enterprise/DatasmithCADImporter/Source/DatasmithWireTranslator/Private/OpenModelUtils.h
--- a/Engine/Plugins/Enterprise/DatasmithCADImporter/Source/DatasmithWireTranslator/Private/OpenModelUtils.h
+++ b/Engine/Plugins/Enterprise/DatasmithCADImporter/Source/DatasmithWireTranslator/Private/OpenModelUtils.h
@@ -22,6 +22,8 @@ struct FMeshDescription;
 
 typedef double AlMatrix4x4[4][4];
 
+namespace UE_DATASMITHWIRETRANSLATOR_NAMESPACE
+{
 
 enum class ETesselatorType : uint8
 {
@@ -76,6 +78,7 @@ namespace OpenModelUtils
 	TSharedPtr<AlDagNode> TesselateDagLeaf(const AlDagNode& DagLeaf, ETesselatorType TessType, double Tolerance);
 }
 
+}
 
 #endif
 
diff --git a/Engine/Plugins/Enterprise/DatasmithCADImporter/Source/DatasmithWireTranslator/Public/DatasmithWireTranslatorModule.h b/Engine/Plugins/Enterprise/DatasmithCADImporter/Source/DatasmithWireTranslator/Public/DatasmithWireTranslatorModule.h
--- a/Engine/Plugins/Enterprise/DatasmithCADImporter/Source/DatasmithWireTranslator/Public/DatasmithWireTranslatorModule.h
+++ b/Engine/Plugins/Enterprise/DatasmithCADImporter/Source/DatasmithWireTranslator/Public/DatasmithWireTranslatorModule.h
@@ -6,17 +6,8 @@
 #include "Modules/ModuleInterface.h"
 #include "Modules/ModuleManager.h"
 
-#if defined(OPEN_MODEL_2020)
-#define DATASMITHWIRETRANSLATOR_MODULE_NAME TEXT("DatasmithWireTranslator2020")
-#elif defined(OPEN_MODEL_2021_3)
-#define DATASMITHWIRETRANSLATOR_MODULE_NAME TEXT("DatasmithWireTranslator2021_3")
-#elif defined(OPEN_MODEL_2022)
-#define DATASMITHWIRETRANSLATOR_MODULE_NAME TEXT("DatasmithWireTranslator2022")
-#elif defined(OPEN_MODEL_2022_1)
-#define DATASMITHWIRETRANSLATOR_MODULE_NAME TEXT("DatasmithWireTranslator2022_1")
-#elif defined(OPEN_MODEL_2022_2)
-#define DATASMITHWIRETRANSLATOR_MODULE_NAME TEXT("DatasmithWireTranslator2022_2")
-#endif
+namespace UE_DATASMITHWIRETRANSLATOR_NAMESPACE
+{
 
 /**
  * Datasmith Translator for .wire files.
@@ -27,12 +18,12 @@ public:
 
     static FDatasmithWireTranslatorModule& Get()
     {
-        return FModuleManager::LoadModuleChecked< FDatasmithWireTranslatorModule >(DATASMITHWIRETRANSLATOR_MODULE_NAME);
+        return FModuleManager::LoadModuleChecked< FDatasmithWireTranslatorModule >(PREPROCESSOR_TO_STRING(UE_DATASMITHWIRETRANSLATOR_MODULE_NAME));
 	}
 
     static bool IsAvailable()
     {
-        return FModuleManager::Get().IsModuleLoaded(DATASMITHWIRETRANSLATOR_MODULE_NAME);
+        return FModuleManager::Get().IsModuleLoaded(PREPROCESSOR_TO_STRING(UE_DATASMITHWIRETRANSLATOR_MODULE_NAME));
     }
 
 	virtual void StartupModule() override;
@@ -44,3 +35,5 @@ public:
 private:
 	FString TempDir;
 };
+
+}
\ No newline at end of file
