Fix NoPCH Compile UnrealEditor Win64 (Log)
remove all using namspace

diff --git a/Engine/Plugins/Enterprise/DatasmithCADImporter/Source/DatasmithWireTranslator/Private/AliasModelToCADKernelConverter.cpp b/Engine/Plugins/Enterprise/DatasmithCADImporter/Source/DatasmithWireTranslator/Private/AliasModelToCADKernelConverter.cpp
--- a/Engine/Plugins/Enterprise/DatasmithCADImporter/Source/DatasmithWireTranslator/Private/AliasModelToCADKernelConverter.cpp
+++ b/Engine/Plugins/Enterprise/DatasmithCADImporter/Source/DatasmithWireTranslator/Private/AliasModelToCADKernelConverter.cpp
@@ -17,7 +17,6 @@
 #include "CADKernel/Mesh/Structure/ModelMesh.h"
 
 #include "CADKernel/Topo/Body.h"
-//#include "CADKernel/Topo/Joiner.h"
 #include "CADKernel/Topo/Model.h"
 #include "CADKernel/Topo/Shell.h"
 #include "CADKernel/Topo/TopologicalEdge.h"
@@ -40,8 +39,6 @@
 #include "AlTrimRegion.h"
 #include "AlTM.h"
 
-using namespace CADKernel;
-
 namespace UE_DATASMITHWIRETRANSLATOR_NAMESPACE
 {
 
@@ -49,9 +46,9 @@ namespace AliasToCADKernelUtils
 {
 
 	template<typename Surface_T>
-	TSharedPtr<FSurface> AddNURBSSurface(double GeometricTolerance, Surface_T& AliasSurface, EAliasObjectReference InObjectReference, const AlMatrix4x4& InAlMatrix)
+	TSharedPtr<CADKernel::FSurface> AddNURBSSurface(double GeometricTolerance, Surface_T& AliasSurface, EAliasObjectReference InObjectReference, const AlMatrix4x4& InAlMatrix)
 	{
-		FNurbsSurfaceHomogeneousData NURBSData;
+		CADKernel::FNurbsSurfaceHomogeneousData NURBSData;
 		NURBSData.bSwapUV= true;
 		NURBSData.bIsRational = true;
 
@@ -92,13 +89,13 @@ namespace AliasToCADKernelUtils
 			AliasSurface.CVsUnaffectedPositionInclMultiples(NURBSData.HomogeneousPoles.GetData());
 		}
 
-		return FEntity::MakeShared<FNURBSSurface>(GeometricTolerance, NURBSData);
+		return CADKernel::FEntity::MakeShared<CADKernel::FNURBSSurface>(GeometricTolerance, NURBSData);
 	}
 }
 
-TSharedPtr<FTopologicalEdge> FAliasModelToCADKernelConverter::AddEdge(const AlTrimCurve& AliasTrimCurve, TSharedPtr<CADKernel::FSurface>& CarrierSurface)
+TSharedPtr<CADKernel::FTopologicalEdge> FAliasModelToCADKernelConverter::AddEdge(const AlTrimCurve& AliasTrimCurve, TSharedPtr<CADKernel::FSurface>& CarrierSurface)
 {
-	FNurbsCurveData NurbsCurveData;
+	CADKernel::FNurbsCurveData NurbsCurveData;
 
 	NurbsCurveData.Degree = AliasTrimCurve.degree();
 	int ControlPointCount = AliasTrimCurve.numberOfCVs();
@@ -127,13 +124,13 @@ TSharedPtr<FTopologicalEdge> FAliasModelToCADKernelConverter::AddEdge(const AlTr
 		NurbsCurveData.Poles[Index].Z = 0;
 	}
 
-	TSharedRef<FNURBSCurve> Nurbs = FEntity::MakeShared<FNURBSCurve>(NurbsCurveData);
+	TSharedRef<CADKernel::FNURBSCurve> Nurbs = CADKernel::FEntity::MakeShared<CADKernel::FNURBSCurve>(NurbsCurveData);
 
-	TSharedRef<FRestrictionCurve> RestrictionCurve = FEntity::MakeShared<FRestrictionCurve>(CarrierSurface.ToSharedRef(), Nurbs);
-	TSharedPtr<FTopologicalEdge> Edge = FTopologicalEdge::Make(RestrictionCurve);
+	TSharedRef<CADKernel::FRestrictionCurve> RestrictionCurve = CADKernel::FEntity::MakeShared<CADKernel::FRestrictionCurve>(CarrierSurface.ToSharedRef(), Nurbs);
+	TSharedPtr<CADKernel::FTopologicalEdge> Edge = CADKernel::FTopologicalEdge::Make(RestrictionCurve);
 	if (!Edge.IsValid())
 	{
-		return TSharedPtr<FTopologicalEdge>();
+		return TSharedPtr<CADKernel::FTopologicalEdge>();
 	}
 
 	// Only TrimCurve with twin need to be in the map used in LinkEdgesLoop 
@@ -146,14 +143,14 @@ TSharedPtr<FTopologicalEdge> FAliasModelToCADKernelConverter::AddEdge(const AlTr
 	return Edge;
 }
 
-TSharedPtr<FTopologicalLoop> FAliasModelToCADKernelConverter::AddLoop(const AlTrimBoundary& TrimBoundary, TSharedPtr<CADKernel::FSurface>& CarrierSurface)
+TSharedPtr<CADKernel::FTopologicalLoop> FAliasModelToCADKernelConverter::AddLoop(const AlTrimBoundary& TrimBoundary, TSharedPtr<CADKernel::FSurface>& CarrierSurface)
 {
-	TArray<TSharedPtr<FTopologicalEdge>> Edges;
+	TArray<TSharedPtr<CADKernel::FTopologicalEdge>> Edges;
 	TArray<CADKernel::EOrientation> Directions;
 
 	for (TUniquePtr<AlTrimCurve> TrimCurve(TrimBoundary.firstCurve()); TrimCurve.IsValid(); TrimCurve = TUniquePtr<AlTrimCurve>(TrimCurve->nextCurve()))
 	{
-		TSharedPtr<FTopologicalEdge> Edge = AddEdge(*TrimCurve, CarrierSurface);
+		TSharedPtr<CADKernel::FTopologicalEdge> Edge = AddEdge(*TrimCurve, CarrierSurface);
 		if (Edge.IsValid())
 		{
 			Edges.Add(Edge);
@@ -163,10 +160,10 @@ TSharedPtr<FTopologicalLoop> FAliasModelToCADKernelConverter::AddLoop(const AlTr
 
 	if (Edges.Num() == 0)
 	{
-		return TSharedPtr<FTopologicalLoop>();
+		return TSharedPtr<CADKernel::FTopologicalLoop>();
 	}
 
-	TSharedPtr<FTopologicalLoop> Loop = FTopologicalLoop::Make(Edges, Directions, GeometricTolerance);
+	TSharedPtr<CADKernel::FTopologicalLoop> Loop = CADKernel::FTopologicalLoop::Make(Edges, Directions, GeometricTolerance);
 	return Loop;
 }
 
@@ -174,7 +171,7 @@ void FAliasModelToCADKernelConverter::LinkEdgesLoop(const AlTrimBoundary& TrimBo
 {
 	for (TUniquePtr<AlTrimCurve> TrimCurve(TrimBoundary.firstCurve()); TrimCurve.IsValid(); TrimCurve = TUniquePtr<AlTrimCurve>(TrimCurve->nextCurve()))
 	{
-		TSharedPtr<FTopologicalEdge>* Edge = AlEdge2CADKernelEdge.Find(TrimCurve->fSpline);
+		TSharedPtr<CADKernel::FTopologicalEdge>* Edge = AlEdge2CADKernelEdge.Find(TrimCurve->fSpline);
 		if (!Edge || !Edge->IsValid() || (*Edge)->IsDeleted() || (*Edge)->IsDegenerated())
 		{
 			continue;
@@ -186,7 +183,7 @@ void FAliasModelToCADKernelConverter::LinkEdgesLoop(const AlTrimBoundary& TrimBo
 		TUniquePtr<AlTrimCurve> TwinCurve(TrimCurve->getTwinCurve());
 		if (TwinCurve.IsValid())
 		{
-			if (TSharedPtr<FTopologicalEdge>* TwinEdge = AlEdge2CADKernelEdge.Find(TwinCurve->fSpline))
+			if (TSharedPtr<CADKernel::FTopologicalEdge>* TwinEdge = AlEdge2CADKernelEdge.Find(TwinCurve->fSpline))
 			{
 				if (TwinEdge->IsValid() && !(*TwinEdge)->IsDeleted() && !(*TwinEdge)->IsDegenerated())
 				{
@@ -199,16 +196,16 @@ void FAliasModelToCADKernelConverter::LinkEdgesLoop(const AlTrimBoundary& TrimBo
 
 TSharedPtr<CADKernel::FTopologicalFace> FAliasModelToCADKernelConverter::AddTrimRegion(const AlTrimRegion& TrimRegion, EAliasObjectReference InObjectReference, const AlMatrix4x4& InAlMatrix, bool bInOrientation)
 {
-	TSharedPtr<FSurface> Surface = AliasToCADKernelUtils::AddNURBSSurface(GeometricTolerance, TrimRegion, InObjectReference, InAlMatrix);
+	TSharedPtr<CADKernel::FSurface> Surface = AliasToCADKernelUtils::AddNURBSSurface(GeometricTolerance, TrimRegion, InObjectReference, InAlMatrix);
 	if (!Surface.IsValid())
 	{
-		return TSharedPtr<FTopologicalFace>();
+		return TSharedPtr<CADKernel::FTopologicalFace>();
 	}
 
-	TArray<TSharedPtr<FTopologicalLoop>> Loops;
+	TArray<TSharedPtr<CADKernel::FTopologicalLoop>> Loops;
 	for (TUniquePtr<AlTrimBoundary> TrimBoundary(TrimRegion.firstBoundary()); TrimBoundary.IsValid(); TrimBoundary = TUniquePtr<AlTrimBoundary>(TrimBoundary->nextBoundary()))
 	{
-		TSharedPtr<FTopologicalLoop> Loop = AddLoop(*TrimBoundary, Surface);
+		TSharedPtr<CADKernel::FTopologicalLoop> Loop = AddLoop(*TrimBoundary, Surface);
 		if (Loop.IsValid())
 		{
 			LinkEdgesLoop(*TrimBoundary, *Loop);
@@ -218,11 +215,11 @@ TSharedPtr<CADKernel::FTopologicalFace> FAliasModelToCADKernelConverter::AddTrim
 
 	if (Loops.Num() == 0)
 	{
-		FMessage::Printf(Log, TEXT("The Face %s is degenerate, this face is ignored\n"), TrimRegion.name());
-		return TSharedPtr<FTopologicalFace>();
+		CADKernel::FMessage::Printf(CADKernel::EVerboseLevel::Log, TEXT("The Face %s is degenerate, this face is ignored\n"), TrimRegion.name());
+		return TSharedPtr<CADKernel::FTopologicalFace>();
 	}
 
-	TSharedRef<FTopologicalFace> Face = FEntity::MakeShared<FTopologicalFace>(Surface);
+	TSharedRef<CADKernel::FTopologicalFace> Face = CADKernel::FEntity::MakeShared<CADKernel::FTopologicalFace>(Surface);
 	Face->SetPatchId(LastFaceId++);
 
 	Face->AddLoops(Loops);
@@ -236,7 +233,7 @@ void FAliasModelToCADKernelConverter::AddFace(const AlSurface& Surface, EAliasOb
 	{
 		for (; TrimRegion.IsValid(); TrimRegion = TUniquePtr<AlTrimRegion>(TrimRegion->nextRegion()))
 		{
-			TSharedPtr<FTopologicalFace> Face = AddTrimRegion(*TrimRegion, InObjectReference, InAlMatrix, bInOrientation);
+			TSharedPtr<CADKernel::FTopologicalFace> Face = AddTrimRegion(*TrimRegion, InObjectReference, InAlMatrix, bInOrientation);
 			if (Face.IsValid())
 			{
 				Shell->Add(Face.ToSharedRef(), bInOrientation ? CADKernel::EOrientation::Front : CADKernel::EOrientation::Back);
@@ -245,10 +242,10 @@ void FAliasModelToCADKernelConverter::AddFace(const AlSurface& Surface, EAliasOb
 		return;
 	}
 
-	TSharedPtr<FSurface> CADKernelSurface = AliasToCADKernelUtils::AddNURBSSurface(GeometricTolerance, Surface, InObjectReference, InAlMatrix);
+	TSharedPtr<CADKernel::FSurface> CADKernelSurface = AliasToCADKernelUtils::AddNURBSSurface(GeometricTolerance, Surface, InObjectReference, InAlMatrix);
 	if (CADKernelSurface.IsValid())
 	{
-		TSharedRef<FTopologicalFace> Face = FEntity::MakeShared<FTopologicalFace>(CADKernelSurface);
+		TSharedRef<CADKernel::FTopologicalFace> Face = CADKernel::FEntity::MakeShared<CADKernel::FTopologicalFace>(CADKernelSurface);
 		Face->ApplyNaturalLoops();
 		Shell->Add(Face, bInOrientation ? CADKernel::EOrientation::Front : CADKernel::EOrientation::Back);
 	}
@@ -258,7 +255,7 @@ void FAliasModelToCADKernelConverter::AddShell(const AlShell& InShell, EAliasObj
 {
 	for(TUniquePtr<AlTrimRegion> TrimRegion(InShell.firstTrimRegion()); TrimRegion.IsValid(); TrimRegion = TUniquePtr<AlTrimRegion>(TrimRegion->nextRegion()))
 	{
-		TSharedPtr<FTopologicalFace> Face = AddTrimRegion(*TrimRegion, InObjectReference, InAlMatrix, bInOrientation);
+		TSharedPtr<CADKernel::FTopologicalFace> Face = AddTrimRegion(*TrimRegion, InObjectReference, InAlMatrix, bInOrientation);
 		if (Face.IsValid())
 		{
 			CADKernelShell->Add(Face.ToSharedRef(), bInOrientation ? CADKernel::EOrientation::Front : CADKernel::EOrientation::Back);
@@ -270,8 +267,8 @@ bool FAliasModelToCADKernelConverter::AddBRep(AlDagNode& DagNode, EAliasObjectRe
 {
 	AlEdge2CADKernelEdge.Empty();
 
-	TSharedRef<FBody> CADKernelBody = FEntity::MakeShared<FBody>();
-	TSharedRef<FShell> CADKernelShell = FEntity::MakeShared<FShell>();
+	TSharedRef<CADKernel::FBody> CADKernelBody = CADKernel::FEntity::MakeShared<CADKernel::FBody>();
+	TSharedRef<CADKernel::FShell> CADKernelShell = CADKernel::FEntity::MakeShared<CADKernel::FShell>();
 	CADKernelBody->AddShell(CADKernelShell);
 
 	boolean bAlOrientation;
diff --git a/Engine/Plugins/Enterprise/DatasmithCADImporter/Source/DatasmithWireTranslator/Private/AliasModelToCoretechConverter.cpp b/Engine/Plugins/Enterprise/DatasmithCADImporter/Source/DatasmithWireTranslator/Private/AliasModelToCoretechConverter.cpp
--- a/Engine/Plugins/Enterprise/DatasmithCADImporter/Source/DatasmithWireTranslator/Private/AliasModelToCoretechConverter.cpp
+++ b/Engine/Plugins/Enterprise/DatasmithCADImporter/Source/DatasmithWireTranslator/Private/AliasModelToCoretechConverter.cpp
@@ -25,8 +25,6 @@
 #include "AlTrimRegion.h"
 #include "AlTM.h"
 
-using namespace CADLibrary;
-
 namespace UE_DATASMITHWIRETRANSLATOR_NAMESPACE
 {
 
@@ -35,7 +33,7 @@ namespace AliasToCoreTechUtils
 	template<typename Surface_T>
 	uint64 CreateCTNurbs(Surface_T& Surface, EAliasObjectReference InObjectReference, const AlMatrix4x4& InAlMatrix)
 	{
-		FNurbsSurface CTSurface;
+		CADLibrary::FNurbsSurface CTSurface;
 
 		CTSurface.ControlPointDimension = 4;  // Control Hull Dimension (3 for non-rational or 4 for rational)
 
@@ -86,7 +84,7 @@ namespace AliasToCoreTechUtils
 
 uint64 FAliasModelToCoretechConverter::AddTrimCurve(const AlTrimCurve& TrimCurve)
 {
-	FNurbsCurve CTCurve;
+	CADLibrary::FNurbsCurve CTCurve;
 
 	CTCurve.Order = TrimCurve.degree() + 1;
 	CTCurve.ControlPointSize = TrimCurve.numberOfCVs();
@@ -107,14 +105,14 @@ uint64 FAliasModelToCoretechConverter::AddTrimCurve(const AlTrimCurve& TrimCurve
 	CTCurve.KnotValues[CTCurve.KnotSize - 1] = CTCurve.KnotValues[CTCurve.KnotSize - 2];
 
 	uint64 CoedgeID = 0;
-	if (CTKIO_CreateCoedge(CTCurve, (bool) TrimCurve.isReversed(), CoedgeID))
+	if (CADLibrary::CTKIO_CreateCoedge(CTCurve, (bool) TrimCurve.isReversed(), CoedgeID))
 	{
 		// Build topo
 		if (AlTrimCurve *TwinCurve = TrimCurve.getTwinCurve())
 		{
 			if (uint64* TwinCoedgeID = AlEdge2CTEdge.Find(TwinCurve->fSpline))
 			{
-				CTKIO_MatchCoedges(*TwinCoedgeID, CoedgeID);
+				CADLibrary::CTKIO_MatchCoedges(*TwinCoedgeID, CoedgeID);
 			}
 			// only TrimCurve with twin need to be in the map
 			AlEdge2CTEdge.Add(TrimCurve.fSpline, CoedgeID);
@@ -137,7 +135,7 @@ uint64 FAliasModelToCoretechConverter::AddTrimBoundary(const AlTrimBoundary& Tri
 	}
 
 	uint64 LoopID;
-	return CTKIO_CreateLoop(Edges, LoopID) ? LoopID : 0;
+	return CADLibrary::CTKIO_CreateLoop(Edges, LoopID) ? LoopID : 0;
 }
 
 uint64 FAliasModelToCoretechConverter::AddTrimRegion(const AlTrimRegion& TrimRegion, EAliasObjectReference InObjectReference, const AlMatrix4x4& InAlMatrix, bool bInOrientation)
@@ -158,7 +156,7 @@ uint64 FAliasModelToCoretechConverter::AddTrimRegion(const AlTrimRegion& TrimReg
 	}
 
 	uint64 FaceID;
-	return CTKIO_CreateFace(SurfaceID, bInOrientation, Boundaries, FaceID) ? FaceID : 0;
+	return CADLibrary::CTKIO_CreateFace(SurfaceID, bInOrientation, Boundaries, FaceID) ? FaceID : 0;
 }
 
 void FAliasModelToCoretechConverter::AddFace(const AlSurface& Surface, EAliasObjectReference InObjectReference, const AlMatrix4x4& InAlMatrix, bool bInOrientation, TArray<uint64>& OutFaceList)
@@ -179,7 +177,7 @@ void FAliasModelToCoretechConverter::AddFace(const AlSurface& Surface, EAliasObj
 	uint64 SurfaceID = AliasToCoreTechUtils::CreateCTNurbs(Surface, InObjectReference, InAlMatrix);
 
 	uint64 FaceID;
-	if (CTKIO_CreateFace(SurfaceID, bInOrientation, TArray<uint64>(), FaceID) && FaceID != 0)
+	if (CADLibrary::CTKIO_CreateFace(SurfaceID, bInOrientation, TArray<uint64>(), FaceID) && FaceID != 0)
 	{
 		OutFaceList.Add(FaceID);
 	}
diff --git a/Engine/Plugins/Enterprise/DatasmithCADImporter/Source/DatasmithWireTranslator/Private/DatasmithWireTranslator.cpp b/Engine/Plugins/Enterprise/DatasmithCADImporter/Source/DatasmithWireTranslator/Private/DatasmithWireTranslator.cpp
--- a/Engine/Plugins/Enterprise/DatasmithCADImporter/Source/DatasmithWireTranslator/Private/DatasmithWireTranslator.cpp
+++ b/Engine/Plugins/Enterprise/DatasmithCADImporter/Source/DatasmithWireTranslator/Private/DatasmithWireTranslator.cpp
@@ -101,9 +101,6 @@ const uint64 LibAliasVersionMax = LibAliasNext_Version;
 const FString AliasVersionChar = TEXT("AliasStudio 2022.2, Model files");
 #endif
 
-
-using namespace CADLibrary;
-
 class BodyData
 {
 public:
@@ -225,8 +222,8 @@ public:
 			DatasmithScene->SetProductVersion(TEXT("Alias 2022.2"));
 		}
 
-		FImportParameters ImportParameters(0.01, 1);
-		if(FImportParameters::bGDisableCADKernelTessellation)
+		CADLibrary::FImportParameters ImportParameters(0.01, 1);
+		if(CADLibrary::FImportParameters::bGDisableCADKernelTessellation)
 		{
 			TSharedRef<FAliasModelToCoretechConverter> AliasToCoretechConverter = MakeShared<FAliasModelToCoretechConverter>(TEXT("Al2CTSharedSession"), ImportParameters);
 			CADModelConverter = AliasToCoretechConverter;
@@ -374,7 +371,7 @@ private:
 
 	FDatasmithTessellationOptions TessellationOptions;
 
-	TSharedPtr<ICADModelConverter> CADModelConverter;
+	TSharedPtr<CADLibrary::ICADModelConverter> CADModelConverter;
 	TSharedPtr<IAliasBRepConverter> AliasBRepConverter;
 
 };
diff --git a/Engine/Plugins/Enterprise/DatasmithCADImporter/Source/DatasmithWireTranslator/Private/OpenModelUtils.cpp b/Engine/Plugins/Enterprise/DatasmithCADImporter/Source/DatasmithWireTranslator/Private/OpenModelUtils.cpp
--- a/Engine/Plugins/Enterprise/DatasmithCADImporter/Source/DatasmithWireTranslator/Private/OpenModelUtils.cpp
+++ b/Engine/Plugins/Enterprise/DatasmithCADImporter/Source/DatasmithWireTranslator/Private/OpenModelUtils.cpp
@@ -20,8 +20,6 @@
 namespace UE_DATASMITHWIRETRANSLATOR_NAMESPACE
 {
 
-using namespace OpenModelUtils;
-
 const TCHAR* OpenModelUtils::AlObjectTypeToString(AlObjectType type)
 {
 	switch (type)
