Formating

diff --git a/Engine/Plugins/Enterprise/DatasmithCADImporter/Source/CADInterfaces/Private/TUniqueTechSoftObj.h b/Engine/Plugins/Enterprise/DatasmithCADImporter/Source/CADInterfaces/Private/TUniqueTechSoftObj.h
--- a/Engine/Plugins/Enterprise/DatasmithCADImporter/Source/CADInterfaces/Private/TUniqueTechSoftObj.h
+++ b/Engine/Plugins/Enterprise/DatasmithCADImporter/Source/CADInterfaces/Private/TUniqueTechSoftObj.h
@@ -10,189 +10,189 @@
 namespace CADLibrary
 {
 
-	// Single-ownership smart TeshSoft object
-	// Use this when you need to manage TechSoft object's lifetime.
-	//
-	// TechSoft give access to void pointers
-	// According to the context, the class name of the void pointer is known but the class is unknown
-	// i.e. A3DSDKTypes.h defines all type like :
-	// 	   typedef void A3DEntity;		
-	// 	   typedef void A3DAsmModelFile; ...
-	// 
-	// From a pointer, TechSoft give access to a copy of the associated structure :
-	//
-	// const A3DXXXXX* pPointer;
-	// A3DXXXXXData sData; // the structure
-	// A3D_INITIALIZE_DATA(A3DXXXXXData, sData); // initialization of the structure
-	// A3DXXXXXXGet(pPointer, &sData); // Copy of the data of the pointer in the structure
-	// ...
-	// A3DXXXXXXGet(NULL, &sData); // Free the structure
-	//
-	// A3D_INITIALIZE_DATA, and all A3DXXXXXXGet methods are TechSoft macro
-	//
-
-	template<class ObjectType, class IndexerType>
-	class TUniqueTSObjBase
-	{
-	public:
+// Single-ownership smart TeshSoft object
+// Use this when you need to manage TechSoft object's lifetime.
+//
+// TechSoft give access to void pointers
+// According to the context, the class name of the void pointer is known but the class is unknown
+// i.e. A3DSDKTypes.h defines all type like :
+// 	   typedef void A3DEntity;		
+// 	   typedef void A3DAsmModelFile; ...
+// 
+// From a pointer, TechSoft give access to a copy of the associated structure :
+//
+// const A3DXXXXX* pPointer;
+// A3DXXXXXData sData; // the structure
+// A3D_INITIALIZE_DATA(A3DXXXXXData, sData); // initialization of the structure
+// A3DXXXXXXGet(pPointer, &sData); // Copy of the data of the pointer in the structure
+// ...
+// A3DXXXXXXGet(NULL, &sData); // Free the structure
+//
+// A3D_INITIALIZE_DATA, and all A3DXXXXXXGet methods are TechSoft macro
+//
+
+template<class ObjectType, class IndexerType>
+class TUniqueTSObjBase
+{
+public:
 
-		/**
-		 * Constructor of an initialized ObjectType object
-		 */
-		explicit TUniqueTSObjBase()
-		{
-			InitializeData();
-		}
+	/**
+	 * Constructor of an initialized ObjectType object
+	 */
+	explicit TUniqueTSObjBase()
+	{
+		InitializeData();
+	}
+
+	/**
+	 * Constructor of an filled ObjectType object with the data of DataPtr
+	 * @param DataPtr: the pointer of the data to copy
+	 */
+	explicit TUniqueTSObjBase(IndexerType DataPtr)
+	{
+		//TechSoftInterfaceImpl::InitializeData(Data);
+		InitializeData();
 
-		/**
-		 * Constructor of an filled ObjectType object with the data of DataPtr
-		 * @param DataPtr: the pointer of the data to copy
-		 */
-		explicit TUniqueTSObjBase(IndexerType DataPtr)
-		{
-			//TechSoftInterfaceImpl::InitializeData(Data);
-			InitializeData();
+		Status = GetData(DataPtr);
+	}
 
-			Status = GetData(DataPtr);
-		}
+	~TUniqueTSObjBase()
+	{
+		ResetData();
+	}
 
-		~TUniqueTSObjBase()
+	/**
+	 * Fill the structure with the data of a new DataPtr
+	 */
+	A3DStatus FillFrom(IndexerType DataPtr)
+	{
+		if (IsValid())
 		{
-			ResetData();
+			Status = ResetData();
 		}
-
-		/**
-		 * Fill the structure with the data of a new DataPtr
-		 */
-		A3DStatus FillFrom(IndexerType DataPtr)
+		else
 		{
-			if (IsValid())
-			{
-				Status = ResetData();
-			}
-			else
-			{
-				Status = A3DStatus::A3D_SUCCESS;
-			}
-
-			if (!IsValid() || (DataPtr == DefaultValue))
-			{
-				Status = A3DStatus::A3D_ERROR;
-				return Status;
-			}
-
-			Status = GetData(DataPtr);
-			return Status;
-		}
-
-		/**
-		 * Empty the structure
-		 */
-		void Reset()
-		{
-			ResetData();
+			Status = A3DStatus::A3D_SUCCESS;
 		}
 
-		/**
-		 * Return
-		 *  - A3DStatus::A3D_SUCCESS if the data is filled
-		 *  - A3DStatus::A3D_ERROR if the data is empty
-		 */
-		A3DStatus GetStatus()
+		if (!IsValid() || (DataPtr == DefaultValue))
 		{
+			Status = A3DStatus::A3D_ERROR;
 			return Status;
 		}
 
-		/**
-		 * Return true if the data is filled
-		 */
-		const bool IsValid() const
-		{
-			return Status == A3DStatus::A3D_SUCCESS;
-		}
-
-		// Non-copyable
-		TUniqueTSObjBase(const TUniqueTSObjBase&) = delete;
-		TUniqueTSObjBase& operator=(const TUniqueTSObjBase&) = delete;
+		Status = GetData(DataPtr);
+		return Status;
+	}
 
-		// Conversion methods
-
-		const ObjectType& operator*() const
-		{
-			return Data;
-		}
+	/**
+	 * Empty the structure
+	 */
+	void Reset()
+	{
+		ResetData();
+	}
+
+	/**
+	 * Return
+	 *  - A3DStatus::A3D_SUCCESS if the data is filled
+	 *  - A3DStatus::A3D_ERROR if the data is empty
+	 */
+	A3DStatus GetStatus()
+	{
+		return Status;
+	}
 
-		ObjectType& operator*()
-		{
-			return Data;
-		}
+	/**
+	 * Return true if the data is filled
+	 */
+	const bool IsValid() const
+	{
+		return Status == A3DStatus::A3D_SUCCESS;
+	}
 
-		const ObjectType* operator->() const
-		{
-			check(IsValid());
-			return &Data;
-		}
+	// Non-copyable
+	TUniqueTSObjBase(const TUniqueTSObjBase&) = delete;
+	TUniqueTSObjBase& operator=(const TUniqueTSObjBase&) = delete;
 
-		ObjectType* operator->()
-		{
-			check(IsValid());
-			return &Data;
-		}
+	// Conversion methods
 
-		/**
-		 * Return the structure and set the status to filled
-		 * The method is used to manage structure filled by UE
-		 */
-		ObjectType* GetEmptyDataPtr()
-		{
-			FillFrom(DefaultValue);
-			Status = A3DStatus::A3D_SUCCESS;
-			return &Data;
-		}
+	const ObjectType& operator*() const
+	{
+		return Data;
+	}
 
-	private:
-		ObjectType Data;
-		A3DStatus Status = A3DStatus::A3D_ERROR;
+	ObjectType& operator*()
+	{
+		return Data;
+	}
 
-		/**
-		 * DefaultValue is used to initialize "Data" with GetData method
-		 * According to IndexerType, the value is either nullptr for const A3DEntity* either something like "A3D_DEFAULT_MATERIAL_INDEX" ((A3DUns16)-1) for uint32
-		 * @see ResetData
-		 */
-		static IndexerType DefaultValue;
+	const ObjectType* operator->() const
+	{
+		check(IsValid());
+		return &Data;
+	}
 
-		void InitializeData()
+	ObjectType* operator->()
+	{
+		check(IsValid());
+		return &Data;
+	}
+
+	/**
+	 * Return the structure and set the status to filled
+	 * The method is used to manage structure filled by UE
+	 */
+	ObjectType* GetEmptyDataPtr()
+	{
+		FillFrom(DefaultValue);
+		Status = A3DStatus::A3D_SUCCESS;
+		return &Data;
+	}
+
+private:
+	ObjectType Data;
+	A3DStatus Status = A3DStatus::A3D_ERROR;
+
+	/**
+	 * DefaultValue is used to initialize "Data" with GetData method
+	 * According to IndexerType, the value is either nullptr for const A3DEntity* either something like "A3D_DEFAULT_MATERIAL_INDEX" ((A3DUns16)-1) for uint32
+	 * @see ResetData
+	 */
+	static IndexerType DefaultValue;
+
+	void InitializeData()
 #ifdef USE_TECHSOFT_SDK
-			;
+		;
 #else
-		{
-			return A3DStatus::A3D_ERROR;
-		}
+	{
+		return A3DStatus::A3D_ERROR;
+	}
 #endif
 
-		A3DStatus GetData(IndexerType AsmModelFilePtr);
+	A3DStatus GetData(IndexerType AsmModelFilePtr);
 #ifdef USE_TECHSOFT_SDK
-		;
+	;
 #else
-		{
-			return A3DStatus::A3D_ERROR;
-		}
+	{
+		return A3DStatus::A3D_ERROR;
+	}
 #endif
 
-		A3DStatus ResetData()
-		{
-			return GetData(DefaultValue);
-		}
-	};
+	A3DStatus ResetData()
+	{
+		return GetData(DefaultValue);
+	}
+};
 
-	template<class ObjectType>
-	using TUniqueTSObj = TUniqueTSObjBase<ObjectType, const A3DEntity*>;
+template<class ObjectType>
+using TUniqueTSObj = TUniqueTSObjBase<ObjectType, const A3DEntity*>;
 
-	template<class ObjectType>
-	using TUniqueTSObjFromIndex = TUniqueTSObjBase<ObjectType, uint32>;
+template<class ObjectType>
+using TUniqueTSObjFromIndex = TUniqueTSObjBase<ObjectType, uint32>;
 
-	template<class ObjectType, class IndexerType>
-	IndexerType TUniqueTSObjBase<ObjectType, IndexerType>::DefaultValue = (IndexerType) nullptr;
+template<class ObjectType, class IndexerType>
+IndexerType TUniqueTSObjBase<ObjectType, IndexerType>::DefaultValue = (IndexerType) nullptr;
 
 }
 
diff --git a/Engine/Plugins/Enterprise/DatasmithCADImporter/Source/CADInterfaces/Private/TechSoftFileParser.cpp b/Engine/Plugins/Enterprise/DatasmithCADImporter/Source/CADInterfaces/Private/TechSoftFileParser.cpp
--- a/Engine/Plugins/Enterprise/DatasmithCADImporter/Source/CADInterfaces/Private/TechSoftFileParser.cpp
+++ b/Engine/Plugins/Enterprise/DatasmithCADImporter/Source/CADInterfaces/Private/TechSoftFileParser.cpp
@@ -11,1890 +11,1891 @@ namespace CADLibrary
 {
 namespace TechSoftFileParserImpl
 {
-	// Functions to clean metadata
+// Functions to clean metadata
 
-	inline void RemoveUnwantedChar(FString& StringToClean, TCHAR UnwantedChar)
+inline void RemoveUnwantedChar(FString& StringToClean, TCHAR UnwantedChar)
+{
+	FString NewString;
+	NewString.Reserve(StringToClean.Len());
+	for (const TCHAR& Char : StringToClean)
 	{
-		FString NewString;
-		NewString.Reserve(StringToClean.Len());
-		for (const TCHAR& Char : StringToClean)
+		if (Char != UnwantedChar)
 		{
-			if (Char != UnwantedChar)
-			{
-				NewString.AppendChar(Char);
-			}
+			NewString.AppendChar(Char);
 		}
-		StringToClean = MoveTemp(NewString);
 	}
+	StringToClean = MoveTemp(NewString);
+}
 
-	FString CleanSdkName(const FString& Name)
+FString CleanSdkName(const FString& Name)
+{
+	int32 Index;
+	if (Name.FindLastChar(TEXT('['), Index))
 	{
-		int32 Index;
-		if (Name.FindLastChar(TEXT('['), Index))
-		{
-			return Name.Left(Index);
-		}
-		return Name;
+		return Name.Left(Index);
 	}
+	return Name;
+}
 
-	FString CleanCatiaInstanceSdkName(const FString& Name)
+FString CleanCatiaInstanceSdkName(const FString& Name)
+{
+	int32 Index;
+	if (Name.FindChar(TEXT('('), Index))
 	{
-		int32 Index;
-		if (Name.FindChar(TEXT('('), Index))
+		FString NewName = Name.Right(Index + 1);
+		if (Name.FindLastChar(TEXT(')'), Index))
 		{
-			FString NewName = Name.Right(Index + 1);
-			if (Name.FindLastChar(TEXT(')'), Index))
-			{
-				NewName = NewName.Left(Index);
-			}
-			return NewName;
+			NewName = NewName.Left(Index);
 		}
-		return Name;
+		return NewName;
 	}
+	return Name;
+}
 
-	FString Clean3dxmlReferenceSdkName(const FString& Name)
+FString Clean3dxmlReferenceSdkName(const FString& Name)
+{
+	int32 Index;
+	if (Name.FindChar(TEXT('('), Index))
 	{
-		int32 Index;
-		if (Name.FindChar(TEXT('('), Index))
-		{
-			FString NewName = Name.Left(Index);
-			return NewName;
-		}
-		return Name;
+		FString NewName = Name.Left(Index);
+		return NewName;
 	}
+	return Name;
+}
 
-	FString CleanSwInstanceSdkName(const FString& Name)
+FString CleanSwInstanceSdkName(const FString& Name)
+{
+	int32 Position;
+	if (Name.FindLastChar(TEXT('-'), Position))
 	{
-		int32 Position;
-		if (Name.FindLastChar(TEXT('-'), Position))
-		{
-			FString NewName = Name.Left(Position) + TEXT("<") + Name.Right(Position + 1) + TEXT(">");
-			return NewName;
-		}
-		return Name;
+		FString NewName = Name.Left(Position) + TEXT("<") + Name.Right(Position + 1) + TEXT(">");
+		return NewName;
+	}
+	return Name;
+}
+
+FString CleanSwReferenceSdkName(const FString& Name)
+{
+	int32 Position;
+	if (Name.FindLastChar(TEXT('-'), Position))
+	{
+		FString NewName = Name.Left(Position);
+		return NewName;
 	}
+	return Name;
+}
 
-	FString CleanSwReferenceSdkName(const FString& Name)
+FString CleanCatiaReferenceName(const FString& Name)
+{
+	int32 Position;
+	if (Name.FindLastChar(TEXT('.'), Position))
 	{
-		int32 Position;
-		if (Name.FindLastChar(TEXT('-'), Position))
+		FString Indice = Name.Right(Position + 1);
+		if (Indice.IsNumeric())
 		{
 			FString NewName = Name.Left(Position);
 			return NewName;
 		}
-		return Name;
 	}
+	return Name;
+}
 
-	FString CleanCatiaReferenceName(const FString& Name)
+FString CleanNameByRemoving_prt(const FString& Name)
+{
+	int32 Position;
+	if (Name.FindLastChar(TEXT('.'), Position))
 	{
-		int32 Position;
-		if (Name.FindLastChar(TEXT('.'), Position))
+		FString Extension = Name.Right(Position);
+		if (Extension.Equals(TEXT("prt"), ESearchCase::IgnoreCase))
 		{
-			FString Indice = Name.Right(Position + 1);
-			if (Indice.IsNumeric())
-			{
-				FString NewName = Name.Left(Position);
-				return NewName;
-			}
+			FString NewName = Name.Left(Position);
+			return NewName;
 		}
-		return Name;
 	}
+	return Name;
+}
 
-	FString CleanNameByRemoving_prt(const FString& Name)
+bool CheckIfNameExists(TMap<FString, FString>& MetaData)
+{
+	FString* NamePtr = MetaData.Find(TEXT("Name"));
+	if (NamePtr != nullptr)
 	{
-		int32 Position;
-		if (Name.FindLastChar(TEXT('.'), Position))
-		{
-			FString Extension = Name.Right(Position);
-			if (Extension.Equals(TEXT("prt"), ESearchCase::IgnoreCase))
-			{
-				FString NewName = Name.Left(Position);
-				return NewName;
-			}
-		}
-		return Name;
+		return true;
 	}
+	return false;
+}
 
-	bool CheckIfNameExists(TMap<FString, FString>& MetaData)
+bool ReplaceOrAddNameValue(TMap<FString, FString>& MetaData, TCHAR* Key)
+{
+	FString* NamePtr = MetaData.Find(Key);
+	if (NamePtr != nullptr)
 	{
-		FString* NamePtr = MetaData.Find(TEXT("Name"));
-		if (NamePtr != nullptr)
-		{
-			return true;
-		}
-		return false;
+		FString& Name = MetaData.FindOrAdd(TEXT("Name"));
+		Name = *NamePtr;
+		return true;
 	}
+	return false;
+}
+
+// Functions used in traverse model process
+bool IsUnloadedPrototype(const A3DAsmProductOccurrenceData& PrototypeData)
+{
+	bool bIsUnloaded = true;
 
-	bool ReplaceOrAddNameValue(TMap<FString, FString>& MetaData, TCHAR* Key)
+	if (PrototypeData.m_uiPOccurrencesSize > 0)
+	{
+		bIsUnloaded = false;
+	}
+	else if (PrototypeData.m_uiEntityReferenceSize > 0)
+	{
+		bIsUnloaded = false;
+	}
+	else if (PrototypeData.m_pPart)
+	{
+		bIsUnloaded = false;
+	}
+	else if (PrototypeData.m_pPrototype)
 	{
-		FString* NamePtr = MetaData.Find(Key);
-		if (NamePtr != nullptr)
+		TUniqueTSObj<A3DAsmProductOccurrenceData> SubPrototypeData(PrototypeData.m_pPrototype);
+		if (SubPrototypeData.IsValid())
 		{
-			FString& Name = MetaData.FindOrAdd(TEXT("Name"));
-			Name = *NamePtr;
-			return true;
+			bIsUnloaded = TechSoftFileParserImpl::IsUnloadedPrototype(*SubPrototypeData);
 		}
-		return false;
 	}
 
-	// Functions used in traverse model process
-	bool IsUnloadedPrototype(const A3DAsmProductOccurrenceData& PrototypeData)
+	return bIsUnloaded;
+}
+
+void TraverseAttribute(const A3DMiscAttributeData& AttributeData, TMap<FString, FString>& OutMetaData)
+{
+	FString AttributeName;
+	if (AttributeData.m_bTitleIsInt)
+	{
+		A3DUns32 UnsignedValue = 0;
+		memcpy(&UnsignedValue, AttributeData.m_pcTitle, sizeof(A3DUns32));
+		AttributeName = FString::Printf(TEXT("%u"), UnsignedValue);
+	}
+	else if (AttributeData.m_pcTitle && AttributeData.m_pcTitle[0] != '\0')
 	{
-		bool bIsUnloaded = true;
+		AttributeName = UTF8_TO_TCHAR(AttributeData.m_pcTitle);
+	}
 
-		if (PrototypeData.m_uiPOccurrencesSize > 0)
+	for (A3DUns32 Index = 0; Index < AttributeData.m_uiSize; ++Index)
+	{
+		FString AttributeValue;
+		switch (AttributeData.m_asSingleAttributesData[Index].m_eType)
 		{
-			bIsUnloaded = false;
-		}
-		else if (PrototypeData.m_uiEntityReferenceSize > 0)
+		case kA3DModellerAttributeTypeTime:
+		case kA3DModellerAttributeTypeInt:
 		{
-			bIsUnloaded = false;
+			A3DInt32 Value;
+			memcpy(&Value, AttributeData.m_asSingleAttributesData[Index].m_pcData, sizeof(A3DInt32));
+			AttributeValue = FString::Printf(TEXT("%d"), Value);
+			break;
 		}
-		else if (PrototypeData.m_pPart)
+
+		case kA3DModellerAttributeTypeReal:
 		{
-			bIsUnloaded = false;
+			A3DDouble Value;
+			memcpy(&Value, AttributeData.m_asSingleAttributesData[Index].m_pcData, sizeof(A3DDouble));
+			AttributeValue = FString::Printf(TEXT("%f"), Value);
+			break;
 		}
-		else if (PrototypeData.m_pPrototype)
+
+		case kA3DModellerAttributeTypeString:
 		{
-			TUniqueTSObj<A3DAsmProductOccurrenceData> SubPrototypeData(PrototypeData.m_pPrototype);
-			if (SubPrototypeData.IsValid())
+			if (AttributeData.m_asSingleAttributesData[Index].m_pcData && AttributeData.m_asSingleAttributesData[Index].m_pcData[0] != '\0')
 			{
-				bIsUnloaded = TechSoftFileParserImpl::IsUnloadedPrototype(*SubPrototypeData);
+				AttributeValue = UTF8_TO_TCHAR(AttributeData.m_asSingleAttributesData[Index].m_pcData);
 			}
+			break;
 		}
 
-		return bIsUnloaded;
-	}
-
-	void TraverseAttribute(const A3DMiscAttributeData& AttributeData, TMap<FString, FString>& OutMetaData)
-	{
-		FString AttributeName;
-		if (AttributeData.m_bTitleIsInt)
-		{
-			A3DUns32 UnsignedValue = 0;
-			memcpy(&UnsignedValue, AttributeData.m_pcTitle, sizeof(A3DUns32));
-			AttributeName = FString::Printf(TEXT("%u"), UnsignedValue);
-		}
-		else if (AttributeData.m_pcTitle && AttributeData.m_pcTitle[0] != '\0')
-		{
-			AttributeName = UTF8_TO_TCHAR(AttributeData.m_pcTitle);
+		default:
+			break;
 		}
 
-		for (A3DUns32 Index = 0; Index < AttributeData.m_uiSize; ++Index)
+		if (AttributeName.Len())
 		{
-			FString AttributeValue;
-			switch (AttributeData.m_asSingleAttributesData[Index].m_eType)
-			{
-			case kA3DModellerAttributeTypeTime:
-			case kA3DModellerAttributeTypeInt:
+			if (Index)
 			{
-				A3DInt32 Value;
-				memcpy(&Value, AttributeData.m_asSingleAttributesData[Index].m_pcData, sizeof(A3DInt32));
-				AttributeValue = FString::Printf(TEXT("%d"), Value);
-				break;
+				OutMetaData.Emplace(FString::Printf(TEXT("%s_%u"), *AttributeName, (int32)Index), AttributeValue);
 			}
-
-			case kA3DModellerAttributeTypeReal:
+			else
 			{
-				A3DDouble Value;
-				memcpy(&Value, AttributeData.m_asSingleAttributesData[Index].m_pcData, sizeof(A3DDouble));
-				AttributeValue = FString::Printf(TEXT("%f"), Value);
-				break;
+				OutMetaData.Emplace(AttributeName, AttributeValue);
 			}
+		}
+	}
+}
 
-			case kA3DModellerAttributeTypeString:
-			{
-				if (AttributeData.m_asSingleAttributesData[Index].m_pcData && AttributeData.m_asSingleAttributesData[Index].m_pcData[0] != '\0')
-				{
-					AttributeValue = UTF8_TO_TCHAR(AttributeData.m_asSingleAttributesData[Index].m_pcData);
-				}
-				break;
-			}
+void SetIOOption(A3DImport& Importer)
+{
+	// A3DRWParamsGeneralData Importer.m_sGeneral
+	Importer.m_sLoadData.m_sGeneral.m_bReadSolids = A3D_TRUE;
+	Importer.m_sLoadData.m_sGeneral.m_bReadSurfaces = A3D_TRUE;
+	Importer.m_sLoadData.m_sGeneral.m_bReadWireframes = A3D_FALSE;
+	Importer.m_sLoadData.m_sGeneral.m_bReadPmis = A3D_FALSE;
+	Importer.m_sLoadData.m_sGeneral.m_bReadAttributes = A3D_TRUE;
+	Importer.m_sLoadData.m_sGeneral.m_bReadHiddenObjects = A3D_FALSE;
+	Importer.m_sLoadData.m_sGeneral.m_bReadConstructionAndReferences = A3D_FALSE;
+	Importer.m_sLoadData.m_sGeneral.m_bReadActiveFilter = A3D_FALSE;
+	Importer.m_sLoadData.m_sGeneral.m_eReadingMode2D3D = kA3DRead_3D;
+
+	Importer.m_sLoadData.m_sGeneral.m_eReadGeomTessMode = kA3DReadGeomOnly;
+	Importer.m_sLoadData.m_sGeneral.m_eDefaultUnit;
+	Importer.m_sLoadData.m_sGeneral.m_bReadFeature = A3D_FALSE;
+
+	Importer.m_sLoadData.m_sGeneral.m_bReadConstraints = A3D_FALSE;
+	Importer.m_sLoadData.m_sGeneral.m_iNbMultiProcess = 1;
+
+	Importer.m_sLoadData.m_sIncremental.m_bLoadNoDependencies = true;
+	Importer.m_sLoadData.m_sIncremental.m_bLoadStructureOnly = false;
+}
 
-			default:
-				break;
-			}
+void UpdateIOOptionAccordingToFormat(const CADLibrary::ECADFormat Format, A3DImport& Importer)
+{
+	switch (Format)
+	{
+	case CADLibrary::ECADFormat::IGES:
+	{
+		break;
+	}
 
-			if (AttributeName.Len())
-			{
-				if (Index)
-				{
-					OutMetaData.Emplace(FString::Printf(TEXT("%s_%u"), *AttributeName, (int32)Index), AttributeValue);
-				}
-				else
-				{
-					OutMetaData.Emplace(AttributeName, AttributeValue);
-				}
-			}
-		}
+	case CADLibrary::ECADFormat::CATIA:
+	{
+		break;
 	}
 
-	void SetIOOption(A3DImport& Importer)
+	case CADLibrary::ECADFormat::INVENTOR:
+	case CADLibrary::ECADFormat::CATIA_3DXML:
 	{
-		// A3DRWParamsGeneralData Importer.m_sGeneral
-		Importer.m_sLoadData.m_sGeneral.m_bReadSolids = A3D_TRUE;
-		Importer.m_sLoadData.m_sGeneral.m_bReadSurfaces = A3D_TRUE;
-		Importer.m_sLoadData.m_sGeneral.m_bReadWireframes = A3D_FALSE;
-		Importer.m_sLoadData.m_sGeneral.m_bReadPmis = A3D_FALSE;
-		Importer.m_sLoadData.m_sGeneral.m_bReadAttributes = A3D_TRUE;
-		Importer.m_sLoadData.m_sGeneral.m_bReadHiddenObjects = A3D_FALSE;
-		Importer.m_sLoadData.m_sGeneral.m_bReadConstructionAndReferences = A3D_FALSE;
-		Importer.m_sLoadData.m_sGeneral.m_bReadActiveFilter = A3D_FALSE;
-		Importer.m_sLoadData.m_sGeneral.m_eReadingMode2D3D = kA3DRead_3D;
+		Importer.m_sLoadData.m_sIncremental.m_bLoadNoDependencies = false;
+		break;
+	}
 
-		Importer.m_sLoadData.m_sGeneral.m_eReadGeomTessMode = kA3DReadGeomOnly;
-		Importer.m_sLoadData.m_sGeneral.m_eDefaultUnit;
-		Importer.m_sLoadData.m_sGeneral.m_bReadFeature = A3D_FALSE;
+	default:
+		break;
+	};
+}
+}
+
+FTechSoftFileParser::FTechSoftFileParser(FCADFileData& InCADData, const FString& EnginePluginsPath)
+	: CADFileData(InCADData)
+	, TechSoftInterface(GetTechSoftInterface())
+{
+}
 
-		Importer.m_sLoadData.m_sGeneral.m_bReadConstraints = A3D_FALSE;
-		Importer.m_sLoadData.m_sGeneral.m_iNbMultiProcess = 1;
+ECADParsingResult FTechSoftFileParser::Process()
+{
+	// Process the file
+	ECADParsingResult Result = ECADParsingResult::ProcessOk;
 
-		Importer.m_sLoadData.m_sIncremental.m_bLoadNoDependencies = true;
-		Importer.m_sLoadData.m_sIncremental.m_bLoadStructureOnly = false;
-	}
+	const FFileDescriptor& File = CADFileData.GetCADFileDescription();
 
-	void UpdateIOOptionAccordingToFormat(const CADLibrary::ECADFormat Format, A3DImport& Importer)
+	if (File.GetPathOfFileToLoad().IsEmpty())
 	{
-		switch (Format)
-		{
-		case CADLibrary::ECADFormat::IGES:
-		{
-			break;
-		}
+		return ECADParsingResult::FileNotFound;
+	}
 
-		case CADLibrary::ECADFormat::CATIA:
-		{
-			break;
-		}
+	A3DImport Import(*File.GetPathOfFileToLoad());
 
-		case CADLibrary::ECADFormat::INVENTOR:
-		case CADLibrary::ECADFormat::CATIA_3DXML:
-		{
-			Importer.m_sLoadData.m_sIncremental.m_bLoadNoDependencies = false;
-			break;
-		}
+	TechSoftFileParserImpl::SetIOOption(Import);
 
-		default:
-			break;
-		};
+	// Add specific options according to format
+	Format = File.GetFileFormat();
+	TechSoftFileParserImpl::UpdateIOOptionAccordingToFormat(Format, Import);
+
+	A3DStatus IRet = TechSoftInterface.Import(Import);
+	if (IRet != A3D_SUCCESS && IRet != A3D_LOAD_MULTI_MODELS_CADFILE && IRet != A3D_LOAD_MISSING_COMPONENTS)
+	{
+		return ECADParsingResult::ProcessFailed;
 	}
-}
 
-	FTechSoftFileParser::FTechSoftFileParser(FCADFileData& InCADData, const FString& EnginePluginsPath)
-		: CADFileData(InCADData)
-		, TechSoftInterface(GetTechSoftInterface())
+	// save the file for the next load
+	if (CADFileData.IsCacheDefined())
 	{
+		FString CacheFilePath = CADFileData.GetCADCachePath();
+		if (CacheFilePath != File.GetPathOfFileToLoad())
+		{
+			// todo
+		}
 	}
 
-	ECADParsingResult FTechSoftFileParser::Process()
+	if (CADFileData.GetImportParameters().GetStitchingTechnique() != StitchingNone && !FImportParameters::bGDisableCADKernelTessellation)
 	{
-		// Process the file
-		ECADParsingResult Result = ECADParsingResult::ProcessOk;
+		// todo
+		//TechSoftInterface.Repair(CADFileData.GetStitchingTechnique());
+	}
 
-		const FFileDescriptor& File = CADFileData.GetCADFileDescription();
+	ReserveCADFileData();
 
-		if (File.GetPathOfFileToLoad().IsEmpty())
-		{
-			return ECADParsingResult::FileNotFound;
-		}
+	ReadMaterialsAndColors();
+
+	return TraverseModel(TechSoftInterface.GetModelFile());
+}
 
-		A3DImport Import(*File.GetPathOfFileToLoad());
+void FTechSoftFileParser::ReserveCADFileData()
+{
+	CountUnderModel(TechSoftInterface.GetModelFile());
 
-		TechSoftFileParserImpl::SetIOOption(Import);
+	CADFileData.ReserveBodyMeshes(ComponentCount[EComponentType::Body]);
 
-		// Add specific options according to format
-		Format = File.GetFileFormat();
-		TechSoftFileParserImpl::UpdateIOOptionAccordingToFormat(Format, Import);
+	FArchiveSceneGraph& SceneGraphArchive = CADFileData.GetSceneGraphArchive();
+	SceneGraphArchive.Bodies.Reserve(ComponentCount[EComponentType::Body]);
+	SceneGraphArchive.Components.Reserve(ComponentCount[EComponentType::Reference]);
+	SceneGraphArchive.UnloadedComponents.Reserve(ComponentCount[EComponentType::Reference]);
+	SceneGraphArchive.ExternalReferences.Reserve(ComponentCount[EComponentType::Reference]);
+	SceneGraphArchive.Instances.Reserve(ComponentCount[EComponentType::Occurrence]);
 
-		A3DStatus IRet = TechSoftInterface.Import(Import);
-		if (IRet != A3D_SUCCESS && IRet != A3D_LOAD_MULTI_MODELS_CADFILE && IRet != A3D_LOAD_MISSING_COMPONENTS)
-		{
-			return ECADParsingResult::ProcessFailed;
-		}
+	SceneGraphArchive.CADIdToBodyIndex.Reserve(ComponentCount[EComponentType::Body]);
+	SceneGraphArchive.CADIdToComponentIndex.Reserve(ComponentCount[EComponentType::Reference]);
+	SceneGraphArchive.CADIdToUnloadedComponentIndex.Reserve(ComponentCount[EComponentType::Reference]);
+	SceneGraphArchive.CADIdToInstanceIndex.Reserve(ComponentCount[EComponentType::Occurrence]);
+
+	uint32 MaterialNum = CountMaterial() + CountColor();
+	SceneGraphArchive.MaterialHIdToMaterial.Reserve(MaterialNum);
+}
+
+ECADParsingResult FTechSoftFileParser::TraverseModel(const A3DAsmModelFile* ModelFile)
+{
+	TUniqueTSObj<A3DAsmModelFileData> ModelFileData(ModelFile);
+	if (!ModelFileData.IsValid())
+	{
+		return ECADParsingResult::ProcessFailed;
+	}
+
+	ModellerType = (EModellerType)ModelFileData->m_eModellerType;
+	FileUnit = ModelFileData->m_dUnit;
 
-		// save the file for the next load
-		if (CADFileData.IsCacheDefined())
+	FEntityMetaData MetaData;
+	TraverseMetaData(ModelFile, MetaData);
+	TraverseSpecificMetaData(ModelFile, MetaData);
+
+	for (uint32 Index = 0; Index < ModelFileData->m_uiPOccurrencesSize; ++Index)
+	{
+		if (IsConfigurationSet(ModelFileData->m_ppPOccurrences[Index]))
 		{
-			FString CacheFilePath = CADFileData.GetCADCachePath();
-			if (CacheFilePath != File.GetPathOfFileToLoad())
-			{
-				// todo
-			}
+			//TraverseConfigurationSet(ModelFileData->m_ppPOccurrences[Index], MetaData);
 		}
-
-		if (CADFileData.GetImportParameters().GetStitchingTechnique() != StitchingNone && !FImportParameters::bGDisableCADKernelTessellation)
+		else
 		{
-			// todo
-			//TechSoftInterface.Repair(CADFileData.GetStitchingTechnique());
+			TraverseReference(ModelFileData->m_ppPOccurrences[Index]);
 		}
+	}
 
-		ReserveCADFileData();
+	return ECADParsingResult::ProcessOk;
+}
 
-		ReadMaterialsAndColors();
+void FTechSoftFileParser::TraverseReference(const A3DAsmProductOccurrence* ReferencePtr)
+{
+	FEntityMetaData MetaData;
+	TraverseMetaData(ReferencePtr, MetaData);
 
-		return TraverseModel(TechSoftInterface.GetModelFile());
+	if (MetaData.bRemoved || !MetaData.bShow)
+	{
+		return;
 	}
 
-	void FTechSoftFileParser::ReserveCADFileData()
-	{
-		CountUnderModel(TechSoftInterface.GetModelFile());
+	TraverseSpecificMetaData(ReferencePtr, MetaData);
+	BuildReferenceName(MetaData.MetaData);
 
-		CADFileData.ReserveBodyMeshes(ComponentCount[EComponentType::Body]);
+	TraverseMaterialProperties(ReferencePtr);
 
-		FArchiveSceneGraph& SceneGraphArchive = CADFileData.GetSceneGraphArchive();
-		SceneGraphArchive.Bodies.Reserve(ComponentCount[EComponentType::Body]);
-		SceneGraphArchive.Components.Reserve(ComponentCount[EComponentType::Reference]);
-		SceneGraphArchive.UnloadedComponents.Reserve(ComponentCount[EComponentType::Reference]);
-		SceneGraphArchive.ExternalReferences.Reserve(ComponentCount[EComponentType::Reference]);
-		SceneGraphArchive.Instances.Reserve(ComponentCount[EComponentType::Occurrence]);
+	int32 ComponentId = LastEntityId++;
+	int32 Index = CADFileData.AddComponent(ComponentId);
+	FArchiveComponent& Component = CADFileData.GetComponentAt(Index);
+	Component.MetaData = MoveTemp(MetaData.MetaData);
 
-		SceneGraphArchive.CADIdToBodyIndex.Reserve(ComponentCount[EComponentType::Body]);
-		SceneGraphArchive.CADIdToComponentIndex.Reserve(ComponentCount[EComponentType::Reference]);
-		SceneGraphArchive.CADIdToUnloadedComponentIndex.Reserve(ComponentCount[EComponentType::Reference]);
-		SceneGraphArchive.CADIdToInstanceIndex.Reserve(ComponentCount[EComponentType::Occurrence]);
+	TUniqueTSObj<A3DAsmProductOccurrenceData> ReferenceData(ReferencePtr);
+	if (!ReferenceData.IsValid())
+	{
+		return;
+	}
 
-		uint32 MaterialNum = CountMaterial() + CountColor();
-		SceneGraphArchive.MaterialHIdToMaterial.Reserve(MaterialNum);
+	for (uint32 OccurenceIndex = 0; OccurenceIndex < ReferenceData->m_uiPOccurrencesSize; ++OccurenceIndex)
+	{
+		int32 ChildrenId = TraverseOccurrence(ReferenceData->m_ppPOccurrences[OccurenceIndex]);
+		Component.Children.Add(ChildrenId);
 	}
 
-	ECADParsingResult FTechSoftFileParser::TraverseModel(const A3DAsmModelFile* ModelFile)
+	if (ReferenceData->m_pPart)
 	{
-		TUniqueTSObj<A3DAsmModelFileData> ModelFileData(ModelFile);
-		if (!ModelFileData.IsValid())
-		{
-			return ECADParsingResult::ProcessFailed;
-		}
+		TraversePartDefinition(ReferenceData->m_pPart, Component);
+	}
+}
 
-		ModellerType = (EModellerType)ModelFileData->m_eModellerType;
-		FileUnit = ModelFileData->m_dUnit;
 
-		FEntityMetaData MetaData;
-		TraverseMetaData(ModelFile, MetaData);
-		TraverseSpecificMetaData(ModelFile, MetaData);
+FArchiveInstance& FTechSoftFileParser::AddInstance(FEntityMetaData& InstanceMetaData)
+{
+	FCadId InstanceId = LastEntityId++;
+	int32 InstanceIndex = CADFileData.AddInstance(InstanceId);
+	FArchiveInstance& Instance = CADFileData.GetInstanceAt(InstanceIndex);
+	Instance.MetaData = MoveTemp(InstanceMetaData.MetaData);
+	return Instance;
+}
 
-		for (uint32 Index = 0; Index < ModelFileData->m_uiPOccurrencesSize; ++Index)
-		{
-			if (IsConfigurationSet(ModelFileData->m_ppPOccurrences[Index]))
-			{
-				//TraverseConfigurationSet(ModelFileData->m_ppPOccurrences[Index], MetaData);
-			}
-			else
-			{
-				TraverseReference(ModelFileData->m_ppPOccurrences[Index]);
-			}
-		}
+FArchiveComponent& FTechSoftFileParser::AddComponent(FEntityMetaData& ComponentMetaData, FArchiveInstance& Instance)
+{
+	FCadId ComponentId = LastEntityId++;
+	int32 ComponentIndex = CADFileData.AddComponent(ComponentId);
+	FArchiveComponent& Prototype = CADFileData.GetComponentAt(ComponentIndex);
+	Prototype.MetaData = MoveTemp(ComponentMetaData.MetaData);
 
-		return ECADParsingResult::ProcessOk;
-	}
+	Instance.ReferenceNodeId = ComponentId;
 
-	void FTechSoftFileParser::TraverseReference(const A3DAsmProductOccurrence* ReferencePtr)
-	{
-		FEntityMetaData MetaData;
-		TraverseMetaData(ReferencePtr, MetaData);
+	return Prototype;
+}
+
+FArchiveUnloadedComponent& FTechSoftFileParser::AddUnloadedComponent(FEntityMetaData& ComponentMetaData, FArchiveInstance& Instance)
+{
+	FCadId ComponentId = LastEntityId++;
+	int32 ComponentIndex = CADFileData.AddUnloadedComponent(ComponentId);
+	FArchiveUnloadedComponent& Component = CADFileData.GetUnloadedComponentAt(ComponentIndex);
+	Component.MetaData = MoveTemp(ComponentMetaData.MetaData);
 
-		if (MetaData.bRemoved || !MetaData.bShow)
+	Instance.bIsExternalReference = true;
+	Instance.ReferenceNodeId = ComponentId;
+
+	Instance.ExternalReference = ComponentMetaData.ExternalFile;
+
+	if (Format == ECADFormat::SOLIDWORKS /*|| CADFileData.FileFormat() == ECADFormat::JT*/)
+	{
+		if (FString* ConfigurationName = ComponentMetaData.MetaData.Find(TEXT("ConfigurationName")))
 		{
-			return;
+			Instance.ExternalReference.SetConfiguration(*ConfigurationName);
 		}
+	}
+
+	CADFileData.AddExternalRef(Instance.ExternalReference);
+
+	return Component;
+}
 
-		TraverseSpecificMetaData(ReferencePtr, MetaData);
-		BuildReferenceName(MetaData.MetaData);
+FArchiveComponent& FTechSoftFileParser::AddOccurence(FEntityMetaData& InstanceMetaData, FCadId& OutComponentId)
+{
+	FArchiveInstance& Instance = AddInstance(InstanceMetaData);
+	OutComponentId = Instance.ObjectId;
+	FEntityMetaData ReferenceMetaData;
+	FArchiveComponent& Prototype = AddComponent(ReferenceMetaData, Instance);
+	return Prototype;
+}
 
-		TraverseMaterialProperties(ReferencePtr);
+FArchiveComponent& FTechSoftFileParser::AddOccurence(FEntityMetaData& InstanceMetaData, FEntityMetaData& ReferenceMetaData, FCadId& OutComponentId)
+{
+	FArchiveInstance& Instance = AddInstance(InstanceMetaData);
+	OutComponentId = Instance.ObjectId;
+	FArchiveComponent& Prototype = AddComponent(ReferenceMetaData, Instance);
+	return Prototype;
+}
 
-		int32 ComponentId = LastEntityId++;
-		int32 Index = CADFileData.AddComponent(ComponentId);
-		FArchiveComponent& Component = CADFileData.GetComponentAt(Index);
-		Component.MetaData = MoveTemp(MetaData.MetaData);
+FArchiveBody& FTechSoftFileParser::AddBody(FEntityMetaData& BodyMetaData)
+{
+	FCadId BodyId = LastEntityId++;
+	int32 BodyIndex = CADFileData.AddBody(BodyId);
+	FArchiveBody& Body = CADFileData.GetBodyAt(BodyIndex);
+	Body.MetaData = MoveTemp(BodyMetaData.MetaData);
 
-		TUniqueTSObj<A3DAsmProductOccurrenceData> ReferenceData(ReferencePtr);
-		if (!ReferenceData.IsValid())
-		{
-			return;
-		}
+	return Body;
+}
 
-		for (uint32 OccurenceIndex = 0; OccurenceIndex < ReferenceData->m_uiPOccurrencesSize; ++OccurenceIndex)
-		{
-			int32 ChildrenId = TraverseOccurrence(ReferenceData->m_ppPOccurrences[OccurenceIndex]);
-			Component.Children.Add(ChildrenId);
-		}
+FCadId FTechSoftFileParser::TraverseOccurrence(const A3DAsmProductOccurrence* OccurrencePtr)
+{
+	FEntityMetaData InstanceMetaData;
+	TraverseMetaData(OccurrencePtr, InstanceMetaData);
 
-		if (ReferenceData->m_pPart)
-		{
-			TraversePartDefinition(ReferenceData->m_pPart, Component);
-		}
+	if (InstanceMetaData.bRemoved || !InstanceMetaData.bShow)
+	{
+		return 0;
 	}
 
+	TraverseSpecificMetaData(OccurrencePtr, InstanceMetaData);
+	BuildInstanceName(InstanceMetaData.MetaData);
+
+	TraverseMaterialProperties(OccurrencePtr);
+	TraverseLayer(OccurrencePtr);
 
-	FArchiveInstance& FTechSoftFileParser::AddInstance(FEntityMetaData& InstanceMetaData)
+	TUniqueTSObj<A3DAsmProductOccurrenceData> OccurrenceData(OccurrencePtr);
+	if (!OccurrenceData.IsValid())
 	{
-		FCadId InstanceId = LastEntityId++;
-		int32 InstanceIndex = CADFileData.AddInstance(InstanceId);
-		FArchiveInstance& Instance = CADFileData.GetInstanceAt(InstanceIndex);
-		Instance.MetaData = MoveTemp(InstanceMetaData.MetaData);
-		return Instance;
+		return 0;
 	}
 
-	FArchiveComponent& FTechSoftFileParser::AddComponent(FEntityMetaData& ComponentMetaData, FArchiveInstance& Instance)
-	{
-		FCadId ComponentId = LastEntityId++;
-		int32 ComponentIndex = CADFileData.AddComponent(ComponentId);
-		FArchiveComponent& Prototype = CADFileData.GetComponentAt(ComponentIndex);
-		Prototype.MetaData = MoveTemp(ComponentMetaData.MetaData);
+	FArchiveInstance& Instance = AddInstance(InstanceMetaData);
 
-		Instance.ReferenceNodeId = ComponentId;
+	Instance.TransformMatrix = TraverseTransformation(OccurrenceData->m_pLocation);
 
-		return Prototype;
+	FEntityMetaData PrototypeMetaData;
+	bool bIsUnloadedPrototype = false;
+	if (OccurrenceData->m_pPrototype)
+	{
+		FMatrix PrototypeMatrix;
+		TraversePrototype(OccurrenceData->m_pPrototype, PrototypeMetaData, PrototypeMatrix);
+		Instance.TransformMatrix *= PrototypeMatrix;
 	}
 
-	FArchiveUnloadedComponent& FTechSoftFileParser::AddUnloadedComponent(FEntityMetaData& ComponentMetaData, FArchiveInstance& Instance)
+	if (PrototypeMetaData.bUnloaded)
 	{
-		FCadId ComponentId = LastEntityId++;
-		int32 ComponentIndex = CADFileData.AddUnloadedComponent(ComponentId);
-		FArchiveUnloadedComponent& Component = CADFileData.GetUnloadedComponentAt(ComponentIndex);
-		Component.MetaData = MoveTemp(ComponentMetaData.MetaData);
-
-		Instance.bIsExternalReference = true;
-		Instance.ReferenceNodeId = ComponentId;
+		FArchiveUnloadedComponent& Prototype = AddUnloadedComponent(PrototypeMetaData, Instance);
+	}
+	else
+	{
+		FArchiveComponent& Component = AddComponent(PrototypeMetaData, Instance);
 
-		Instance.ExternalReference = ComponentMetaData.ExternalFile;
+		for (uint32 Index = 0; Index < OccurrenceData->m_uiPOccurrencesSize; ++Index)
+		{
+			int32 ChildrenId = TraverseOccurrence(OccurrenceData->m_ppPOccurrences[Index]);
+			Component.Children.Add(ChildrenId);
+		}
 
-		if (Format == ECADFormat::SOLIDWORKS /*|| CADFileData.FileFormat() == ECADFormat::JT*/)
+		if (OccurrenceData->m_pPart)
 		{
-			if (FString* ConfigurationName = ComponentMetaData.MetaData.Find(TEXT("ConfigurationName")))
-			{
-				Instance.ExternalReference.SetConfiguration(*ConfigurationName);
-			}
+			TraversePartDefinition(OccurrenceData->m_pPart, Component);
 		}
+	}
 
-		CADFileData.AddExternalRef(Instance.ExternalReference);
+	return Instance.ObjectId;
+}
 
-		return Component;
+FFileDescriptor FTechSoftFileParser::GetOccurrenceFileName(const A3DAsmProductOccurrence* OccurrencePtr)
+{
+	TUniqueTSObj<A3DAsmProductOccurrenceData> OccurrenceData(OccurrencePtr);
+	if (!OccurrenceData.IsValid())
+	{
+		return FFileDescriptor();
 	}
 
-	FArchiveComponent& FTechSoftFileParser::AddOccurence(FEntityMetaData& InstanceMetaData, FCadId& OutComponentId)
+	if (OccurrenceData->m_pPrototype)
 	{
-		FArchiveInstance& Instance = AddInstance(InstanceMetaData);
-		OutComponentId = Instance.ObjectId;
-		FEntityMetaData ReferenceMetaData;
-		FArchiveComponent& Prototype = AddComponent(ReferenceMetaData, Instance);
-		return Prototype;
+		return GetOccurrenceFileName(OccurrenceData->m_pPrototype);
 	}
 
-	FArchiveComponent& FTechSoftFileParser::AddOccurence(FEntityMetaData& InstanceMetaData, FEntityMetaData& ReferenceMetaData, FCadId& OutComponentId)
+	A3DUTF8Char* FilePathUTF8Ptr = nullptr;
+	int Ret = A3DAsmProductOccurrenceGetFilePathName(OccurrencePtr, &FilePathUTF8Ptr);
+	if (Ret == A3D_SUCCESS && FilePathUTF8Ptr)
 	{
-		FArchiveInstance& Instance = AddInstance(InstanceMetaData);
-		OutComponentId = Instance.ObjectId;
-		FArchiveComponent& Prototype = AddComponent(ReferenceMetaData, Instance);
-		return Prototype;
+		FString OccurrenceFilePath = UTF8_TO_TCHAR(FilePathUTF8Ptr);
+		A3DAsmProductOccurrenceGetFilePathName(NULL, &FilePathUTF8Ptr);
+		FString OccurrenceFileName = FPaths::GetCleanFilename(OccurrenceFilePath);
+		if (OccurrenceFileName != CADFileData.GetCADFileDescription().GetFileName())
+		{
+			return FFileDescriptor(*OccurrenceFilePath, nullptr, *CADFileData.GetCADFileDescription().GetRootFolder());
+		}
 	}
 
-	FArchiveBody& FTechSoftFileParser::AddBody(FEntityMetaData& BodyMetaData)
+	FilePathUTF8Ptr = nullptr;
+	Ret = A3DAsmProductOccurrenceGetOriginalFilePathName(OccurrencePtr, &FilePathUTF8Ptr);
+	if (Ret == A3D_SUCCESS && FilePathUTF8Ptr)
 	{
-		FCadId BodyId = LastEntityId++;
-		int32 BodyIndex = CADFileData.AddBody(BodyId);
-		FArchiveBody& Body = CADFileData.GetBodyAt(BodyIndex);
-		Body.MetaData = MoveTemp(BodyMetaData.MetaData);
-
-		return Body;
+		FString OriginalFilePath = UTF8_TO_TCHAR(FilePathUTF8Ptr);
+		A3DAsmProductOccurrenceGetFilePathName(NULL, &FilePathUTF8Ptr);
+		FString OriginalFileName = FPaths::GetCleanFilename(OriginalFilePath);
+		if (OriginalFileName != CADFileData.GetCADFileDescription().GetFileName())
+		{
+			return FFileDescriptor(*OriginalFilePath, nullptr, *CADFileData.GetCADFileDescription().GetRootFolder());
+		}
 	}
 
-	FCadId FTechSoftFileParser::TraverseOccurrence(const A3DAsmProductOccurrence* OccurrencePtr)
+	return FFileDescriptor();
+}
+
+void FTechSoftFileParser::TraversePrototype(const A3DAsmProductOccurrence* InPrototypePtr, FEntityMetaData& OutPrototypeMetaData, FMatrix& OutPrototypeMatrix)
+{
+	TUniqueTSObj<A3DAsmProductOccurrenceData> PrototypeData(InPrototypePtr);
+	if (!PrototypeData.IsValid())
 	{
-		FEntityMetaData InstanceMetaData;
-		TraverseMetaData(OccurrencePtr, InstanceMetaData);
+		return;
+	}
 
-		if (InstanceMetaData.bRemoved || !InstanceMetaData.bShow)
-		{
-			return 0;
-		}
+	TraverseMetaData(InPrototypePtr, OutPrototypeMetaData);
+	TraverseSpecificMetaData(InPrototypePtr, OutPrototypeMetaData);
 
-		TraverseSpecificMetaData(OccurrencePtr, InstanceMetaData);
-		BuildInstanceName(InstanceMetaData.MetaData);
+	TraverseMaterialProperties(InPrototypePtr);
 
-		TraverseMaterialProperties(OccurrencePtr);
-		TraverseLayer(OccurrencePtr);
+	OutPrototypeMatrix = TraverseTransformation(PrototypeData->m_pLocation);
 
-		TUniqueTSObj<A3DAsmProductOccurrenceData> OccurrenceData(OccurrencePtr);
-		if (!OccurrenceData.IsValid())
-		{
-			return 0;
-		}
+	OutPrototypeMetaData.bUnloaded = TechSoftFileParserImpl::IsUnloadedPrototype(*PrototypeData);
+	if (OutPrototypeMetaData.bUnloaded)
+	{
+		OutPrototypeMetaData.ExternalFile = GetOccurrenceFileName(InPrototypePtr);
+		OutPrototypeMetaData.MetaData.Add(TEXT("Name"), OutPrototypeMetaData.ExternalFile.GetFileName());
+	}
 
-		FArchiveInstance& Instance = AddInstance(InstanceMetaData);
+	BuildReferenceName(OutPrototypeMetaData.MetaData);
+}
 
-		Instance.TransformMatrix = TraverseTransformation(OccurrenceData->m_pLocation);
+void FTechSoftFileParser::TraversePartDefinition(const A3DAsmPartDefinition* PartDefinitionPtr, FArchiveComponent& Part)
+{
+	FEntityMetaData PartMetaData;
+	TraverseMetaData(PartDefinitionPtr, PartMetaData);
 
-		FEntityMetaData PrototypeMetaData;
-		bool bIsUnloadedPrototype = false;
-		if (OccurrenceData->m_pPrototype)
-		{
-			FMatrix PrototypeMatrix;
-			TraversePrototype(OccurrenceData->m_pPrototype, PrototypeMetaData, PrototypeMatrix);
-			Instance.TransformMatrix *= PrototypeMatrix;
-		}
+	if (PartMetaData.bRemoved || !PartMetaData.bShow)
+	{
+		return;
+	}
 
-		if (PrototypeMetaData.bUnloaded)
-		{
-			FArchiveUnloadedComponent& Prototype = AddUnloadedComponent(PrototypeMetaData, Instance);
-		}
-		else
-		{
-			FArchiveComponent& Component = AddComponent(PrototypeMetaData, Instance);
+	TraverseSpecificMetaData(PartDefinitionPtr, PartMetaData);
+	BuildPartName(PartMetaData.MetaData);
 
-			for (uint32 Index = 0; Index < OccurrenceData->m_uiPOccurrencesSize; ++Index)
-			{
-				int32 ChildrenId = TraverseOccurrence(OccurrenceData->m_ppPOccurrences[Index]);
-				Component.Children.Add(ChildrenId);
-			}
+	TraverseMaterialProperties(PartDefinitionPtr);
+	TraverseLayer(PartDefinitionPtr);
 
-			if (OccurrenceData->m_pPart)
-			{
-				TraversePartDefinition(OccurrenceData->m_pPart, Component);
-			}
-		}
-
-		return Instance.ObjectId;
-	}
-
-	FFileDescriptor FTechSoftFileParser::GetOccurrenceFileName(const A3DAsmProductOccurrence* OccurrencePtr)
+	TUniqueTSObj<A3DAsmPartDefinitionData> PartData(PartDefinitionPtr);
+	if (PartData.IsValid())
 	{
-		TUniqueTSObj<A3DAsmProductOccurrenceData> OccurrenceData(OccurrencePtr);
-		if (!OccurrenceData.IsValid())
-		{
-			return FFileDescriptor();
-		}
-
-		if (OccurrenceData->m_pPrototype)
+		for (unsigned int Index = 0; Index < PartData->m_uiRepItemsSize; ++Index)
 		{
-			return GetOccurrenceFileName(OccurrenceData->m_pPrototype);
-		}
-
-		A3DUTF8Char* FilePathUTF8Ptr = nullptr;
-		int Ret = A3DAsmProductOccurrenceGetFilePathName(OccurrencePtr, &FilePathUTF8Ptr);
-		if (Ret == A3D_SUCCESS && FilePathUTF8Ptr)
-		{
-			FString OccurrenceFilePath = UTF8_TO_TCHAR(FilePathUTF8Ptr);
-			A3DAsmProductOccurrenceGetFilePathName(NULL, &FilePathUTF8Ptr);
-			FString OccurrenceFileName = FPaths::GetCleanFilename(OccurrenceFilePath);
-			if (OccurrenceFileName != CADFileData.GetCADFileDescription().GetFileName())
-			{
-				return FFileDescriptor(*OccurrenceFilePath, nullptr, *CADFileData.GetCADFileDescription().GetRootFolder());
-			}
+			int32 ChildId = TraverseRepresentationItem(PartData->m_ppRepItems[Index], PartMetaData);
+			Part.Children.Add(ChildId);
 		}
+	}
+}
 
-		FilePathUTF8Ptr = nullptr;
-		Ret = A3DAsmProductOccurrenceGetOriginalFilePathName(OccurrencePtr, &FilePathUTF8Ptr);
-		if (Ret == A3D_SUCCESS && FilePathUTF8Ptr)
-		{
-			FString OriginalFilePath = UTF8_TO_TCHAR(FilePathUTF8Ptr);
-			A3DAsmProductOccurrenceGetFilePathName(NULL, &FilePathUTF8Ptr);
-			FString OriginalFileName = FPaths::GetCleanFilename(OriginalFilePath);
-			if (OriginalFileName != CADFileData.GetCADFileDescription().GetFileName())
-			{
-				return FFileDescriptor(*OriginalFilePath, nullptr, *CADFileData.GetCADFileDescription().GetRootFolder());
-			}
-		}
+FCadId FTechSoftFileParser::TraverseRepresentationItem(A3DRiRepresentationItem* RepresentationItem, FEntityMetaData& PartMetaData)
+{
+	A3DEEntityType Type;
+	A3DEntityGetType(RepresentationItem, &Type);
 
-		return FFileDescriptor();
+	switch (Type)
+	{
+	case kA3DTypeRiSet:
+		return TraverseRepresentationSet(RepresentationItem, PartMetaData);
+	case kA3DTypeRiBrepModel:
+		return TraverseBRepModel(RepresentationItem, PartMetaData);
+	case kA3DTypeRiPolyBrepModel:
+		return TraversePolyBRepModel(RepresentationItem, PartMetaData);
+	default:
+		break;
 	}
+	return 0;
+}
 
-	void FTechSoftFileParser::TraversePrototype(const A3DAsmProductOccurrence* InPrototypePtr, FEntityMetaData& OutPrototypeMetaData, FMatrix& OutPrototypeMatrix)
+FCadId FTechSoftFileParser::TraverseRepresentationSet(const A3DRiSet* RepresentationSetPtr, FEntityMetaData& PartMetaData)
+{
+	TUniqueTSObj<A3DRiSetData> RepresentationSetData(RepresentationSetPtr);
+	if (!RepresentationSetData.IsValid())
 	{
-		TUniqueTSObj<A3DAsmProductOccurrenceData> PrototypeData(InPrototypePtr);
-		if (!PrototypeData.IsValid())
-		{
-			return;
-		}
+		return 0;
+	}
 
-		TraverseMetaData(InPrototypePtr, OutPrototypeMetaData);
-		TraverseSpecificMetaData(InPrototypePtr, OutPrototypeMetaData);
+	FEntityMetaData RepresentationSetMetaData;
+	TraverseMetaData(RepresentationSetPtr, RepresentationSetMetaData);
 
-		TraverseMaterialProperties(InPrototypePtr);
+	if (RepresentationSetMetaData.bRemoved || !RepresentationSetMetaData.bShow)
+	{
+		return 0;
+	}
 
-		OutPrototypeMatrix = TraverseTransformation(PrototypeData->m_pLocation);
+	TraverseMaterialProperties(RepresentationSetPtr);
 
-		OutPrototypeMetaData.bUnloaded = TechSoftFileParserImpl::IsUnloadedPrototype(*PrototypeData);
-		if (OutPrototypeMetaData.bUnloaded)
-		{
-			OutPrototypeMetaData.ExternalFile = GetOccurrenceFileName(InPrototypePtr);
-			OutPrototypeMetaData.MetaData.Add(TEXT("Name"), OutPrototypeMetaData.ExternalFile.GetFileName());
-		}
+	FCadId RepresentationSetId = 0;
+	FArchiveComponent& RepresentationSet = AddOccurence(RepresentationSetMetaData, RepresentationSetId);
 
-		BuildReferenceName(OutPrototypeMetaData.MetaData);
+	for (A3DUns32 ui = 0; ui < RepresentationSetData->m_uiRepItemsSize; ++ui)
+	{
+		int32 ChildId = TraverseRepresentationItem(RepresentationSetData->m_ppRepItems[ui], RepresentationSetMetaData);
+		RepresentationSet.Children.Add(ChildId);
 	}
+	return RepresentationSetId;
+}
 
-	void FTechSoftFileParser::TraversePartDefinition(const A3DAsmPartDefinition* PartDefinitionPtr, FArchiveComponent& Part)
+FCadId FTechSoftFileParser::TraverseBRepModel(A3DRiBrepModel* BRepModelPtr, FEntityMetaData& PartMetaData)
+{
+	TUniqueTSObj<A3DRiBrepModelData> BodyData(BRepModelPtr);
+	if (!BodyData.IsValid())
 	{
-		FEntityMetaData PartMetaData;
-		TraverseMetaData(PartDefinitionPtr, PartMetaData);
-
-		if (PartMetaData.bRemoved || !PartMetaData.bShow)
-		{
-			return;
-		}
-
-		TraverseSpecificMetaData(PartDefinitionPtr, PartMetaData);
-		BuildPartName(PartMetaData.MetaData);
-
-		TraverseMaterialProperties(PartDefinitionPtr);
-		TraverseLayer(PartDefinitionPtr);
-
-		TUniqueTSObj<A3DAsmPartDefinitionData> PartData(PartDefinitionPtr);
-		if (PartData.IsValid())
-		{
-			for (unsigned int Index = 0; Index < PartData->m_uiRepItemsSize; ++Index)
-			{
-				int32 ChildId = TraverseRepresentationItem(PartData->m_ppRepItems[Index], PartMetaData);
-				Part.Children.Add(ChildId);
-			}
-		}
+		return 0;
 	}
 
-	FCadId FTechSoftFileParser::TraverseRepresentationItem(A3DRiRepresentationItem* RepresentationItem, FEntityMetaData& PartMetaData)
+	FEntityMetaData BRepMetaData;
+	TraverseMetaData(BRepModelPtr, BRepMetaData);
+	if (!BRepMetaData.bShow || BRepMetaData.bRemoved)
 	{
-		A3DEEntityType Type;
-		A3DEntityGetType(RepresentationItem, &Type);
-
-		switch (Type)
-		{
-		case kA3DTypeRiSet:
-			return TraverseRepresentationSet(RepresentationItem, PartMetaData);
-		case kA3DTypeRiBrepModel:
-			return TraverseBRepModel(RepresentationItem, PartMetaData);
-		case kA3DTypeRiPolyBrepModel:
-			return TraversePolyBRepModel(RepresentationItem, PartMetaData);
-		default:
-			break;
-		}
 		return 0;
 	}
 
-	FCadId FTechSoftFileParser::TraverseRepresentationSet(const A3DRiSet* RepresentationSetPtr, FEntityMetaData& PartMetaData)
-	{
-		TUniqueTSObj<A3DRiSetData> RepresentationSetData(RepresentationSetPtr);
-		if (!RepresentationSetData.IsValid())
-		{
-			return 0;
-		}
+	TraverseSpecificMetaData(BRepModelPtr, BRepMetaData);
+	TraverseMaterialProperties(BRepModelPtr);
 
-		FEntityMetaData RepresentationSetMetaData;
-		TraverseMetaData(RepresentationSetPtr, RepresentationSetMetaData);
+	FArchiveBody& Body = AddBody(BRepMetaData);
 
-		if (RepresentationSetMetaData.bRemoved || !RepresentationSetMetaData.bShow)
-		{
-			return 0;
-		}
+	TraverseRepresentationContent(BRepModelPtr, Body);
+	if (FImportParameters::bGDisableCADKernelTessellation)
+	{
+		MeshRepresentationWithTechSoft(BRepModelPtr, Body);
+	}
+	else
+	{
+		// Mesh with CADKernel
+	}
+	return Body.ObjectId;
+}
 
-		TraverseMaterialProperties(RepresentationSetPtr);
+void FTechSoftFileParser::TraverseRepresentationContent(const A3DRiRepresentationItem* RepresentationItemPtr, FArchiveBody& Body)
+{
+	TUniqueTSObj<A3DRiRepresentationItemData> RepresentationItemData(RepresentationItemPtr);
+	if (!RepresentationItemData.IsValid())
+	{
+		return;
+	}
 
-		FCadId RepresentationSetId = 0;
-		FArchiveComponent& RepresentationSet = AddOccurence(RepresentationSetMetaData, RepresentationSetId);
+	if (RepresentationItemData->m_pCoordinateSystem)
+	{
+		TraverseCoordinateSystem(RepresentationItemData->m_pCoordinateSystem);
+	}
 
-		for (A3DUns32 ui = 0; ui < RepresentationSetData->m_uiRepItemsSize; ++ui)
-		{
-			int32 ChildId = TraverseRepresentationItem(RepresentationSetData->m_ppRepItems[ui], RepresentationSetMetaData);
-			RepresentationSet.Children.Add(ChildId);
-		}
-		return RepresentationSetId;
+	if (RepresentationItemData->m_pTessBase)
+	{
+		TraverseTessellationBase(RepresentationItemData->m_pTessBase, Body);
 	}
+}
 
-	FCadId FTechSoftFileParser::TraverseBRepModel(A3DRiBrepModel* BRepModelPtr, FEntityMetaData& PartMetaData)
+FCadId FTechSoftFileParser::TraversePolyBRepModel(const A3DRiPolyBrepModel* PolygonalPtr, FEntityMetaData& PartMetaData)
+{
+	TUniqueTSObj<A3DRiPolyBrepModelData> BodyData(PolygonalPtr);
+	if (!BodyData.IsValid())
 	{
-		TUniqueTSObj<A3DRiBrepModelData> BodyData(BRepModelPtr);
-		if (!BodyData.IsValid())
-		{
-			return 0;
-		}
+		return 0;
+	}
 
-		FEntityMetaData BRepMetaData;
-		TraverseMetaData(BRepModelPtr, BRepMetaData);
-		if (!BRepMetaData.bShow || BRepMetaData.bRemoved)
-		{
-			return 0;
-		}
+	FEntityMetaData BRepMetaData;
+	TraverseMetaData(PolygonalPtr, BRepMetaData);
+	if (!BRepMetaData.bShow || BRepMetaData.bRemoved)
+	{
+		return 0;
+	}
 
-		TraverseSpecificMetaData(BRepModelPtr, BRepMetaData);
-		TraverseMaterialProperties(BRepModelPtr);
+	TraverseSpecificMetaData(PolygonalPtr, BRepMetaData);
+	TraverseMaterialProperties(PolygonalPtr);
 
-		FArchiveBody& Body = AddBody(BRepMetaData);
+	FArchiveBody& Body = AddBody(BRepMetaData);
+	TraverseRepresentationContent(PolygonalPtr, Body);
 
-		TraverseRepresentationContent(BRepModelPtr, Body);
-		if (FImportParameters::bGDisableCADKernelTessellation)
-		{
-			MeshRepresentationWithTechSoft(BRepModelPtr, Body);
-		}
-		else
-		{
-			// Mesh with CADKernel
-		}
-		return Body.ObjectId;
-	}
+	return Body.ObjectId;
+}
 
-	void FTechSoftFileParser::TraverseRepresentationContent(const A3DRiRepresentationItem* RepresentationItemPtr, FArchiveBody& Body)
+void FTechSoftFileParser::TraverseMetaData(const A3DEntity* Entity, FEntityMetaData& OutMetaData)
+{
+	TUniqueTSObj<A3DRootBaseData> MetaData(Entity);
+	if (MetaData.IsValid())
 	{
-		TUniqueTSObj<A3DRiRepresentationItemData> RepresentationItemData(RepresentationItemPtr);
-		if (!RepresentationItemData.IsValid())
+		if (false && MetaData->m_uiPersistentId > 0)
 		{
-			return;
+			FString PersistentId = FString::FromInt(MetaData->m_uiPersistentId);
+			OutMetaData.MetaData.Emplace(TEXT("PersistentId"), PersistentId);
 		}
 
-		if (RepresentationItemData->m_pCoordinateSystem)
+		if (MetaData->m_pcName && MetaData->m_pcName[0] != '\0')
 		{
-			TraverseCoordinateSystem(RepresentationItemData->m_pCoordinateSystem);
+			FString SDKName = UTF8_TO_TCHAR(MetaData->m_pcName);
+			SDKName = TechSoftFileParserImpl::CleanSdkName(SDKName);
+			OutMetaData.MetaData.Emplace(TEXT("SDKName"), SDKName);
 		}
 
-		if (RepresentationItemData->m_pTessBase)
+		TUniqueTSObj<A3DMiscAttributeData> AttributeData;
+		for (A3DUns32 Index = 0; Index < MetaData->m_uiSize; ++Index)
 		{
-			TraverseTessellationBase(RepresentationItemData->m_pTessBase, Body);
+			AttributeData.FillFrom(MetaData->m_ppAttributes[Index]);
+			if (AttributeData.IsValid())
+			{
+				TechSoftFileParserImpl::TraverseAttribute(*AttributeData, OutMetaData.MetaData);
+			}
 		}
 	}
 
-	FCadId FTechSoftFileParser::TraversePolyBRepModel(const A3DRiPolyBrepModel* PolygonalPtr, FEntityMetaData& PartMetaData)
+	if (A3DEntityIsBaseWithGraphicsType(Entity))
 	{
-		TUniqueTSObj<A3DRiPolyBrepModelData> BodyData(PolygonalPtr);
-		if (!BodyData.IsValid())
-		{
-			return 0;
-		}
-
-		FEntityMetaData BRepMetaData;
-		TraverseMetaData(PolygonalPtr, BRepMetaData);
-		if (!BRepMetaData.bShow || BRepMetaData.bRemoved)
+		TUniqueTSObj<A3DRootBaseWithGraphicsData> MetaDataWithGraphics(Entity);
+		if (MetaDataWithGraphics.IsValid())
 		{
-			return 0;
+			if (MetaDataWithGraphics->m_pGraphics != NULL)
+			{
+				TraverseGraphics(MetaDataWithGraphics->m_pGraphics, OutMetaData);
+			}
 		}
-
-		TraverseSpecificMetaData(PolygonalPtr, BRepMetaData);
-		TraverseMaterialProperties(PolygonalPtr);
-
-		FArchiveBody& Body = AddBody(BRepMetaData);
-		TraverseRepresentationContent(PolygonalPtr, Body);
-
-		return Body.ObjectId;
 	}
+}
 
-	void FTechSoftFileParser::TraverseMetaData(const A3DEntity* Entity, FEntityMetaData& OutMetaData)
+void FTechSoftFileParser::BuildReferenceName(TMap<FString, FString>& MetaData)
+{
+	if (MetaData.IsEmpty())
 	{
-		TUniqueTSObj<A3DRootBaseData> MetaData(Entity);
-		if (MetaData.IsValid())
-		{
-			if (false && MetaData->m_uiPersistentId > 0)
-			{
-				FString PersistentId = FString::FromInt(MetaData->m_uiPersistentId);
-				OutMetaData.MetaData.Emplace(TEXT("PersistentId"), PersistentId);
-			}
-
-			if (MetaData->m_pcName && MetaData->m_pcName[0] != '\0')
-			{
-				FString SDKName = UTF8_TO_TCHAR(MetaData->m_pcName);
-				SDKName = TechSoftFileParserImpl::CleanSdkName(SDKName);
-				OutMetaData.MetaData.Emplace(TEXT("SDKName"), SDKName);
-			}
-
-			TUniqueTSObj<A3DMiscAttributeData> AttributeData;
-			for (A3DUns32 Index = 0; Index < MetaData->m_uiSize; ++Index)
-			{
-				AttributeData.FillFrom(MetaData->m_ppAttributes[Index]);
-				if (AttributeData.IsValid())
-				{
-					TechSoftFileParserImpl::TraverseAttribute(*AttributeData, OutMetaData.MetaData);
-				}
-			}
-		}
+		return;
+	}
 
-		if (A3DEntityIsBaseWithGraphicsType(Entity))
+	FString* NamePtr = MetaData.Find(TEXT("InstanceName"));
+	if (NamePtr != nullptr)
+	{
+		FString& Name = MetaData.FindOrAdd(TEXT("Name"));
+		Name = *NamePtr;
+		if (Format == ECADFormat::CATIA)
 		{
-			TUniqueTSObj<A3DRootBaseWithGraphicsData> MetaDataWithGraphics(Entity);
-			if (MetaDataWithGraphics.IsValid())
-			{
-				if (MetaDataWithGraphics->m_pGraphics != NULL)
-				{
-					TraverseGraphics(MetaDataWithGraphics->m_pGraphics, OutMetaData);
-				}
-			}
+			Name = TechSoftFileParserImpl::CleanCatiaReferenceName(Name);
 		}
+		return;
 	}
 
-	void FTechSoftFileParser::BuildReferenceName(TMap<FString, FString>& MetaData)
+	if (Format == ECADFormat::JT)
 	{
-		if (MetaData.IsEmpty())
+		if (TechSoftFileParserImpl::ReplaceOrAddNameValue(MetaData, TEXT("SDKName")))
 		{
 			return;
 		}
+	}
 
-		FString* NamePtr = MetaData.Find(TEXT("InstanceName"));
-		if (NamePtr != nullptr)
-		{
-			FString& Name = MetaData.FindOrAdd(TEXT("Name"));
-			Name = *NamePtr;
-			if (Format == ECADFormat::CATIA)
-			{
-				Name = TechSoftFileParserImpl::CleanCatiaReferenceName(Name);
-			}
-			return;
-		}
+	if (TechSoftFileParserImpl::CheckIfNameExists(MetaData))
+	{
+		return;
+	}
 
-		if (Format == ECADFormat::JT)
-		{
-			if (TechSoftFileParserImpl::ReplaceOrAddNameValue(MetaData, TEXT("SDKName")))
-			{
-				return;
-			}
-		}
+	if (TechSoftFileParserImpl::ReplaceOrAddNameValue(MetaData, TEXT("PartNumber")))
+	{
+		return;
+	}
 
-		if (TechSoftFileParserImpl::CheckIfNameExists(MetaData))
-		{
-			return;
-		}
+	NamePtr = MetaData.Find(TEXT("SDKName"));
+	if (NamePtr != nullptr)
+	{
+		FString SdkName = *NamePtr;
 
-		if (TechSoftFileParserImpl::ReplaceOrAddNameValue(MetaData, TEXT("PartNumber")))
+		switch (Format)
 		{
-			return;
-		}
+		case ECADFormat::CATIA_3DXML:
+			SdkName = TechSoftFileParserImpl::Clean3dxmlReferenceSdkName(SdkName);
+			break;
 
-		NamePtr = MetaData.Find(TEXT("SDKName"));
-		if (NamePtr != nullptr)
-		{
-			FString SdkName = *NamePtr;
+		case ECADFormat::SOLIDWORKS:
+			SdkName = TechSoftFileParserImpl::CleanSwReferenceSdkName(SdkName);
+			break;
 
-			switch (Format)
-			{
-			case ECADFormat::CATIA_3DXML:
-				SdkName = TechSoftFileParserImpl::Clean3dxmlReferenceSdkName(SdkName);
-				break;
+		default:
+			break;
+		}
 
-			case ECADFormat::SOLIDWORKS:
-				SdkName = TechSoftFileParserImpl::CleanSwReferenceSdkName(SdkName);
-				break;
+		FString& Name = MetaData.FindOrAdd(TEXT("Name"));
+		Name = SdkName;
+		return;
+	}
+}
 
-			default:
-				break;
-			}
+void FTechSoftFileParser::BuildInstanceName(TMap<FString, FString>& MetaData)
+{
+	if (MetaData.IsEmpty())
+	{
+		return;
+	}
 
-			FString& Name = MetaData.FindOrAdd(TEXT("Name"));
-			Name = SdkName;
-			return;
-		}
+	if (TechSoftFileParserImpl::ReplaceOrAddNameValue(MetaData, TEXT("InstanceName")))
+	{
+		return;
 	}
 
-	void FTechSoftFileParser::BuildInstanceName(TMap<FString, FString>& MetaData)
+	if (TechSoftFileParserImpl::CheckIfNameExists(MetaData))
 	{
-		if (MetaData.IsEmpty())
-		{
-			return;
-		}
+		return;
+	}
 
-		if (TechSoftFileParserImpl::ReplaceOrAddNameValue(MetaData, TEXT("InstanceName")))
-		{
-			return;
-		}
+	FString* NamePtr = MetaData.Find(TEXT("SDKName"));
+	if (NamePtr != nullptr)
+	{
+		FString SdkName = *NamePtr;
 
-		if (TechSoftFileParserImpl::CheckIfNameExists(MetaData))
+		switch (Format)
 		{
-			return;
-		}
+		case ECADFormat::CATIA:
+		case ECADFormat::CATIA_3DXML:
+			SdkName = TechSoftFileParserImpl::CleanCatiaInstanceSdkName(SdkName);
+			break;
 
-		FString* NamePtr = MetaData.Find(TEXT("SDKName"));
-		if (NamePtr != nullptr)
-		{
-			FString SdkName = *NamePtr;
+		case ECADFormat::SOLIDWORKS:
+			SdkName = TechSoftFileParserImpl::CleanSwInstanceSdkName(SdkName);
+			break;
 
-			switch (Format)
-			{
-			case ECADFormat::CATIA:
-			case ECADFormat::CATIA_3DXML:
-				SdkName = TechSoftFileParserImpl::CleanCatiaInstanceSdkName(SdkName);
-				break;
+		default:
+			break;
+		}
 
-			case ECADFormat::SOLIDWORKS:
-				SdkName = TechSoftFileParserImpl::CleanSwInstanceSdkName(SdkName);
-				break;
+		FString& Name = MetaData.FindOrAdd(TEXT("Name"));
+		Name = SdkName;
+		return;
+	}
 
-			default:
-				break;
-			}
+	if (TechSoftFileParserImpl::ReplaceOrAddNameValue(MetaData, TEXT("PartNumber")))
+	{
+		return;
+	}
 
-			FString& Name = MetaData.FindOrAdd(TEXT("Name"));
-			Name = SdkName;
-			return;
-		}
+}
 
-		if (TechSoftFileParserImpl::ReplaceOrAddNameValue(MetaData, TEXT("PartNumber")))
-		{
-			return;
-		}
+void FTechSoftFileParser::BuildPartName(TMap<FString, FString>& MetaData)
+{
+	if (MetaData.IsEmpty())
+	{
+		return;
+	}
 
+	if (TechSoftFileParserImpl::CheckIfNameExists(MetaData))
+	{
+		return;
 	}
 
-	void FTechSoftFileParser::BuildPartName(TMap<FString, FString>& MetaData)
+	if (TechSoftFileParserImpl::ReplaceOrAddNameValue(MetaData, TEXT("PartNumber")))
 	{
-		if (MetaData.IsEmpty())
-		{
-			return;
-		}
+		return;
+	}
 
-		if (TechSoftFileParserImpl::CheckIfNameExists(MetaData))
-		{
-			return;
-		}
+	if (TechSoftFileParserImpl::ReplaceOrAddNameValue(MetaData, TEXT("SDKName")))
+	{
+		return;
+	}
+}
 
-		if (TechSoftFileParserImpl::ReplaceOrAddNameValue(MetaData, TEXT("PartNumber")))
-		{
-			return;
-		}
+void FTechSoftFileParser::BuildBodyName(TMap<FString, FString>& MetaData)
+{
+	if (MetaData.IsEmpty())
+	{
+		return;
+	}
 
-		if (TechSoftFileParserImpl::ReplaceOrAddNameValue(MetaData, TEXT("SDKName")))
-		{
-			return;
-		}
+	if (TechSoftFileParserImpl::CheckIfNameExists(MetaData))
+	{
+		return;
 	}
 
-	void FTechSoftFileParser::BuildBodyName(TMap<FString, FString>& MetaData)
+	FString* NamePtr = MetaData.Find(TEXT("SDKName"));
+	if (NamePtr != nullptr)
 	{
-		if (MetaData.IsEmpty())
+		FString SdkName = *NamePtr;
+		if (Format == ECADFormat::CREO)
 		{
-			return;
+			SdkName = TechSoftFileParserImpl::CleanNameByRemoving_prt(SdkName);
 		}
 
-		if (TechSoftFileParserImpl::CheckIfNameExists(MetaData))
-		{
-			return;
-		}
+		FString& Name = MetaData.FindOrAdd(TEXT("Name"));
+		Name = SdkName;
+		return;
+	}
 
-		FString* NamePtr = MetaData.Find(TEXT("SDKName"));
-		if (NamePtr != nullptr)
+	MetaData.Add(TEXT("Name"), TEXT("NoName"));
+}
+
+void FTechSoftFileParser::TraverseSpecificMetaData(const A3DAsmProductOccurrence* Occurrence, FEntityMetaData& OutMetaData)
+{
+	//----------- Export Specific information per CAD format -----------
+	switch (ModellerType)
+	{
+	case kA3DModellerSlw:
+	{
+		TUniqueTSObj<A3DAsmProductOccurrenceDataSLW> SolidWorksSpecificData(Occurrence);
+		if (SolidWorksSpecificData.IsValid())
 		{
-			FString SdkName = *NamePtr;
-			if (Format == ECADFormat::CREO)
+			if (SolidWorksSpecificData->m_psCfgName)
 			{
-				SdkName = TechSoftFileParserImpl::CleanNameByRemoving_prt(SdkName);
+				FString ConfigurationName = UTF8_TO_TCHAR(SolidWorksSpecificData->m_psCfgName);
+				OutMetaData.MetaData.Emplace(TEXT("ConfigurationName"), ConfigurationName);
+				FString ConfigurationIndex = FString::FromInt(SolidWorksSpecificData->m_iIndexCfg);
+				OutMetaData.MetaData.Emplace(TEXT("ConfigurationIndex"), ConfigurationIndex);
 			}
-
-			FString& Name = MetaData.FindOrAdd(TEXT("Name"));
-			Name = SdkName;
-			return;
 		}
-
-		MetaData.Add(TEXT("Name"), TEXT("NoName"));
+		break;
 	}
-
-	void FTechSoftFileParser::TraverseSpecificMetaData(const A3DAsmProductOccurrence* Occurrence, FEntityMetaData& OutMetaData)
+	case kA3DModellerUnigraphics:
 	{
-		//----------- Export Specific information per CAD format -----------
-		switch (ModellerType)
+		TUniqueTSObj<A3DAsmProductOccurrenceDataUg> UnigraphicsSpecificData(Occurrence);
+		if (UnigraphicsSpecificData.IsValid())
 		{
-		case kA3DModellerSlw:
-		{
-			TUniqueTSObj<A3DAsmProductOccurrenceDataSLW> SolidWorksSpecificData(Occurrence);
-			if (SolidWorksSpecificData.IsValid())
+			if (UnigraphicsSpecificData->m_psPartUID)
 			{
-				if (SolidWorksSpecificData->m_psCfgName)
-				{
-					FString ConfigurationName = UTF8_TO_TCHAR(SolidWorksSpecificData->m_psCfgName);
-					OutMetaData.MetaData.Emplace(TEXT("ConfigurationName"), ConfigurationName);
-					FString ConfigurationIndex = FString::FromInt(SolidWorksSpecificData->m_iIndexCfg);
-					OutMetaData.MetaData.Emplace(TEXT("ConfigurationIndex"), ConfigurationIndex);
-				}
+				FString PartUID = UTF8_TO_TCHAR(UnigraphicsSpecificData->m_psPartUID);
+				OutMetaData.MetaData.Emplace(TEXT("UnigraphicsPartUID"), PartUID);
 			}
-			break;
-		}
-		case kA3DModellerUnigraphics:
-		{
-			TUniqueTSObj<A3DAsmProductOccurrenceDataUg> UnigraphicsSpecificData(Occurrence);
-			if (UnigraphicsSpecificData.IsValid())
+			if (UnigraphicsSpecificData->m_psInstanceFileName)
 			{
-				if (UnigraphicsSpecificData->m_psPartUID)
-				{
-					FString PartUID = UTF8_TO_TCHAR(UnigraphicsSpecificData->m_psPartUID);
-					OutMetaData.MetaData.Emplace(TEXT("UnigraphicsPartUID"), PartUID);
-				}
-				if (UnigraphicsSpecificData->m_psInstanceFileName)
-				{
-					FString InstanceFileName = UTF8_TO_TCHAR(UnigraphicsSpecificData->m_psInstanceFileName);
-					OutMetaData.MetaData.Emplace(TEXT("UnigraphicsPartUID"), InstanceFileName);
-				}
+				FString InstanceFileName = UTF8_TO_TCHAR(UnigraphicsSpecificData->m_psInstanceFileName);
+				OutMetaData.MetaData.Emplace(TEXT("UnigraphicsPartUID"), InstanceFileName);
+			}
 
-				if (UnigraphicsSpecificData->m_uiInstanceTag)
-				{
-					FString InstanceTag = FString::FromInt(UnigraphicsSpecificData->m_uiInstanceTag);
-					OutMetaData.MetaData.Emplace(TEXT("UnigraphicsInstanceTag"), InstanceTag);
-				}
+			if (UnigraphicsSpecificData->m_uiInstanceTag)
+			{
+				FString InstanceTag = FString::FromInt(UnigraphicsSpecificData->m_uiInstanceTag);
+				OutMetaData.MetaData.Emplace(TEXT("UnigraphicsInstanceTag"), InstanceTag);
 			}
-			break;
 		}
+		break;
+	}
 
-		case kA3DModellerCatiaV5:
+	case kA3DModellerCatiaV5:
+	{
+		TUniqueTSObj<A3DAsmProductOccurrenceDataCV5> CatiaV5SpecificData(Occurrence);
+		if (CatiaV5SpecificData.IsValid())
 		{
-			TUniqueTSObj<A3DAsmProductOccurrenceDataCV5> CatiaV5SpecificData(Occurrence);
-			if (CatiaV5SpecificData.IsValid())
+			if (CatiaV5SpecificData->m_psVersion)
 			{
-				if (CatiaV5SpecificData->m_psVersion)
-				{
-					FString Version = UTF8_TO_TCHAR(CatiaV5SpecificData->m_psVersion);
-					OutMetaData.MetaData.Emplace(L"CatiaVersion", Version);
-				}
+				FString Version = UTF8_TO_TCHAR(CatiaV5SpecificData->m_psVersion);
+				OutMetaData.MetaData.Emplace(L"CatiaVersion", Version);
+			}
 
-				if (CatiaV5SpecificData->m_psPartNumber)
-				{
-					FString PartNumber = UTF8_TO_TCHAR(CatiaV5SpecificData->m_psPartNumber);
-					OutMetaData.MetaData.Emplace(L"CatiaPartNumber", PartNumber);
-				}
+			if (CatiaV5SpecificData->m_psPartNumber)
+			{
+				FString PartNumber = UTF8_TO_TCHAR(CatiaV5SpecificData->m_psPartNumber);
+				OutMetaData.MetaData.Emplace(L"CatiaPartNumber", PartNumber);
 			}
-			break;
 		}
+		break;
+	}
 
-		default:
-			break;
-		}
+	default:
+		break;
 	}
+}
 
-	FArchiveColor& FTechSoftFileParser::FindOrAddColor(uint32 ColorIndex, uint8 Alpha)
+FArchiveColor& FTechSoftFileParser::FindOrAddColor(uint32 ColorIndex, uint8 Alpha)
+{
+	uint32 ColorHId = BuildColorId(ColorIndex, Alpha);
+	if (FArchiveColor* Color = CADFileData.FindColor(ColorHId))
 	{
-		uint32 ColorHId = BuildColorId(ColorIndex, Alpha);
-		if (FArchiveColor* Color = CADFileData.FindColor(ColorHId))
-		{
-			return *Color;
-		}
+		return *Color;
+	}
+
+	FArchiveColor& NewColor = CADFileData.AddColor(ColorHId);
+	FColor& Color = NewColor.Color;
+	TUniqueTSObjFromIndex<A3DGraphRgbColorData> ColorData(ColorIndex);
+	if (ColorData.IsValid())
+	{
+		Color.R = (uint8)(ColorData->m_dRed * 255);
+		Color.G = (uint8)(ColorData->m_dGreen * 255);
+		Color.B = (uint8)(ColorData->m_dBlue * 255);
+		Color.A = Alpha;
+	}
+	else
+	{
+		Color = FColor(200, 200, 200);
+	}
+
+	NewColor.UEMaterialName = BuildColorName(NewColor.Color);
+	return NewColor;
+}
+
+FArchiveMaterial& FTechSoftFileParser::FindOrAddMaterial(uint32 MaterialIndex)
+{
+	if (FArchiveMaterial* MaterialArchive = CADFileData.FindMaterial(MaterialIndex))
+	{
+		return *MaterialArchive;
+	}
 
-		FArchiveColor& NewColor = CADFileData.AddColor(ColorHId);
-		FColor& Color = NewColor.Color;
-		TUniqueTSObjFromIndex<A3DGraphRgbColorData> ColorData(ColorIndex);
+	TUniqueTSObjFromIndex<A3DGraphRgbColorData> ColorData;
+	TFunction<const FColor(uint32)> GetColor = [&](uint32 ColorIndex) -> const FColor
+	{
+		ColorData.FillFrom(ColorIndex);
 		if (ColorData.IsValid())
 		{
-			Color.R = (uint8)(ColorData->m_dRed * 255);
-			Color.G = (uint8)(ColorData->m_dGreen * 255);
-			Color.B = (uint8)(ColorData->m_dBlue * 255);
-			Color.A = Alpha;
+			return FColor((uint8)ColorData->m_dRed * 255,
+				(uint8)ColorData->m_dGreen * 255,
+				(uint8)ColorData->m_dBlue * 255);
 		}
 		else
 		{
-			Color = FColor(200, 200, 200);
+			return FColor(200, 200, 200);
 		}
+	};
 
-		NewColor.UEMaterialName = BuildColorName(NewColor.Color);
-		return NewColor;
-	}
+	FArchiveMaterial& NewMaterial = CADFileData.AddMaterial(MaterialIndex);
+	FCADMaterial& Material = NewMaterial.Material;
 
-	FArchiveMaterial& FTechSoftFileParser::FindOrAddMaterial(uint32 MaterialIndex)
+	A3DBool bIsTexture = false;
+	A3DGlobalIsMaterialTexture(MaterialIndex, &bIsTexture);
+	if (bIsTexture)
 	{
-		if (FArchiveMaterial* MaterialArchive = CADFileData.FindMaterial(MaterialIndex))
+#ifdef NOTYETDEFINE
+		// style is a texture
+		TUniqueTSObj<A3DGraphTextureDefinitionData> TextureDefinitionData(TextureIndex);
+		if (TextureDefinitionData.IsValid())
 		{
-			return *MaterialArchive;
+			TUniqueTSObj<A3DGraphPictureData> PictureData(TextureDefinitionData->m_uiPictureIndex);
 		}
+#endif
+	}
+	else
+	{
+		TUniqueTSObjFromIndex<A3DGraphMaterialData> MaterialData(MaterialIndex);
+		Material.Diffuse = GetColor(MaterialData->m_uiDiffuse);
+		Material.Ambient = GetColor(MaterialData->m_uiAmbient);
+		Material.Specular = GetColor(MaterialData->m_uiSpecular);
+		Material.Shininess = MaterialData->m_dShininess;
+		Material.Transparency = MaterialData->m_dAmbientAlpha;
+		// todo: find how to convert Emissive color into ? reflexion coef...
+		// Material.Emissive = GetColor(MaterialData->m_uiEmissive);
+		// Material.Reflexion;
+	}
 
-		TUniqueTSObjFromIndex<A3DGraphRgbColorData> ColorData;
-		TFunction<const FColor(uint32)> GetColor = [&](uint32 ColorIndex) -> const FColor
-		{
-			ColorData.FillFrom(ColorIndex);
-			if (ColorData.IsValid())
-			{
-				return FColor((uint8)ColorData->m_dRed * 255,
-					(uint8)ColorData->m_dGreen * 255,
-					(uint8)ColorData->m_dBlue * 255);
-			}
-			else
-			{
-				return FColor(200, 200, 200);
-			}
-		};
+	return NewMaterial;
+}
 
-		FArchiveMaterial& NewMaterial = CADFileData.AddMaterial(MaterialIndex);
-		FCADMaterial& Material = NewMaterial.Material;
+void FTechSoftFileParser::TraverseGraphics(const A3DGraphics* Graphics, FEntityMetaData& OutMetaData)
+{
+	TUniqueTSObj<A3DGraphicsData> GraphicsData(Graphics);
+	if (!GraphicsData.IsValid())
+	{
+		return;
+	}
+	FEntityBehaviour GraphicsBehaviour;
 
-		A3DBool bIsTexture = false;
-		A3DGlobalIsMaterialTexture(MaterialIndex, &bIsTexture);
-		if (bIsTexture)
-		{
-#ifdef NOTYETDEFINE
-			// style is a texture
-			TUniqueTSObj<A3DGraphTextureDefinitionData> TextureDefinitionData(TextureIndex);
-			if (TextureDefinitionData.IsValid())
-			{
-				TUniqueTSObj<A3DGraphPictureData> PictureData(TextureDefinitionData->m_uiPictureIndex);
-			}
-#endif
-		}
-		else
-		{
-			TUniqueTSObjFromIndex<A3DGraphMaterialData> MaterialData(MaterialIndex);
-			Material.Diffuse = GetColor(MaterialData->m_uiDiffuse);
-			Material.Ambient = GetColor(MaterialData->m_uiAmbient);
-			Material.Specular = GetColor(MaterialData->m_uiSpecular);
-			Material.Shininess = MaterialData->m_dShininess;
-			Material.Transparency = MaterialData->m_dAmbientAlpha;
-			// todo: find how to convert Emissive color into ? reflexion coef...
-			// Material.Emissive = GetColor(MaterialData->m_uiEmissive);
-			// Material.Reflexion;
-		}
+	GraphicsBehaviour.bFatherHeritColor = GraphicsData->m_usBehaviour & kA3DGraphicsFatherHeritColor;
+	GraphicsBehaviour.bFatherHeritLayer = GraphicsData->m_usBehaviour & kA3DGraphicsFatherHeritLayer;
+	GraphicsBehaviour.bFatherHeritLinePattern = GraphicsData->m_usBehaviour & kA3DGraphicsFatherHeritLinePattern;
+	GraphicsBehaviour.bFatherHeritLineWidth = GraphicsData->m_usBehaviour & kA3DGraphicsFatherHeritLineWidth;
+	GraphicsBehaviour.bFatherHeritShow = GraphicsData->m_usBehaviour & kA3DGraphicsFatherHeritShow;
+	GraphicsBehaviour.bFatherHeritTransparency = GraphicsData->m_usBehaviour & kA3DGraphicsFatherHeritTransparency;
+	GraphicsBehaviour.bRemoved = GraphicsData->m_usBehaviour & kA3DGraphicsRemoved;
+	GraphicsBehaviour.bShow = GraphicsData->m_usBehaviour & kA3DGraphicsShow;
+	GraphicsBehaviour.bSonHeritColor = GraphicsData->m_usBehaviour & kA3DGraphicsSonHeritColor;
+	GraphicsBehaviour.bSonHeritLayer = GraphicsData->m_usBehaviour & kA3DGraphicsSonHeritLayer;
+	GraphicsBehaviour.bSonHeritLinePattern = GraphicsData->m_usBehaviour & kA3DGraphicsSonHeritLinePattern;
+	GraphicsBehaviour.bSonHeritLineWidth = GraphicsData->m_usBehaviour & kA3DGraphicsSonHeritLineWidth;
+	GraphicsBehaviour.bSonHeritShow = GraphicsData->m_usBehaviour & kA3DGraphicsSonHeritShow;
+	GraphicsBehaviour.bSonHeritTransparency = GraphicsData->m_usBehaviour & kA3DGraphicsSonHeritTransparency;
 
-		return NewMaterial;
-	}
+	OutMetaData.bRemoved = GraphicsBehaviour.bRemoved;
+	OutMetaData.bShow = GraphicsBehaviour.bShow;
 
-	void FTechSoftFileParser::TraverseGraphics(const A3DGraphics* Graphics, FEntityMetaData& OutMetaData)
-	{
-		TUniqueTSObj<A3DGraphicsData> GraphicsData(Graphics);
-		if (!GraphicsData.IsValid())
-		{
-			return;
-		}
-		FEntityBehaviour GraphicsBehaviour;
-
-		GraphicsBehaviour.bFatherHeritColor = GraphicsData->m_usBehaviour & kA3DGraphicsFatherHeritColor;
-		GraphicsBehaviour.bFatherHeritLayer = GraphicsData->m_usBehaviour & kA3DGraphicsFatherHeritLayer;
-		GraphicsBehaviour.bFatherHeritLinePattern = GraphicsData->m_usBehaviour & kA3DGraphicsFatherHeritLinePattern;
-		GraphicsBehaviour.bFatherHeritLineWidth = GraphicsData->m_usBehaviour & kA3DGraphicsFatherHeritLineWidth;
-		GraphicsBehaviour.bFatherHeritShow = GraphicsData->m_usBehaviour & kA3DGraphicsFatherHeritShow;
-		GraphicsBehaviour.bFatherHeritTransparency = GraphicsData->m_usBehaviour & kA3DGraphicsFatherHeritTransparency;
-		GraphicsBehaviour.bRemoved = GraphicsData->m_usBehaviour & kA3DGraphicsRemoved;
-		GraphicsBehaviour.bShow = GraphicsData->m_usBehaviour & kA3DGraphicsShow;
-		GraphicsBehaviour.bSonHeritColor = GraphicsData->m_usBehaviour & kA3DGraphicsSonHeritColor;
-		GraphicsBehaviour.bSonHeritLayer = GraphicsData->m_usBehaviour & kA3DGraphicsSonHeritLayer;
-		GraphicsBehaviour.bSonHeritLinePattern = GraphicsData->m_usBehaviour & kA3DGraphicsSonHeritLinePattern;
-		GraphicsBehaviour.bSonHeritLineWidth = GraphicsData->m_usBehaviour & kA3DGraphicsSonHeritLineWidth;
-		GraphicsBehaviour.bSonHeritShow = GraphicsData->m_usBehaviour & kA3DGraphicsSonHeritShow;
-		GraphicsBehaviour.bSonHeritTransparency = GraphicsData->m_usBehaviour & kA3DGraphicsSonHeritTransparency;
-
-		OutMetaData.bRemoved = GraphicsBehaviour.bRemoved;
-		OutMetaData.bShow = GraphicsBehaviour.bShow;
-
-		FCADUUID ColorName;
-		FCADUUID MaterialName;
-
-		TraverseGraphStyleData(GraphicsData->m_uiStyleIndex, ColorName, MaterialName);
-
-		if (ColorName)
-		{
-			OutMetaData.MetaData.Add(TEXT("ColorName"), FString::Printf(TEXT("%u"), ColorName));
-		}
+	FCADUUID ColorName;
+	FCADUUID MaterialName;
 
-		if (MaterialName)
-		{
-			OutMetaData.MetaData.Add(TEXT("MaterialName"), FString::Printf(TEXT("%u"), MaterialName));
-		}
+	TraverseGraphStyleData(GraphicsData->m_uiStyleIndex, ColorName, MaterialName);
+
+	if (ColorName)
+	{
+		OutMetaData.MetaData.Add(TEXT("ColorName"), FString::Printf(TEXT("%u"), ColorName));
 	}
 
-	void FTechSoftFileParser::TraverseGraphStyleData(uint32 StyleIndex, FCADUUID& OutColorName, FCADUUID& OutMaterialName)
+	if (MaterialName)
 	{
-		OutColorName = 0;
-		OutMaterialName = 0;
+		OutMetaData.MetaData.Add(TEXT("MaterialName"), FString::Printf(TEXT("%u"), MaterialName));
+	}
+}
+
+void FTechSoftFileParser::TraverseGraphStyleData(uint32 StyleIndex, FCADUUID& OutColorName, FCADUUID& OutMaterialName)
+{
+	OutColorName = 0;
+	OutMaterialName = 0;
 
-		TUniqueTSObjFromIndex<A3DGraphStyleData> GraphStyleData(StyleIndex);
+	TUniqueTSObjFromIndex<A3DGraphStyleData> GraphStyleData(StyleIndex);
 
-		if (GraphStyleData.IsValid())
+	if (GraphStyleData.IsValid())
+	{
+		if (GraphStyleData->m_bMaterial)
+		{
+			FArchiveMaterial& MaterialArchive = FindOrAddMaterial(GraphStyleData->m_uiRgbColorIndex);
+			OutMaterialName = MaterialArchive.UEMaterialName;
+		}
+		else
 		{
-			if (GraphStyleData->m_bMaterial)
+			uint8 Alpha = 255;
+			if (GraphStyleData->m_bIsTransparencyDefined)
 			{
-				FArchiveMaterial& MaterialArchive = FindOrAddMaterial(GraphStyleData->m_uiRgbColorIndex);
-				OutMaterialName = MaterialArchive.UEMaterialName;
+				Alpha = GraphStyleData->m_ucTransparency;
 			}
-			else
-			{
-				uint8 Alpha = 255;
-				if (GraphStyleData->m_bIsTransparencyDefined)
-				{
-					Alpha = GraphStyleData->m_ucTransparency;
-				}
 
-				FArchiveColor& ColorArchive = FindOrAddColor(GraphStyleData->m_uiRgbColorIndex, Alpha);
-				OutColorName = ColorArchive.UEMaterialName;
-			}
+			FArchiveColor& ColorArchive = FindOrAddColor(GraphStyleData->m_uiRgbColorIndex, Alpha);
+			OutColorName = ColorArchive.UEMaterialName;
 		}
 	}
+}
 
-	void FTechSoftFileParser::TraverseMaterialProperties(const A3DEntity* Entity)
+void FTechSoftFileParser::TraverseMaterialProperties(const A3DEntity* Entity)
+{
+	TUniqueTSObj<A3DMiscMaterialPropertiesData> MaterialPropertiesData(Entity);
+	if (!MaterialPropertiesData.IsValid())
 	{
-		TUniqueTSObj<A3DMiscMaterialPropertiesData> MaterialPropertiesData(Entity);
-		if (!MaterialPropertiesData.IsValid())
-		{
-			return;
-		}
+		return;
+	}
 
-		//xml->SetDoubleAttribute("m_dDensity", MaterialPropertiesData.m_dDensity);
-		//xml->SetAttribute("m_pcMaterialName", MaterialPropertiesData.m_pcMaterialName ? MaterialPropertiesData.m_pcMaterialName : "NULL");
-		//xml->SetAttribute("m_ePhysicType", MaterialPropertiesData.m_ePhysicType);
-		switch (MaterialPropertiesData->m_ePhysicType)
-		{
-		case A3DPhysicType_None:
-		case A3DPhysicType_Fiber:
-		case A3DPhysicType_HoneyComb:
-		case A3DPhysicType_Isotropic:
-		case A3DPhysicType_Orthotropic2D:
-		case A3DPhysicType_Orthotropic3D:
-		case A3DPhysicType_Anisotropic:
+	//xml->SetDoubleAttribute("m_dDensity", MaterialPropertiesData.m_dDensity);
+	//xml->SetAttribute("m_pcMaterialName", MaterialPropertiesData.m_pcMaterialName ? MaterialPropertiesData.m_pcMaterialName : "NULL");
+	//xml->SetAttribute("m_ePhysicType", MaterialPropertiesData.m_ePhysicType);
+	switch (MaterialPropertiesData->m_ePhysicType)
+	{
+	case A3DPhysicType_None:
+	case A3DPhysicType_Fiber:
+	case A3DPhysicType_HoneyComb:
+	case A3DPhysicType_Isotropic:
+	case A3DPhysicType_Orthotropic2D:
+	case A3DPhysicType_Orthotropic3D:
+	case A3DPhysicType_Anisotropic:
 
-		default:
-			break;
-		}
+	default:
+		break;
 	}
+}
 
-	FMatrix FTechSoftFileParser::TraverseTransformation3D(const A3DMiscCartesianTransformation* CartesianTransformation)
+FMatrix FTechSoftFileParser::TraverseTransformation3D(const A3DMiscCartesianTransformation* CartesianTransformation)
+{
+	if (CartesianTransformation)
 	{
-		if (CartesianTransformation)
-		{
-			TUniqueTSObj<A3DMiscCartesianTransformationData> CartesianTransformationData(CartesianTransformation);
+		TUniqueTSObj<A3DMiscCartesianTransformationData> CartesianTransformationData(CartesianTransformation);
 
-			if (CartesianTransformationData.IsValid())
-			{
-				FVector Origin(CartesianTransformationData->m_sOrigin.m_dX, CartesianTransformationData->m_sOrigin.m_dY, CartesianTransformationData->m_sOrigin.m_dZ);
-				FVector XVector(CartesianTransformationData->m_sXVector.m_dX, CartesianTransformationData->m_sXVector.m_dY, CartesianTransformationData->m_sXVector.m_dZ);;
-				FVector YVector(CartesianTransformationData->m_sYVector.m_dX, CartesianTransformationData->m_sYVector.m_dY, CartesianTransformationData->m_sYVector.m_dZ);;
+		if (CartesianTransformationData.IsValid())
+		{
+			FVector Origin(CartesianTransformationData->m_sOrigin.m_dX, CartesianTransformationData->m_sOrigin.m_dY, CartesianTransformationData->m_sOrigin.m_dZ);
+			FVector XVector(CartesianTransformationData->m_sXVector.m_dX, CartesianTransformationData->m_sXVector.m_dY, CartesianTransformationData->m_sXVector.m_dZ);;
+			FVector YVector(CartesianTransformationData->m_sYVector.m_dX, CartesianTransformationData->m_sYVector.m_dY, CartesianTransformationData->m_sYVector.m_dZ);;
 
-				FVector ZVector = XVector ^ YVector;
-				FMatrix Matrix(XVector, YVector, ZVector, FVector::Zero());
+			FVector ZVector = XVector ^ YVector;
+			FMatrix Matrix(XVector, YVector, ZVector, FVector::Zero());
 
-				FMatrix Scale = FMatrix::Identity;
-				Scale.M[0][0] = CartesianTransformationData->m_sScale.m_dX;
-				Scale.M[1][1] = CartesianTransformationData->m_sScale.m_dY;
-				Scale.M[2][2] = CartesianTransformationData->m_sScale.m_dZ;
-				Matrix = Matrix * Scale;
+			FMatrix Scale = FMatrix::Identity;
+			Scale.M[0][0] = CartesianTransformationData->m_sScale.m_dX;
+			Scale.M[1][1] = CartesianTransformationData->m_sScale.m_dY;
+			Scale.M[2][2] = CartesianTransformationData->m_sScale.m_dZ;
+			Matrix = Matrix * Scale;
 
-				Matrix = Matrix.ConcatTranslation(Origin);
+			Matrix = Matrix.ConcatTranslation(Origin);
 
-				return Matrix;
-			}
-		}
-		return FMatrix::Identity;
-	}
-
-	FMatrix FTechSoftFileParser::TraverseGeneralTransformation(const A3DMiscGeneralTransformation* GeneralTransformation)
-	{
-		TUniqueTSObj<A3DMiscGeneralTransformationData> GeneralTransformationData(GeneralTransformation);
-		if (GeneralTransformationData.IsValid())
-		{
-			FMatrix Matrix;
-			int32 Index = 0;
-			for (int32 Andex = 0; Andex < 4; ++Andex)
-			{
-				for (int32 Bndex = 0; Bndex < 4; ++Bndex, ++Index)
-				{
-					Matrix.M[Andex][Bndex] = GeneralTransformationData->m_adCoeff[Index];
-				}
-			}
 			return Matrix;
 		}
-		return FMatrix::Identity;
 	}
+	return FMatrix::Identity;
+}
 
-	FMatrix FTechSoftFileParser::TraverseTransformation(const A3DMiscCartesianTransformation* Transformation3D)
+FMatrix FTechSoftFileParser::TraverseGeneralTransformation(const A3DMiscGeneralTransformation* GeneralTransformation)
+{
+	TUniqueTSObj<A3DMiscGeneralTransformationData> GeneralTransformationData(GeneralTransformation);
+	if (GeneralTransformationData.IsValid())
 	{
-		if (Transformation3D == NULL)
-		{
-			return FMatrix::Identity;
-		}
-
-		A3DEEntityType Type = kA3DTypeUnknown;
-		A3DEntityGetType(Transformation3D, &Type);
-
-		if (Type == kA3DTypeMiscCartesianTransformation)
-		{
-			return TraverseTransformation3D(Transformation3D);
-		}
-		else if (Type == kA3DTypeMiscGeneralTransformation)
+		FMatrix Matrix;
+		int32 Index = 0;
+		for (int32 Andex = 0; Andex < 4; ++Andex)
 		{
-			return TraverseGeneralTransformation(Transformation3D);
+			for (int32 Bndex = 0; Bndex < 4; ++Bndex, ++Index)
+			{
+				Matrix.M[Andex][Bndex] = GeneralTransformationData->m_adCoeff[Index];
+			}
 		}
-		return FMatrix::Identity;
+		return Matrix;
 	}
+	return FMatrix::Identity;
+}
 
-	FMatrix FTechSoftFileParser::TraverseCoordinateSystem(const A3DRiCoordinateSystem* CoordinateSystem)
+FMatrix FTechSoftFileParser::TraverseTransformation(const A3DMiscCartesianTransformation* Transformation3D)
+{
+	if (Transformation3D == NULL)
 	{
-		TUniqueTSObj<A3DRiCoordinateSystemData> CoordinateSystemData(CoordinateSystem);
-		if (CoordinateSystemData.IsValid())
-		{
-			return TraverseTransformation3D(CoordinateSystemData->m_pTransformation);
-		}
 		return FMatrix::Identity;
 	}
 
-	void FTechSoftFileParser::TraverseLayer(const A3DAsmProductOccurrence* Occurrence)
+	A3DEEntityType Type = kA3DTypeUnknown;
+	A3DEntityGetType(Transformation3D, &Type);
+
+	if (Type == kA3DTypeMiscCartesianTransformation)
+	{
+		return TraverseTransformation3D(Transformation3D);
+	}
+	else if (Type == kA3DTypeMiscGeneralTransformation)
+	{
+		return TraverseGeneralTransformation(Transformation3D);
+	}
+	return FMatrix::Identity;
+}
+
+FMatrix FTechSoftFileParser::TraverseCoordinateSystem(const A3DRiCoordinateSystem* CoordinateSystem)
+{
+	TUniqueTSObj<A3DRiCoordinateSystemData> CoordinateSystemData(CoordinateSystem);
+	if (CoordinateSystemData.IsValid())
 	{
-		// TODO
+		return TraverseTransformation3D(CoordinateSystemData->m_pTransformation);
+	}
+	return FMatrix::Identity;
+}
+
+void FTechSoftFileParser::TraverseLayer(const A3DAsmProductOccurrence* Occurrence)
+{
+	// TODO
 
-		A3DUns32 LayerCount = 0;
-		A3DAsmLayer* AsmLayer = 0;
-		if (A3DAsmProductOccurrenceGetLayerList(Occurrence, &LayerCount, &AsmLayer) == A3D_SUCCESS)
+	A3DUns32 LayerCount = 0;
+	A3DAsmLayer* AsmLayer = 0;
+	if (A3DAsmProductOccurrenceGetLayerList(Occurrence, &LayerCount, &AsmLayer) == A3D_SUCCESS)
+	{
+		if (LayerCount)
 		{
-			if (LayerCount)
+			for (A3DUns32 Index = 0; Index < LayerCount; ++Index)
 			{
-				for (A3DUns32 Index = 0; Index < LayerCount; ++Index)
-				{
-					//Layer->SetAttribute("Name", asLayers[Index].m_pcLayerName ? asLayers[Index].m_pcLayerName : "null");
-					//Layer->SetAttribute("Layer", asLayers[Index].m_usLayer);
-				}
+				//Layer->SetAttribute("Name", asLayers[Index].m_pcLayerName ? asLayers[Index].m_pcLayerName : "null");
+				//Layer->SetAttribute("Layer", asLayers[Index].m_usLayer);
 			}
-			A3DAsmProductOccurrenceGetLayerList(NULL, &LayerCount, &AsmLayer);
 		}
+		A3DAsmProductOccurrenceGetLayerList(NULL, &LayerCount, &AsmLayer);
 	}
+}
 
-	bool FTechSoftFileParser::IsConfigurationSet(const A3DAsmProductOccurrence* Occurrence)
+bool FTechSoftFileParser::IsConfigurationSet(const A3DAsmProductOccurrence* Occurrence)
+{
+	TUniqueTSObj<A3DAsmProductOccurrenceData> OccurrenceData(Occurrence);
+	if (!OccurrenceData.IsValid())
 	{
-		TUniqueTSObj<A3DAsmProductOccurrenceData> OccurrenceData(Occurrence);
-		if (!OccurrenceData.IsValid())
-		{
-			return false;
-		}
+		return false;
+	}
 
-		bool bIsConfiguration = false;
-		if (OccurrenceData->m_uiPOccurrencesSize)
+	bool bIsConfiguration = false;
+	if (OccurrenceData->m_uiPOccurrencesSize)
+	{
+		TUniqueTSObj<A3DAsmProductOccurrenceData> ChildData;
+		for (uint32 Index = 0; Index < OccurrenceData->m_uiPOccurrencesSize; ++Index)
 		{
-			TUniqueTSObj<A3DAsmProductOccurrenceData> ChildData;
-			for (uint32 Index = 0; Index < OccurrenceData->m_uiPOccurrencesSize; ++Index)
+			if (ChildData.FillFrom(OccurrenceData->m_ppPOccurrences[Index]) == A3D_SUCCESS)
 			{
-				if (ChildData.FillFrom(OccurrenceData->m_ppPOccurrences[Index]) == A3D_SUCCESS)
+				if (ChildData->m_uiProductFlags & A3D_PRODUCT_FLAG_CONFIG)
 				{
-					if (ChildData->m_uiProductFlags & A3D_PRODUCT_FLAG_CONFIG)
-					{
-						bIsConfiguration = true;
-					}
-					break;
+					bIsConfiguration = true;
 				}
+				break;
 			}
 		}
-		return bIsConfiguration;
 	}
+	return bIsConfiguration;
+}
 
-	void FTechSoftFileParser::CountUnderModel(const A3DAsmModelFile* AsmModel)
+void FTechSoftFileParser::CountUnderModel(const A3DAsmModelFile* AsmModel)
+{
+	TUniqueTSObj<A3DAsmModelFileData> ModelFileData(AsmModel);
+	if (!ModelFileData.IsValid())
 	{
-		TUniqueTSObj<A3DAsmModelFileData> ModelFileData(AsmModel);
-		if (!ModelFileData.IsValid())
-		{
-			return;
-		}
+		return;
+	}
 
-		ComponentCount[EComponentType::Occurrence] ++;
+	ComponentCount[EComponentType::Occurrence] ++;
 
-		for (uint32 Index = 0; Index < ModelFileData->m_uiPOccurrencesSize; ++Index)
+	for (uint32 Index = 0; Index < ModelFileData->m_uiPOccurrencesSize; ++Index)
+	{
+		//if (IsConfigurationSet(ModelFileData.m_ppPOccurrences[Index]))
+		//{
+		//	TraverseConfigurationSet(ModelFileData.m_ppPOccurrences[Index], MetaData);
+		//}
+		//else
 		{
-			//if (IsConfigurationSet(ModelFileData.m_ppPOccurrences[Index]))
-			//{
-			//	TraverseConfigurationSet(ModelFileData.m_ppPOccurrences[Index], MetaData);
-			//}
-			//else
-			{
-				CountUnderOccurrence(ModelFileData->m_ppPOccurrences[Index]);
-			}
+			CountUnderOccurrence(ModelFileData->m_ppPOccurrences[Index]);
 		}
 	}
+}
 
-	void FTechSoftFileParser::CountUnderOccurrence(const A3DAsmProductOccurrence* Occurrence)
+void FTechSoftFileParser::CountUnderOccurrence(const A3DAsmProductOccurrence* Occurrence)
+{
+	TUniqueTSObj<A3DAsmProductOccurrenceData> OccurrenceData(Occurrence);
+	if (!OccurrenceData.IsValid())
 	{
-		TUniqueTSObj<A3DAsmProductOccurrenceData> OccurrenceData(Occurrence);
-		if (!OccurrenceData.IsValid())
-		{
-			return;
-		}
-
-		ComponentCount[EComponentType::Occurrence] ++;
-
-		if (OccurrenceData->m_pPrototype)
-		{
-			CountUnderPrototype(OccurrenceData->m_pPrototype);
-		}
+		return;
+	}
 
-		for (uint32 Index = 0; Index < OccurrenceData->m_uiPOccurrencesSize; ++Index)
-		{
-			CountUnderOccurrence(OccurrenceData->m_ppPOccurrences[Index]);
-		}
+	ComponentCount[EComponentType::Occurrence] ++;
 
-		if (OccurrenceData->m_pPart)
-		{
-			CountUnderPartDefinition(OccurrenceData->m_pPart);
-		}
+	if (OccurrenceData->m_pPrototype)
+	{
+		CountUnderPrototype(OccurrenceData->m_pPrototype);
 	}
 
-	void FTechSoftFileParser::CountUnderPrototype(const A3DAsmProductOccurrence* Prototype)
+	for (uint32 Index = 0; Index < OccurrenceData->m_uiPOccurrencesSize; ++Index)
 	{
-		if (PrototypeCounted.Contains(Prototype))
-		{
-			return;
-		}
-		PrototypeCounted.Add(Prototype);
-
-		TUniqueTSObj<A3DAsmProductOccurrenceData> PrototypeData(Prototype);
-		if (!PrototypeData.IsValid())
-		{
-			return;
-		}
+		CountUnderOccurrence(OccurrenceData->m_ppPOccurrences[Index]);
+	}
 
-		ComponentCount[EComponentType::Reference] ++;
+	if (OccurrenceData->m_pPart)
+	{
+		CountUnderPartDefinition(OccurrenceData->m_pPart);
+	}
+}
 
-		if (PrototypeData->m_pPrototype)
-		{
-			CountUnderPrototype(PrototypeData->m_pPrototype);
-		}
+void FTechSoftFileParser::CountUnderPrototype(const A3DAsmProductOccurrence* Prototype)
+{
+	if (PrototypeCounted.Contains(Prototype))
+	{
+		return;
 	}
+	PrototypeCounted.Add(Prototype);
 
-	void FTechSoftFileParser::CountUnderPartDefinition(const A3DAsmPartDefinition* PartDefinition)
+	TUniqueTSObj<A3DAsmProductOccurrenceData> PrototypeData(Prototype);
+	if (!PrototypeData.IsValid())
 	{
-		TUniqueTSObj<A3DAsmPartDefinitionData> PartData(PartDefinition);
-		if (!PartData.IsValid())
-		{
-			return;
-		}
+		return;
+	}
 
-		for (unsigned int Index = 0; Index < PartData->m_uiRepItemsSize; ++Index)
-		{
-			CountUnderRepresentation(PartData->m_ppRepItems[Index]);
-		}
+	ComponentCount[EComponentType::Reference] ++;
 
+	if (PrototypeData->m_pPrototype)
+	{
+		CountUnderPrototype(PrototypeData->m_pPrototype);
 	}
+}
 
-	void FTechSoftFileParser::CountUnderRepresentation(const A3DRiRepresentationItem* RepresentationItem)
+void FTechSoftFileParser::CountUnderPartDefinition(const A3DAsmPartDefinition* PartDefinition)
+{
+	TUniqueTSObj<A3DAsmPartDefinitionData> PartData(PartDefinition);
+	if (!PartData.IsValid())
 	{
-		A3DEEntityType Type;
-		A3DEntityGetType(RepresentationItem, &Type);
-
-		switch (Type)
-		{
-		case kA3DTypeRiSet:
-			CountUnderRepresentationSet(RepresentationItem);
-			break;
-		case kA3DTypeRiBrepModel:
-		case kA3DTypeRiPolyBrepModel:
-			ComponentCount[EComponentType::Body] ++;
-			break;
-
-		default:
-			break;
-		}
+		return;
 	}
 
-	void FTechSoftFileParser::CountUnderRepresentationSet(const A3DRiSet* RepresentationSet)
+	for (unsigned int Index = 0; Index < PartData->m_uiRepItemsSize; ++Index)
 	{
-		TUniqueTSObj<A3DRiSetData> RepresentationSetData(RepresentationSet);
-		if (!RepresentationSetData.IsValid())
-		{
-			return;
-		}
-		ComponentCount[EComponentType::Body] += RepresentationSetData->m_uiRepItemsSize;
+		CountUnderRepresentation(PartData->m_ppRepItems[Index]);
 	}
 
-	int32 FTechSoftFileParser::CountMaterial()
+}
+
+void FTechSoftFileParser::CountUnderRepresentation(const A3DRiRepresentationItem* RepresentationItem)
+{
+	A3DEEntityType Type;
+	A3DEntityGetType(RepresentationItem, &Type);
+
+	switch (Type)
 	{
-		// TODO
+	case kA3DTypeRiSet:
+		CountUnderRepresentationSet(RepresentationItem);
+		break;
+	case kA3DTypeRiBrepModel:
+	case kA3DTypeRiPolyBrepModel:
+		ComponentCount[EComponentType::Body] ++;
+		break;
 
-		return 0;
+	default:
+		break;
 	}
+}
 
-	int32 FTechSoftFileParser::CountColor()
+void FTechSoftFileParser::CountUnderRepresentationSet(const A3DRiSet* RepresentationSet)
+{
+	TUniqueTSObj<A3DRiSetData> RepresentationSetData(RepresentationSet);
+	if (!RepresentationSetData.IsValid())
 	{
-		// TODO
-
-		return 0;
+		return;
 	}
+	ComponentCount[EComponentType::Body] += RepresentationSetData->m_uiRepItemsSize;
+}
+
+int32 FTechSoftFileParser::CountMaterial()
+{
+	// TODO
+
+	return 0;
+}
 
-	void FTechSoftFileParser::ReadMaterialsAndColors()
+int32 FTechSoftFileParser::CountColor()
+{
+	// TODO
+
+	return 0;
+}
+
+void FTechSoftFileParser::ReadMaterialsAndColors()
+{
+	// TODO
+}
+
+void FTechSoftFileParser::MeshRepresentationWithTechSoft(A3DRiRepresentationItem* RepresentationItemPtr, FArchiveBody& Body)
+{
+	TUniqueTSObj<A3DRiRepresentationItemData> RepresentationItemData;
+
+	// TUniqueTechSoftObj does not work in this case
+	A3DRWParamsTessellationData TessellationParameters;
+	A3D_INITIALIZE_DATA(A3DRWParamsTessellationData, TessellationParameters);
+
+	if (true)
+	{
+		TessellationParameters.m_eTessellationLevelOfDetail = kA3DTessLODMedium;		/*!< Enum to specify predefined values for some following members. */
+	}
+	else
 	{
 		// TODO
+		// Check unity conversion
+		ensure(false);
+
+		const FImportParameters& ImportParameters = CADFileData.GetImportParameters();
+		////A3DRWParamsTessellationData /*!< The tessellation reading parameters. */
+		TessellationParameters.m_eTessellationLevelOfDetail = kA3DTessLODUserDefined;		/*!< Enum to specify predefined values for some following members. */
+		TessellationParameters.m_bUseHeightInsteadOfRatio = A3D_TRUE;
+		TessellationParameters.m_dMaxChordHeight = ImportParameters.GetChordTolerance();
+		TessellationParameters.m_dAngleToleranceDeg = ImportParameters.GetMaxNormalAngle();
+		TessellationParameters.m_dMaximalTriangleEdgeLength = 0; //ImportParameters.MaxEdgeLength;
+
+		TessellationParameters.m_bAccurateTessellation = A3D_FALSE;  // A3D_FALSE' indicates the tessellation is set for visualization
+		TessellationParameters.m_bAccurateTessellationWithGrid = A3D_FALSE; /*!< Enable accurate tessellation with faces inner points on a grid.*/
+		TessellationParameters.m_dAccurateTessellationWithGridMaximumStitchLength = 0; 	/*!< Maximal grid stitch length. Disabled if value is 0. Be careful, a too small value can generate a huge tessellation. */
 	}
 
-	void FTechSoftFileParser::MeshRepresentationWithTechSoft(A3DRiRepresentationItem* RepresentationItemPtr, FArchiveBody& Body)
-	{
-		TUniqueTSObj<A3DRiRepresentationItemData> RepresentationItemData;
+	TessellationParameters.m_bKeepUVPoints = A3D_TRUE; /*!< Keep parametric points as texture points. */
 
-		// TUniqueTechSoftObj does not work in this case
-		A3DRWParamsTessellationData TessellationParameters;
-		A3D_INITIALIZE_DATA(A3DRWParamsTessellationData, TessellationParameters);
+	// Get the tessellation
+	A3DStatus Status = A3DRiRepresentationItemComputeTessellation(RepresentationItemPtr, &TessellationParameters);
+	Status = A3DRiRepresentationItemGet(RepresentationItemPtr, RepresentationItemData.GetEmptyDataPtr());
+	TraverseTessellationBase(RepresentationItemData->m_pTessBase, Body);
+}
 
-		if (true)
-		{
-			TessellationParameters.m_eTessellationLevelOfDetail = kA3DTessLODMedium;		/*!< Enum to specify predefined values for some following members. */
-		}
-		else
+void FTechSoftFileParser::TraverseTessellationBase(const A3DTessBase* Tessellation, FArchiveBody& Body)
+{
+	A3DEEntityType eType;
+	if (A3DEntityGetType(Tessellation, &eType) == A3D_SUCCESS)
+	{
+		switch (eType)
 		{
-			// TODO
-			// Check unity conversion
-			ensure(false);
-
-			const FImportParameters& ImportParameters = CADFileData.GetImportParameters();
-			////A3DRWParamsTessellationData /*!< The tessellation reading parameters. */
-			TessellationParameters.m_eTessellationLevelOfDetail = kA3DTessLODUserDefined;		/*!< Enum to specify predefined values for some following members. */
-			TessellationParameters.m_bUseHeightInsteadOfRatio = A3D_TRUE;
-			TessellationParameters.m_dMaxChordHeight = ImportParameters.GetChordTolerance();
-			TessellationParameters.m_dAngleToleranceDeg = ImportParameters.GetMaxNormalAngle();
-			TessellationParameters.m_dMaximalTriangleEdgeLength = 0; //ImportParameters.MaxEdgeLength;
-
-			TessellationParameters.m_bAccurateTessellation = A3D_FALSE;  // A3D_FALSE' indicates the tessellation is set for visualization
-			TessellationParameters.m_bAccurateTessellationWithGrid = A3D_FALSE; /*!< Enable accurate tessellation with faces inner points on a grid.*/
-			TessellationParameters.m_dAccurateTessellationWithGridMaximumStitchLength = 0; 	/*!< Maximal grid stitch length. Disabled if value is 0. Be careful, a too small value can generate a huge tessellation. */
+		case kA3DTypeTess3D:
+			TraverseTessellation3D(Tessellation, Body);
+			break;
+		case kA3DTypeTess3DWire:
+		case kA3DTypeTessMarkup:
+		default:
+			break;
 		}
+	}
+}
 
-		TessellationParameters.m_bKeepUVPoints = A3D_TRUE; /*!< Keep parametric points as texture points. */
+void AddFaceTriangleWithUniqueNormal(FTessellationData& Tessellation, const A3DTess3DData& Tessellation3DData, long& InOutStartIndex, unsigned long InTrinangleCount)
+{
+	TechSoftFileParserImpl::Reserve(Tessellation, InTrinangleCount, /*bWithTexture*/ false);
 
-		// Get the tessellation
-		A3DStatus Status = A3DRiRepresentationItemComputeTessellation(RepresentationItemPtr, &TessellationParameters);
-		Status = A3DRiRepresentationItemGet(RepresentationItemPtr, RepresentationItemData.GetEmptyDataPtr());
-		TraverseTessellationBase(RepresentationItemData->m_pTessBase, Body);
-	}
+	unsigned long Index = InOutStartIndex;
+
+	int32 FaceIndex[3] = { 0, 0, 0 };
+	int32 NormalIndex[3] = { 0, 0, 0 };
+	int32 VertexIndex = 0;
 
-	void FTechSoftFileParser::TraverseTessellationBase(const A3DTessBase* Tessellation, FArchiveBody& Body)
+	NormalIndex[0] = Tessellation3DData.m_puiTriangulatedIndexes[Index++];
+	NormalIndex[1] = NormalIndex[0];
+	NormalIndex[2] = NormalIndex[0];
+
+	// Get Triangles
+	for (unsigned long TriangleIndex = 0; TriangleIndex < InTrinangleCount; TriangleIndex++)
 	{
-		A3DEEntityType eType;
-		if (A3DEntityGetType(Tessellation, &eType) == A3D_SUCCESS)
+		FaceIndex[0] = Tessellation3DData.m_puiTriangulatedIndexes[Index++];
+		FaceIndex[1] = Tessellation3DData.m_puiTriangulatedIndexes[Index++];
+		FaceIndex[2] = Tessellation3DData.m_puiTriangulatedIndexes[Index++];
+
+		if (!TechSoftFileParserImpl::AddFace(FaceIndex, Tessellation, VertexIndex))
 		{
-			switch (eType)
-			{
-			case kA3DTypeTess3D:
-				TraverseTessellation3D(Tessellation, Body);
-				break;
-			case kA3DTypeTess3DWire:
-			case kA3DTypeTessMarkup:
-			default:
-				break;
-			}
+			continue;
 		}
+
+		TechSoftFileParserImpl::AddNormals(Tessellation3DData.m_pdNormals, NormalIndex, Tessellation.NormalArray);
 	}
 
-	void AddFaceTriangleWithUniqueNormal(FTessellationData& Tessellation, const A3DTess3DData& Tessellation3DData, long& InOutStartIndex, unsigned long InTrinangleCount)
-	{
-		TechSoftFileParserImpl::Reserve(Tessellation, InTrinangleCount, /*bWithTexture*/ false);
+	InOutStartIndex = Index;
+}
 
-		unsigned long Index = InOutStartIndex;
+void AddFaceTriangleWithUniqueNormalAndTexture(FTessellationData& Tessellation, const A3DTess3DData& Tessellation3DData, long& InOutStartIndex, unsigned long InTrinangleCount, long TextureCount)
+{
+	TechSoftFileParserImpl::Reserve(Tessellation, InTrinangleCount, /*bWithTexture*/ true);
 
-		int32 FaceIndex[3] = { 0, 0, 0 };
-		int32 NormalIndex[3] = { 0, 0, 0 };
-		int32 VertexIndex = 0;
+	unsigned long Index = InOutStartIndex;
 
-		NormalIndex[0] = Tessellation3DData.m_puiTriangulatedIndexes[Index++];
-		NormalIndex[1] = NormalIndex[0];
-		NormalIndex[2] = NormalIndex[0];
+	int32 FaceIndex[3] = { 0, 0, 0 };
+	int32 NormalIndex[3] = { 0, 0, 0 };
+	int32 TextureIndex[3] = { 0, 0, 0 };
+	int32 VertexIndex = 0;
 
-		// Get Triangles
-		for (unsigned long TriangleIndex = 0; TriangleIndex < InTrinangleCount; TriangleIndex++)
-		{
-			FaceIndex[0] = Tessellation3DData.m_puiTriangulatedIndexes[Index++];
-			FaceIndex[1] = Tessellation3DData.m_puiTriangulatedIndexes[Index++];
-			FaceIndex[2] = Tessellation3DData.m_puiTriangulatedIndexes[Index++];
+	NormalIndex[0] = Tessellation3DData.m_puiTriangulatedIndexes[Index++];
+	NormalIndex[1] = NormalIndex[0];
+	NormalIndex[2] = NormalIndex[0];
 
-			if (!TechSoftFileParserImpl::AddFace(FaceIndex, Tessellation, VertexIndex))
-			{
-				continue;
-			}
+	// Get Triangles
+	for (unsigned long TriangleIndex = 0; TriangleIndex < InTrinangleCount; TriangleIndex++)
+	{
+		TextureIndex[0] = Tessellation3DData.m_puiTriangulatedIndexes[Index];
+		Index += TextureCount;
+		FaceIndex[0] = Tessellation3DData.m_puiTriangulatedIndexes[Index++];
+		TextureIndex[1] = Tessellation3DData.m_puiTriangulatedIndexes[Index];
+		Index += TextureCount;
+		FaceIndex[1] = Tessellation3DData.m_puiTriangulatedIndexes[Index++];
+		TextureIndex[2] = Tessellation3DData.m_puiTriangulatedIndexes[Index];
+		Index += TextureCount;
+		FaceIndex[2] = Tessellation3DData.m_puiTriangulatedIndexes[Index++];
 
-			TechSoftFileParserImpl::AddNormals(Tessellation3DData.m_pdNormals, NormalIndex, Tessellation.NormalArray);
+		if (!TechSoftFileParserImpl::AddFace(FaceIndex, Tessellation, VertexIndex))
+		{
+			continue;
 		}
 
-		InOutStartIndex = Index;
+		TechSoftFileParserImpl::AddNormals(Tessellation3DData.m_pdNormals, NormalIndex, Tessellation.NormalArray);
+		TechSoftFileParserImpl::AddTextureCoordinates(Tessellation3DData.m_pdTextureCoords, TextureIndex, Tessellation.TexCoordArray);
 	}
 
-	void AddFaceTriangleWithUniqueNormalAndTexture(FTessellationData& Tessellation, const A3DTess3DData& Tessellation3DData, long& InOutStartIndex, unsigned long InTrinangleCount, long TextureCount)
-	{
-		TechSoftFileParserImpl::Reserve(Tessellation, InTrinangleCount, /*bWithTexture*/ true);
+	InOutStartIndex = Index;
+}
 
-		unsigned long Index = InOutStartIndex;
+void AddFaceTriangle(FTessellationData& Tessellation, const A3DTess3DData& Tessellation3DData, long& InOutStartIndex, unsigned long InTrinangleCount)
+{
+	TechSoftFileParserImpl::Reserve(Tessellation, InTrinangleCount, /*bWithTexture*/ false);
+
+	unsigned long Index = InOutStartIndex;
 
-		int32 FaceIndex[3] = { 0, 0, 0 };
-		int32 NormalIndex[3] = { 0, 0, 0 };
-		int32 TextureIndex[3] = { 0, 0, 0 };
-		int32 VertexIndex = 0;
+	int32 FaceIndex[3] = { 0, 0, 0 };
+	int32 NormalIndex[3] = { 0, 0, 0 };
+	int32 VertexIndex = 0;
 
+	// Get Triangles
+	for (unsigned long TriangleIndex = 0; TriangleIndex < InTrinangleCount; TriangleIndex++)
+	{
 		NormalIndex[0] = Tessellation3DData.m_puiTriangulatedIndexes[Index++];
-		NormalIndex[1] = NormalIndex[0];
-		NormalIndex[2] = NormalIndex[0];
+		FaceIndex[0] = Tessellation3DData.m_puiTriangulatedIndexes[Index++];
+		NormalIndex[1] = Tessellation3DData.m_puiTriangulatedIndexes[Index++];
+		FaceIndex[1] = Tessellation3DData.m_puiTriangulatedIndexes[Index++];
+		NormalIndex[2] = Tessellation3DData.m_puiTriangulatedIndexes[Index++];
+		FaceIndex[2] = Tessellation3DData.m_puiTriangulatedIndexes[Index++];
 
-		// Get Triangles
-		for (unsigned long TriangleIndex = 0; TriangleIndex < InTrinangleCount; TriangleIndex++)
+		if (!TechSoftFileParserImpl::AddFace(FaceIndex, Tessellation, VertexIndex))
 		{
-			TextureIndex[0] = Tessellation3DData.m_puiTriangulatedIndexes[Index];
-			Index += TextureCount;
-			FaceIndex[0] = Tessellation3DData.m_puiTriangulatedIndexes[Index++];
-			TextureIndex[1] = Tessellation3DData.m_puiTriangulatedIndexes[Index];
-			Index += TextureCount;
-			FaceIndex[1] = Tessellation3DData.m_puiTriangulatedIndexes[Index++];
-			TextureIndex[2] = Tessellation3DData.m_puiTriangulatedIndexes[Index];
-			Index += TextureCount;
-			FaceIndex[2] = Tessellation3DData.m_puiTriangulatedIndexes[Index++];
-
-			if (!TechSoftFileParserImpl::AddFace(FaceIndex, Tessellation, VertexIndex))
-			{
-				continue;
-			}
-
-			TechSoftFileParserImpl::AddNormals(Tessellation3DData.m_pdNormals, NormalIndex, Tessellation.NormalArray);
-			TechSoftFileParserImpl::AddTextureCoordinates(Tessellation3DData.m_pdTextureCoords, TextureIndex, Tessellation.TexCoordArray);
+			continue;
 		}
 
-		InOutStartIndex = Index;
+		TechSoftFileParserImpl::AddNormals(Tessellation3DData.m_pdNormals, NormalIndex, Tessellation.NormalArray);
 	}
 
-	void AddFaceTriangle(FTessellationData& Tessellation, const A3DTess3DData& Tessellation3DData, long& InOutStartIndex, unsigned long InTrinangleCount)
-	{
-		TechSoftFileParserImpl::Reserve(Tessellation, InTrinangleCount, /*bWithTexture*/ false);
+	InOutStartIndex = Index;
+}
 
-		unsigned long Index = InOutStartIndex;
+void AddFaceTriangleWithTexture(FTessellationData& Tessellation, const A3DTess3DData& Tessellation3DData, long& InOutStartIndex, unsigned long InTrinangleCount, long TextureCount)
+{
+	TechSoftFileParserImpl::Reserve(Tessellation, InTrinangleCount, /*bWithTexture*/ true);
 
-		int32 FaceIndex[3] = { 0, 0, 0 };
-		int32 NormalIndex[3] = { 0, 0, 0 };
-		int32 VertexIndex = 0;
+	unsigned long Index = InOutStartIndex;
 
-		// Get Triangles
-		for (unsigned long TriangleIndex = 0; TriangleIndex < InTrinangleCount; TriangleIndex++)
-		{
-			NormalIndex[0] = Tessellation3DData.m_puiTriangulatedIndexes[Index++];
-			FaceIndex[0] = Tessellation3DData.m_puiTriangulatedIndexes[Index++];
-			NormalIndex[1] = Tessellation3DData.m_puiTriangulatedIndexes[Index++];
-			FaceIndex[1] = Tessellation3DData.m_puiTriangulatedIndexes[Index++];
-			NormalIndex[2] = Tessellation3DData.m_puiTriangulatedIndexes[Index++];
-			FaceIndex[2] = Tessellation3DData.m_puiTriangulatedIndexes[Index++];
-
-			if (!TechSoftFileParserImpl::AddFace(FaceIndex, Tessellation, VertexIndex))
-			{
-				continue;
-			}
+	int32 FaceIndex[3] = { 0, 0, 0 };
+	int32 NormalIndex[3] = { 0, 0, 0 };
+	int32 TextureIndex[3] = { 0, 0, 0 };
+	int32 VertexIndex = 0;
 
-			TechSoftFileParserImpl::AddNormals(Tessellation3DData.m_pdNormals, NormalIndex, Tessellation.NormalArray);
+	// Get Triangles
+	for (unsigned long TriangleIndex = 0; TriangleIndex < InTrinangleCount; TriangleIndex++)
+	{
+		NormalIndex[0] = Tessellation3DData.m_puiTriangulatedIndexes[Index++];
+		TextureIndex[0] = Tessellation3DData.m_puiTriangulatedIndexes[Index];
+		Index += TextureCount;
+		FaceIndex[0] = Tessellation3DData.m_puiTriangulatedIndexes[Index++];
+		NormalIndex[1] = Tessellation3DData.m_puiTriangulatedIndexes[Index++];
+		TextureIndex[1] = Tessellation3DData.m_puiTriangulatedIndexes[Index];
+		Index += TextureCount;
+		FaceIndex[1] = Tessellation3DData.m_puiTriangulatedIndexes[Index++];
+		NormalIndex[2] = Tessellation3DData.m_puiTriangulatedIndexes[Index++];
+		TextureIndex[2] = Tessellation3DData.m_puiTriangulatedIndexes[Index];
+		Index += TextureCount;
+		FaceIndex[2] = Tessellation3DData.m_puiTriangulatedIndexes[Index++];
+
+		if (!TechSoftFileParserImpl::AddFace(FaceIndex, Tessellation, VertexIndex))
+		{
+			continue;
 		}
 
-		InOutStartIndex = Index;
+		TechSoftFileParserImpl::AddNormals(Tessellation3DData.m_pdNormals, NormalIndex, Tessellation.NormalArray);
+		TechSoftFileParserImpl::AddTextureCoordinates(Tessellation3DData.m_pdTextureCoords, TextureIndex, Tessellation.TexCoordArray);
 	}
 
-	void AddFaceTriangleWithTexture(FTessellationData& Tessellation, const A3DTess3DData& Tessellation3DData, long& InOutStartIndex, unsigned long InTrinangleCount, long TextureCount)
-	{
-		TechSoftFileParserImpl::Reserve(Tessellation, InTrinangleCount, /*bWithTexture*/ true);
+	InOutStartIndex = Index;
+}
 
-		unsigned long Index = InOutStartIndex;
+// TODO import fan and strip meshes
+#ifdef UNUSED
+double ComputeFaceAreaTriangleFan(double* InCoordinates, unsigned long* InTriangulatedIndexes, long& InOutStartIndex, unsigned long InTrinangleCount, bool bInOneNormal, long TextureCount, int& TriangleCount)
+{
+	double Area = 0;
 
-		int32 FaceIndex[3] = { 0, 0, 0 };
-		int32 NormalIndex[3] = { 0, 0, 0 };
-		int32 TextureIndex[3] = { 0, 0, 0 };
-		int32 VertexIndex = 0;
+	unsigned long OffSet = TextureCount + (bInOneNormal ? 0 : 1) + 1;
+	unsigned long Index = InOutStartIndex + 2 + TextureCount;
 
-		// Get Triangles
-		for (unsigned long TriangleIndex = 0; TriangleIndex < InTrinangleCount; TriangleIndex++)
-		{
-			NormalIndex[0] = Tessellation3DData.m_puiTriangulatedIndexes[Index++];
-			TextureIndex[0] = Tessellation3DData.m_puiTriangulatedIndexes[Index];
-			Index += TextureCount;
-			FaceIndex[0] = Tessellation3DData.m_puiTriangulatedIndexes[Index++];
-			NormalIndex[1] = Tessellation3DData.m_puiTriangulatedIndexes[Index++];
-			TextureIndex[1] = Tessellation3DData.m_puiTriangulatedIndexes[Index];
-			Index += TextureCount;
-			FaceIndex[1] = Tessellation3DData.m_puiTriangulatedIndexes[Index++];
-			NormalIndex[2] = Tessellation3DData.m_puiTriangulatedIndexes[Index++];
-			TextureIndex[2] = Tessellation3DData.m_puiTriangulatedIndexes[Index];
-			Index += TextureCount;
-			FaceIndex[2] = Tessellation3DData.m_puiTriangulatedIndexes[Index++];
-
-			if (!TechSoftFileParserImpl::AddFace(FaceIndex, Tessellation, VertexIndex))
-			{
-				continue;
-			}
+	int Index0 = InTriangulatedIndexes[Index];
+	double* Point0 = InCoordinates + (InTriangulatedIndexes[Index]);
+	Index += OffSet;
+	int Index1 = InTriangulatedIndexes[Index];
+	double* Point1 = InCoordinates + (InTriangulatedIndexes[Index]);
+	Index += OffSet;
 
-			TechSoftFileParserImpl::AddNormals(Tessellation3DData.m_pdNormals, NormalIndex, Tessellation.NormalArray);
-			TechSoftFileParserImpl::AddTextureCoordinates(Tessellation3DData.m_pdTextureCoords, TextureIndex, Tessellation.TexCoordArray);
+	for (unsigned long TriangleIndex = 2; TriangleIndex < InTrinangleCount; TriangleIndex++)
+	{
+		double* Point2 = InCoordinates + (InTriangulatedIndexes[Index]);
+		int Index2 = InTriangulatedIndexes[Index];
+		if (Index0 == Index1 || Index0 == Index2 || Index1 == Index2)
+		{
+			TriangleCount--;
 		}
-
-		InOutStartIndex = Index;
+		else
+		{
+			Area += FaceArea(Point0, Point1, Point2);
+		}
+		Index += OffSet;
+		Point1 = Point2;
+		Index1 = Index2;
 	}
+	InOutStartIndex = Index - OffSet;
 
-	// TODO import fan and strip meshes
-#ifdef UNUSED
-	double ComputeFaceAreaTriangleFan(double* InCoordinates, unsigned long* InTriangulatedIndexes, long& InOutStartIndex, unsigned long InTrinangleCount, bool bInOneNormal, long TextureCount, int& TriangleCount)
-	{
-		double Area = 0;
+	return Area;
+}
 
-		unsigned long OffSet = TextureCount + (bInOneNormal ? 0 : 1) + 1;
-		unsigned long Index = InOutStartIndex + 2 + TextureCount;
+double ComputeFaceAreaTriangleStrip(double* InCoordinates, unsigned long* InTriangulatedIndexes, long& InOutStartIndex, unsigned long InTrinangleCount, bool bInOneNormal, long TextureCount, int& TriangleCount)
+{
+	double Area = 0;
 
-		int Index0 = InTriangulatedIndexes[Index];
-		double* Point0 = InCoordinates + (InTriangulatedIndexes[Index]);
-		Index += OffSet;
-		int Index1 = InTriangulatedIndexes[Index];
-		double* Point1 = InCoordinates + (InTriangulatedIndexes[Index]);
-		Index += OffSet;
+	unsigned long OffSet = TextureCount + (bInOneNormal ? 0 : 1) + 1;
+	unsigned long Index = InOutStartIndex + 2 + TextureCount;
 
-		for (unsigned long TriangleIndex = 2; TriangleIndex < InTrinangleCount; TriangleIndex++)
+	int Index0 = InTriangulatedIndexes[Index];
+	double* Point0 = InCoordinates + (InTriangulatedIndexes[Index]);
+	Index += OffSet;
+	int Index1 = InTriangulatedIndexes[Index];
+	double* Point1 = InCoordinates + (InTriangulatedIndexes[Index]);
+	Index += OffSet;
+
+	for (unsigned long TriangleIndex = 2; TriangleIndex < InTrinangleCount; TriangleIndex++)
+	{
+		double* Point2 = InCoordinates + (InTriangulatedIndexes[Index]);
+		int Index2 = InTriangulatedIndexes[Index];
+		if (Index0 == Index1 || Index0 == Index2 || Index1 == Index2)
 		{
-			double* Point2 = InCoordinates + (InTriangulatedIndexes[Index]);
-			int Index2 = InTriangulatedIndexes[Index];
-			if (Index0 == Index1 || Index0 == Index2 || Index1 == Index2)
-			{
-				TriangleCount--;
-			}
-			else
-			{
-				Area += FaceArea(Point0, Point1, Point2);
-			}
-			Index += OffSet;
-			Point1 = Point2;
-			Index1 = Index2;
+			TriangleCount--;
+		}
+		else
+		{
+			Area += FaceArea(Point0, Point1, Point2);
 		}
-		InOutStartIndex = Index - OffSet;
 
-		return Area;
+		Index += OffSet;
+		Point0 = Point1;
+		Point1 = Point2;
+		Index0 = Index1;
+		Index1 = Index2;
 	}
+	InOutStartIndex = Index - OffSet;
 
-	double ComputeFaceAreaTriangleStrip(double* InCoordinates, unsigned long* InTriangulatedIndexes, long& InOutStartIndex, unsigned long InTrinangleCount, bool bInOneNormal, long TextureCount, int& TriangleCount)
-	{
-		double Area = 0;
+	return Area;
+}
+#endif
 
-		unsigned long OffSet = TextureCount + (bInOneNormal ? 0 : 1) + 1;
-		unsigned long Index = InOutStartIndex + 2 + TextureCount;
+void FTechSoftFileParser::TraverseTessellation3D(const A3DTess3D* TessellationPtr, FArchiveBody& Body)
+{
+	FBodyMesh& BodyMesh = CADFileData.AddBodyMesh(Body.ObjectId, Body);
 
-		int Index0 = InTriangulatedIndexes[Index];
-		double* Point0 = InCoordinates + (InTriangulatedIndexes[Index]);
-		Index += OffSet;
-		int Index1 = InTriangulatedIndexes[Index];
-		double* Point1 = InCoordinates + (InTriangulatedIndexes[Index]);
-		Index += OffSet;
+	const int TessellationFaceDataWithTriangle = 0x2222;
+	const int TessellationFaceDataWithFan = 0x4444;
+	const int TessellationFaceDataWithStrip = 0x8888;
+	const int TessellationFaceDataWithOneNormal = 0xE0E0;
 
-		for (unsigned long TriangleIndex = 2; TriangleIndex < InTrinangleCount; TriangleIndex++)
+	// Coordinates
+	TUniqueTSObj<A3DTessBaseData> TessellationBaseData(TessellationPtr);
+	{
+		if (TessellationBaseData.IsValid() && TessellationBaseData->m_uiCoordSize > 0)
 		{
-			double* Point2 = InCoordinates + (InTriangulatedIndexes[Index]);
-			int Index2 = InTriangulatedIndexes[Index];
-			if (Index0 == Index1 || Index0 == Index2 || Index1 == Index2)
+			int32 VertexCount = TessellationBaseData->m_uiCoordSize / 3;
+			BodyMesh.VertexArray.Reserve(VertexCount);
+
+			double* Coordinates = TessellationBaseData->m_pdCoords;
+			for (unsigned int Index = 0; Index < TessellationBaseData->m_uiCoordSize; ++Index)
 			{
-				TriangleCount--;
+				Coordinates[Index] *= FileUnit;
 			}
-			else
+
+			for (unsigned int Index = 0; Index < TessellationBaseData->m_uiCoordSize; Index += 3)
 			{
-				Area += FaceArea(Point0, Point1, Point2);
+				BodyMesh.VertexArray.Emplace(Coordinates[Index], Coordinates[Index + 1], Coordinates[Index + 2]);
 			}
-
-			Index += OffSet;
-			Point0 = Point1;
-			Point1 = Point2;
-			Index0 = Index1;
-			Index1 = Index2;
 		}
-		InOutStartIndex = Index - OffSet;
-
-		return Area;
+		else
+		{
+			// No vertex, no mesh...
+			return;
+		}
 	}
-#endif
 
-	void FTechSoftFileParser::TraverseTessellation3D(const A3DTess3D* TessellationPtr, FArchiveBody& Body)
+	TUniqueTSObj<A3DTess3DData> Tessellation3DData(TessellationPtr);
+	if (Tessellation3DData.IsValid())
 	{
-		FBodyMesh& BodyMesh = CADFileData.AddBodyMesh(Body.ObjectId, Body);
-
-		const int TessellationFaceDataWithTriangle = 0x2222;
-		const int TessellationFaceDataWithFan = 0x4444;
-		const int TessellationFaceDataWithStrip = 0x8888;
-		const int TessellationFaceDataWithOneNormal = 0xE0E0;
-
-		// Coordinates
-		TUniqueTSObj<A3DTessBaseData> TessellationBaseData(TessellationPtr);
+		for (unsigned int Index = 0; Index < Tessellation3DData->m_uiFaceTessSize; ++Index)
 		{
-			if (TessellationBaseData.IsValid() && TessellationBaseData->m_uiCoordSize > 0)
-			{
-				int32 VertexCount = TessellationBaseData->m_uiCoordSize / 3;
-				BodyMesh.VertexArray.Reserve(VertexCount);
+			const A3DTessFaceData& FaceTessData = Tessellation3DData->m_psFaceTessData[Index];
+			FTessellationData& Tessellation = BodyMesh.Faces.Emplace_GetRef();
 
-				double* Coordinates = TessellationBaseData->m_pdCoords;
-				for (unsigned int Index = 0; Index < TessellationBaseData->m_uiCoordSize; ++Index)
+			if (FaceTessData.m_uiStyleIndexesSize == 1)
+			{
+				A3DUns32 StyleIndex = FaceTessData.m_puiStyleIndexes[0];
+				FCADUUID ColorName;
+				FCADUUID MaterialName;
+				TraverseGraphStyleData(StyleIndex, ColorName, MaterialName);
+				if (ColorName)
 				{
-					Coordinates[Index] *= FileUnit;
+					Tessellation.ColorName = ColorName;
+					BodyMesh.ColorSet.Add(ColorName);
 				}
 
-				for (unsigned int Index = 0; Index < TessellationBaseData->m_uiCoordSize; Index += 3)
+				if (MaterialName)
 				{
-					BodyMesh.VertexArray.Emplace(Coordinates[Index], Coordinates[Index + 1], Coordinates[Index + 2]);
+					Tessellation.MaterialName = MaterialName;
+					BodyMesh.MaterialSet.Add(MaterialName);
 				}
 			}
-			else
-			{
-				// No vertex, no mesh...
-				return;
-			}
-		}
 
-		TUniqueTSObj<A3DTess3DData> Tessellation3DData(TessellationPtr);
-		if (Tessellation3DData.IsValid())
-		{
-			for (unsigned int Index = 0; Index < Tessellation3DData->m_uiFaceTessSize; ++Index)
-			{
-				const A3DTessFaceData& FaceTessData = Tessellation3DData->m_psFaceTessData[Index];
-				FTessellationData& Tessellation = BodyMesh.Faces.Emplace_GetRef();
+			unsigned int UsedEntitiesFlags = FaceTessData.m_usUsedEntitiesFlags;
+			long StartTriangulated = FaceTessData.m_uiStartTriangulated;
 
-				if (FaceTessData.m_uiStyleIndexesSize == 1)
+			unsigned int FaceSetIndex = 0;
+			// Triangles
+			if (UsedEntitiesFlags & TessellationFaceDataWithTriangle)
+			{
+				bool bTessellationWithOneNormal = (UsedEntitiesFlags & TessellationFaceDataWithOneNormal);
+				if (bTessellationWithOneNormal)
 				{
-					A3DUns32 StyleIndex = FaceTessData.m_puiStyleIndexes[0];
-					FCADUUID ColorName;
-					FCADUUID MaterialName;
-					TraverseGraphStyleData(StyleIndex, ColorName, MaterialName);
-					if (ColorName)
+					if (FaceTessData.m_uiTextureCoordIndexesSize)
 					{
-						Tessellation.ColorName = ColorName;
-						BodyMesh.ColorSet.Add(ColorName);
+						AddFaceTriangleWithUniqueNormalAndTexture(Tessellation, *Tessellation3DData, StartTriangulated, FaceTessData.m_puiSizesTriangulated[0], FaceTessData.m_uiTextureCoordIndexesSize);
 					}
-
-					if (MaterialName)
+					else
 					{
-						Tessellation.MaterialName = MaterialName;
-						BodyMesh.MaterialSet.Add(MaterialName);
+						AddFaceTriangleWithUniqueNormal(Tessellation, *Tessellation3DData, StartTriangulated, FaceTessData.m_puiSizesTriangulated[0]);
 					}
 				}
-
-				unsigned int UsedEntitiesFlags = FaceTessData.m_usUsedEntitiesFlags;
-				long StartTriangulated = FaceTessData.m_uiStartTriangulated;
-
-				unsigned int FaceSetIndex = 0;
-				// Triangles
-				if (UsedEntitiesFlags & TessellationFaceDataWithTriangle)
+				else
 				{
-					bool bTessellationWithOneNormal = (UsedEntitiesFlags & TessellationFaceDataWithOneNormal);
-					if (bTessellationWithOneNormal)
+					if (FaceTessData.m_uiTextureCoordIndexesSize)
 					{
-						if (FaceTessData.m_uiTextureCoordIndexesSize)
-						{
-							AddFaceTriangleWithUniqueNormalAndTexture(Tessellation, *Tessellation3DData, StartTriangulated, FaceTessData.m_puiSizesTriangulated[0], FaceTessData.m_uiTextureCoordIndexesSize);
-						}
-						else
-						{
-							AddFaceTriangleWithUniqueNormal(Tessellation, *Tessellation3DData, StartTriangulated, FaceTessData.m_puiSizesTriangulated[0]);
-						}
+						AddFaceTriangleWithTexture(Tessellation, *Tessellation3DData, StartTriangulated, FaceTessData.m_puiSizesTriangulated[0], FaceTessData.m_uiTextureCoordIndexesSize);
 					}
 					else
 					{
-						if (FaceTessData.m_uiTextureCoordIndexesSize)
-						{
-							AddFaceTriangleWithTexture(Tessellation, *Tessellation3DData, StartTriangulated, FaceTessData.m_puiSizesTriangulated[0], FaceTessData.m_uiTextureCoordIndexesSize);
-						}
-						else
-						{
-							AddFaceTriangle(Tessellation, *Tessellation3DData, StartTriangulated, FaceTessData.m_puiSizesTriangulated[0]);
-						}
-					}
-					FaceSetIndex++;
-				}
-
-				// Fans TODO
-				if (FaceTessData.m_uiSizesTriangulatedSize > FaceSetIndex)
-				{
-					if (UsedEntitiesFlags & TessellationFaceDataWithFan)
-					{
-						unsigned int LastFanIndex = 1 + FaceSetIndex + FaceTessData.m_puiSizesTriangulated[FaceSetIndex];
-						FaceSetIndex++;
-						for (; FaceSetIndex < LastFanIndex; FaceSetIndex++)
-						{
-							bool bTessellationWithOneNormal = (UsedEntitiesFlags & TessellationFaceDataWithOneNormal) && (FaceTessData.m_puiSizesTriangulated[FaceSetIndex] & kA3DTessFaceDataNormalSingle);
-							A3DUns32 FanSize = (FaceTessData.m_puiSizesTriangulated[FaceSetIndex] & kA3DTessFaceDataNormalMask);
-							//TriangleCount += (FanSize - 2);
-							//BodyMeshMetaData.MeshArea += ComputeFaceAreaTriangleFan(TessellationBaseData.m_pdCoords, Tessellation3DData->m_puiTriangulatedIndexes, StartTriangulated, FanSize, bTessellationWithOneNormal, FaceTessData.m_uiTextureCoordIndexesSize, TriangleCount);
-						}
+						AddFaceTriangle(Tessellation, *Tessellation3DData, StartTriangulated, FaceTessData.m_puiSizesTriangulated[0]);
 					}
 				}
+				FaceSetIndex++;
+			}
 
-				// Strip TODO
-				if (FaceTessData.m_uiSizesTriangulatedSize > FaceSetIndex)
+			// Fans TODO
+			if (FaceTessData.m_uiSizesTriangulatedSize > FaceSetIndex)
+			{
+				if (UsedEntitiesFlags & TessellationFaceDataWithFan)
 				{
+					unsigned int LastFanIndex = 1 + FaceSetIndex + FaceTessData.m_puiSizesTriangulated[FaceSetIndex];
 					FaceSetIndex++;
-					for (; FaceSetIndex < FaceTessData.m_uiSizesTriangulatedSize; FaceSetIndex++)
+					for (; FaceSetIndex < LastFanIndex; FaceSetIndex++)
 					{
 						bool bTessellationWithOneNormal = (UsedEntitiesFlags & TessellationFaceDataWithOneNormal) && (FaceTessData.m_puiSizesTriangulated[FaceSetIndex] & kA3DTessFaceDataNormalSingle);
-						A3DUns32 StripSize = (FaceTessData.m_puiSizesTriangulated[FaceSetIndex] & kA3DTessFaceDataNormalMask);
-						//TriangleCount += (StripSize - 2);
-						//BodyMeshMetaData.MeshArea += ComputeFaceAreaTriangleStrip(TessellationBaseData.m_pdCoords, Tessellation3DData->m_puiTriangulatedIndexes, StartTriangulated, StripSize, bTessellationWithOneNormal, FaceTessData.m_uiTextureCoordIndexesSize, TriangleCount);
+						A3DUns32 FanSize = (FaceTessData.m_puiSizesTriangulated[FaceSetIndex] & kA3DTessFaceDataNormalMask);
+						//TriangleCount += (FanSize - 2);
+						//BodyMeshMetaData.MeshArea += ComputeFaceAreaTriangleFan(TessellationBaseData.m_pdCoords, Tessellation3DData->m_puiTriangulatedIndexes, StartTriangulated, FanSize, bTessellationWithOneNormal, FaceTessData.m_uiTextureCoordIndexesSize, TriangleCount);
 					}
 				}
+			}
 
+			// Strip TODO
+			if (FaceTessData.m_uiSizesTriangulatedSize > FaceSetIndex)
+			{
+				FaceSetIndex++;
+				for (; FaceSetIndex < FaceTessData.m_uiSizesTriangulatedSize; FaceSetIndex++)
+				{
+					bool bTessellationWithOneNormal = (UsedEntitiesFlags & TessellationFaceDataWithOneNormal) && (FaceTessData.m_puiSizesTriangulated[FaceSetIndex] & kA3DTessFaceDataNormalSingle);
+					A3DUns32 StripSize = (FaceTessData.m_puiSizesTriangulated[FaceSetIndex] & kA3DTessFaceDataNormalMask);
+					//TriangleCount += (StripSize - 2);
+					//BodyMeshMetaData.MeshArea += ComputeFaceAreaTriangleStrip(TessellationBaseData.m_pdCoords, Tessellation3DData->m_puiTriangulatedIndexes, StartTriangulated, StripSize, bTessellationWithOneNormal, FaceTessData.m_uiTextureCoordIndexesSize, TriangleCount);
+				}
 			}
-		}
 
-		Body.ColorFaceSet = BodyMesh.ColorSet;
-		Body.MaterialFaceSet = BodyMesh.MaterialSet;
+		}
 	}
+
+	Body.ColorFaceSet = BodyMesh.ColorSet;
+	Body.MaterialFaceSet = BodyMesh.MaterialSet;
 }
 
+} // ns
+
 #endif
\ No newline at end of file
diff --git a/Engine/Plugins/Enterprise/DatasmithCADImporter/Source/CADInterfaces/Private/TechSoftFileParser.h b/Engine/Plugins/Enterprise/DatasmithCADImporter/Source/CADInterfaces/Private/TechSoftFileParser.h
--- a/Engine/Plugins/Enterprise/DatasmithCADImporter/Source/CADInterfaces/Private/TechSoftFileParser.h
+++ b/Engine/Plugins/Enterprise/DatasmithCADImporter/Source/CADInterfaces/Private/TechSoftFileParser.h
@@ -29,260 +29,262 @@ typedef void A3DTopoBrepData;
 
 namespace CADLibrary
 {
-	class FEntityData;
-	class FComponentData;
 
-	enum EComponentType : uint32
-	{
-		Reference = 0,
-		Occurrence,
-		UnloadedComponent,
-		Body,
-		Undefined,
-		LastType
-	};
-
-	enum EModellerType : uint32
-	{
-		ModellerUnknown = 0,              /*!< User modeller. */
-		ModellerCatia = 2,                /*!< CATIA modeller. */
-		ModellerCatiaV5 = 3,              /*!< CATIA V5 modeller. */
-		ModellerCadds = 4,                /*!< CADDS modeller. */
-		ModellerUnigraphics = 5,          /*!< Unigraphics modeller. */
-		ModellerParasolid = 6,            /*!< Parasolid modeller. */
-		ModellerEuclid = 7,               /*!< Euclid modeller. */
-		ModellerIges = 9,                 /*!< IGES modeller. */
-		ModellerUnisurf = 10,             /*!< Unisurf modeller. */
-		ModellerVda = 11,                 /*!< VDA modeller. */
-		ModellerStl = 12,                 /*!< STL modeller. */
-		ModellerWrl = 13,                 /*!< WRL modeller. */
-		ModellerDxf = 14,                 /*!< DXF modeller. */
-		ModellerAcis = 15,                /*!< ACIS modeller. */
-		ModellerProE = 16,                /*!< Pro/E modeller. */
-		ModellerStep = 18,                /*!< STEP modeller. */
-		ModellerIdeas = 19,               /*!< I-DEAS modeller. */
-		ModellerJt = 20,                  /*!< JT modeller. */
-		ModellerSlw = 22,                 /*!< SolidWorks modeller. */
-		ModellerCgr = 23,                 /*!< CGR modeller. */
-		ModellerPrc = 24,                 /*!< PRC modeller. */
-		ModellerXvl = 25,                 /*!< XVL modeller. */
-		ModellerHpgl = 26,                /*!< HPGL modeller. */
-		ModellerTopSolid = 27,            /*!< TopSolid modeller. */
-		ModellerOneSpaceDesigner = 28,    /*!< OneSpace designer modeller. */
-		Modeller3dxml = 29,               /*!< 3DXML modeller. */
-		ModellerInventor = 30,            /*!< Inventor modeller. */
-		ModellerPostScript = 31,          /*!< Postscript modeller. */
-		ModellerPDF = 32,                 /*!< PDF modeller. */
-		ModellerU3D = 33,                 /*!< U3D modeller. */
-		ModellerIFC = 34,                 /*!< IFC modeller. */
-		ModellerDWG = 35,                 /*!< DWG modeller. */
-		ModellerDWF = 36,                 /*!< DWF modeller. */
-		ModellerSE = 37,                  /*!< SolidEdge modeller. */
-		ModellerOBJ = 38,                 /*!< OBJ modeller. */
-		ModellerKMZ = 39,                 /*!< KMZ modeller. */
-		ModellerDAE = 40,                 /*!< COLLADA modeller. */
-		Modeller3DS = 41,                 /*!< 3DS modeller. */
-		ModellerRhino = 43,               /*!< Rhino modeller. */
-		ModellerXML = 44,                 /*!< XML modeller. */
-		Modeller3mf = 45,                 /*!< 3MF modeller. */
-		ModellerScs = 46,                 /*!< SCS modeller. */
-		Modeller3dHtml = 47,              /*!< 3DHTML modeller. */
-		ModellerHsf = 48,                 /*!< Hsf modeller. */
-		ModellerGltf = 49,                /*!< GL modeller. */
-		ModellerRevit = 50,               /*!< Revit modeller. */
-		ModellerFBX = 51,                 /*!< FBX modeller. */
-		ModellerStepXML = 52,             /*!< StepXML modeller. */
-		ModellerPLMXML = 53,              /*!< PLMXML modeller. */
-		ModellerNavisworks = 54,			/*!< For Future Use: Navisworks modeller. */
-		ModellerLast
-	};
-
-	struct FEntityBehaviour
-	{
-		bool bFatherHeritColor = false;
-		bool bFatherHeritLayer = false;
-		bool bFatherHeritLinePattern = false;
-		bool bFatherHeritLineWidth = false;
-		bool bFatherHeritShow = false;
-		bool bFatherHeritTransparency = false;
-		bool bRemoved = false;
-		bool bShow = true;
-		bool bSonHeritColor = false;
-		bool bSonHeritLayer = false;
-		bool bSonHeritLinePattern = false;
-		bool bSonHeritLineWidth = false;
-		bool bSonHeritShow = false;
-		bool bSonHeritTransparency = false;
-	};
-
-	struct FEntityMetaData
-	{
-		TMap<FString, FString> MetaData;
-		bool bRemoved = false;
-		bool bShow = true;
-		bool bUnloaded = false;
-		FFileDescriptor ExternalFile;
-	};
+class FEntityData;
+class FComponentData;
 
+enum EComponentType : uint32
+{
+	Reference = 0,
+	Occurrence,
+	UnloadedComponent,
+	Body,
+	Undefined,
+	LastType
+};
+
+enum EModellerType : uint32
+{
+	ModellerUnknown = 0,              /*!< User modeller. */
+	ModellerCatia = 2,                /*!< CATIA modeller. */
+	ModellerCatiaV5 = 3,              /*!< CATIA V5 modeller. */
+	ModellerCadds = 4,                /*!< CADDS modeller. */
+	ModellerUnigraphics = 5,          /*!< Unigraphics modeller. */
+	ModellerParasolid = 6,            /*!< Parasolid modeller. */
+	ModellerEuclid = 7,               /*!< Euclid modeller. */
+	ModellerIges = 9,                 /*!< IGES modeller. */
+	ModellerUnisurf = 10,             /*!< Unisurf modeller. */
+	ModellerVda = 11,                 /*!< VDA modeller. */
+	ModellerStl = 12,                 /*!< STL modeller. */
+	ModellerWrl = 13,                 /*!< WRL modeller. */
+	ModellerDxf = 14,                 /*!< DXF modeller. */
+	ModellerAcis = 15,                /*!< ACIS modeller. */
+	ModellerProE = 16,                /*!< Pro/E modeller. */
+	ModellerStep = 18,                /*!< STEP modeller. */
+	ModellerIdeas = 19,               /*!< I-DEAS modeller. */
+	ModellerJt = 20,                  /*!< JT modeller. */
+	ModellerSlw = 22,                 /*!< SolidWorks modeller. */
+	ModellerCgr = 23,                 /*!< CGR modeller. */
+	ModellerPrc = 24,                 /*!< PRC modeller. */
+	ModellerXvl = 25,                 /*!< XVL modeller. */
+	ModellerHpgl = 26,                /*!< HPGL modeller. */
+	ModellerTopSolid = 27,            /*!< TopSolid modeller. */
+	ModellerOneSpaceDesigner = 28,    /*!< OneSpace designer modeller. */
+	Modeller3dxml = 29,               /*!< 3DXML modeller. */
+	ModellerInventor = 30,            /*!< Inventor modeller. */
+	ModellerPostScript = 31,          /*!< Postscript modeller. */
+	ModellerPDF = 32,                 /*!< PDF modeller. */
+	ModellerU3D = 33,                 /*!< U3D modeller. */
+	ModellerIFC = 34,                 /*!< IFC modeller. */
+	ModellerDWG = 35,                 /*!< DWG modeller. */
+	ModellerDWF = 36,                 /*!< DWF modeller. */
+	ModellerSE = 37,                  /*!< SolidEdge modeller. */
+	ModellerOBJ = 38,                 /*!< OBJ modeller. */
+	ModellerKMZ = 39,                 /*!< KMZ modeller. */
+	ModellerDAE = 40,                 /*!< COLLADA modeller. */
+	Modeller3DS = 41,                 /*!< 3DS modeller. */
+	ModellerRhino = 43,               /*!< Rhino modeller. */
+	ModellerXML = 44,                 /*!< XML modeller. */
+	Modeller3mf = 45,                 /*!< 3MF modeller. */
+	ModellerScs = 46,                 /*!< SCS modeller. */
+	Modeller3dHtml = 47,              /*!< 3DHTML modeller. */
+	ModellerHsf = 48,                 /*!< Hsf modeller. */
+	ModellerGltf = 49,                /*!< GL modeller. */
+	ModellerRevit = 50,               /*!< Revit modeller. */
+	ModellerFBX = 51,                 /*!< FBX modeller. */
+	ModellerStepXML = 52,             /*!< StepXML modeller. */
+	ModellerPLMXML = 53,              /*!< PLMXML modeller. */
+	ModellerNavisworks = 54,			/*!< For Future Use: Navisworks modeller. */
+	ModellerLast
+};
+
+struct FEntityBehaviour
+{
+	bool bFatherHeritColor = false;
+	bool bFatherHeritLayer = false;
+	bool bFatherHeritLinePattern = false;
+	bool bFatherHeritLineWidth = false;
+	bool bFatherHeritShow = false;
+	bool bFatherHeritTransparency = false;
+	bool bRemoved = false;
+	bool bShow = true;
+	bool bSonHeritColor = false;
+	bool bSonHeritLayer = false;
+	bool bSonHeritLinePattern = false;
+	bool bSonHeritLineWidth = false;
+	bool bSonHeritShow = false;
+	bool bSonHeritTransparency = false;
+};
+
+struct FEntityMetaData
+{
+	TMap<FString, FString> MetaData;
+	bool bRemoved = false;
+	bool bShow = true;
+	bool bUnloaded = false;
+	FFileDescriptor ExternalFile;
+};
 
-	class FTechSoftInterface;
-	
-	class CADINTERFACES_API FTechSoftFileParser : public ICADFileParser
-	{
-	public:
 
-		/**
-		 * @param InCADData TODO
-		 * @param EnginePluginsPath Full Path of EnginePlugins. Mandatory to set KernelIO to import DWG, or DGN files
-		 */
-		FTechSoftFileParser(FCADFileData& InCADData, const FString& EnginePluginsPath = TEXT(""));
+class FTechSoftInterface;
+
+class CADINTERFACES_API FTechSoftFileParser : public ICADFileParser
+{
+public:
+
+	/**
+	 * @param InCADData TODO
+	 * @param EnginePluginsPath Full Path of EnginePlugins. Mandatory to set KernelIO to import DWG, or DGN files
+	 */
+	FTechSoftFileParser(FCADFileData& InCADData, const FString& EnginePluginsPath = TEXT(""));
 
 #ifndef USE_TECHSOFT_SDK
-		virtual ECADParsingResult Process() override
-		{
-			return ECADParsingResult::ProcessFailed;
-		}
+	virtual ECADParsingResult Process() override
+	{
+		return ECADParsingResult::ProcessFailed;
+}
 #else
-		virtual ECADParsingResult Process() override;
-
-	private:
-	
-		// Methods to parse a first time the file to count the number of components
-		// Needed to reserve CADFileData
-		// Start with CountUnderModel
-		void CountUnderModel(const A3DAsmModelFile* AsmModel);
-		void CountUnderOccurrence(const A3DAsmProductOccurrence* Occurrence);
-		void CountUnderPrototype(const A3DAsmProductOccurrence* Prototype);
-		void CountUnderPartDefinition(const A3DAsmPartDefinition* PartDefinition);
-		void CountUnderRepresentation(const A3DRiRepresentationItem* RepresentationItem);
-		void CountUnderRepresentationSet(const A3DRiSet* RepresentationSet);
-
-		void ReserveCADFileData();
-
-		// Materials and colors
-		int32 CountMaterial();
-		int32 CountColor();
-		void ReadMaterialsAndColors();
-
-		// Traverse ASM tree by starting from the model
-		ECADParsingResult TraverseModel(const A3DAsmModelFile* AsmModel);
-		void TraverseReference(const A3DAsmProductOccurrence* Reference);
-		bool IsConfigurationSet(const A3DAsmProductOccurrence* Occurrence);
-		// TODO for SW //void TraverseReferenceFromConfiguration(const A3DAsmProductOccurrence* Reference, FEntityData& ParentMetaData);
-		// TODO for SW //void TraverseConfigurationSet(const A3DAsmProductOccurrence* ConfigurationSet, FEntityData& ParentMetaData);
-		// TODO for SW //void TraverseConfiguration(const A3DAsmProductOccurrence* OccurrConfigurationence, FEntityData& ParentMetaData);
-		FCadId TraverseOccurrence(const A3DAsmProductOccurrence* Occurrence);
-		void TraversePrototype(const A3DAsmProductOccurrence* InPrototype, FEntityMetaData& OutMetaData, FMatrix& OutPrototypeMatrix);
-		void TraversePartDefinition(const A3DAsmPartDefinition* PartDefinition, FArchiveComponent& Component);
-		FCadId TraverseRepresentationSet(const A3DRiSet* pSet, FEntityMetaData& PartMetaData);
-		FCadId TraverseRepresentationItem(A3DRiRepresentationItem* RepresentationItem, FEntityMetaData& PartMetaData);
-		FCadId TraverseBRepModel(A3DRiBrepModel* BrepModel, FEntityMetaData& PartMetaData);
-		FCadId TraversePolyBRepModel(const A3DRiPolyBrepModel* PolygonalBrepModel, FEntityMetaData& PartMetaData);
-
-		// Tessellation methods
-		void MeshRepresentationWithTechSoft(A3DRiRepresentationItem* RepresentationItem, FArchiveBody& Body);
-		void TraverseRepresentationContent(const A3DRiRepresentationItem* RepresentationItem, FArchiveBody& Body);
-		void TraverseTessellationBase(const A3DTessBase* Tessellation, FArchiveBody& Body);
-		void TraverseTessellation3D(const A3DTess3D* Tessellation, FArchiveBody& Body);
-
-		// MetaData
-		void TraverseMetaData(const A3DEntity* Entity, FEntityMetaData& OutMetaData);
-		void TraverseSpecificMetaData(const A3DAsmProductOccurrence* Occurrence, FEntityMetaData& OutMetaData);
-
-		void BuildInstanceName(TMap<FString, FString>& MetaData);
-		void BuildReferenceName(TMap<FString, FString>& MetaData);
-		void BuildPartName(TMap<FString, FString>& MetaData);
-		void BuildBodyName(TMap<FString, FString>& MetaData);
-
-		// Graphic properties
-		void TraverseGraphics(const A3DGraphics* Graphics, FEntityMetaData& OutMetaData);
-		void TraverseGraphStyleData(uint32 StyleIndex, FCADUUID& ColorName, FCADUUID& MaterialName);
-		void TraverseMaterialProperties(const A3DEntity* Entity);
-		void TraverseLayer(const A3DAsmProductOccurrence* Occurrence);
-		FArchiveColor& FindOrAddColor(uint32 ColorIndex, uint8 Alpha);
-		FArchiveMaterial& FindOrAddMaterial(uint32 MaterialIndex);
-
-		// Transform
-		FMatrix TraverseCoordinateSystem(const A3DRiCoordinateSystem* CoordinateSystem);
-		FMatrix TraverseTransformation(const A3DMiscCartesianTransformation* Transformation3d);
-		FMatrix TraverseGeneralTransformation(const A3DMiscGeneralTransformation* GeneralTransformation);
-		FMatrix TraverseTransformation3D(const A3DMiscCartesianTransformation* CartesianTransformation);
-
-		FFileDescriptor GetOccurrenceFileName(const A3DAsmProductOccurrence* OccurrencePtr);
-
-		// Archive methodes
-		FArchiveInstance& AddInstance(FEntityMetaData& InstanceMetaData);
-		FArchiveComponent& AddComponent(FEntityMetaData& ComponentMetaData, FArchiveInstance& Instance);
-		FArchiveUnloadedComponent& AddUnloadedComponent(FEntityMetaData& ComponentMetaData, FArchiveInstance& Instance);
-		FArchiveComponent& AddOccurence(FEntityMetaData& InstanceMetaData, FEntityMetaData& ReferenceMetaData, FCadId& OutComponentId);
-		FArchiveComponent& AddOccurence(FEntityMetaData& InstanceMetaData, FCadId& OutComponentId);
-		FArchiveBody& AddBody(FEntityMetaData& BodyMetaData);
+	virtual ECADParsingResult Process() override;
+
+private:
+
+	// Methods to parse a first time the file to count the number of components
+	// Needed to reserve CADFileData
+	// Start with CountUnderModel
+	void CountUnderModel(const A3DAsmModelFile* AsmModel);
+	void CountUnderOccurrence(const A3DAsmProductOccurrence* Occurrence);
+	void CountUnderPrototype(const A3DAsmProductOccurrence* Prototype);
+	void CountUnderPartDefinition(const A3DAsmPartDefinition* PartDefinition);
+	void CountUnderRepresentation(const A3DRiRepresentationItem* RepresentationItem);
+	void CountUnderRepresentationSet(const A3DRiSet* RepresentationSet);
+
+	void ReserveCADFileData();
+
+	// Materials and colors
+	int32 CountMaterial();
+	int32 CountColor();
+	void ReadMaterialsAndColors();
+
+	// Traverse ASM tree by starting from the model
+	ECADParsingResult TraverseModel(const A3DAsmModelFile* AsmModel);
+	void TraverseReference(const A3DAsmProductOccurrence* Reference);
+	bool IsConfigurationSet(const A3DAsmProductOccurrence* Occurrence);
+	// TODO for SW //void TraverseReferenceFromConfiguration(const A3DAsmProductOccurrence* Reference, FEntityData& ParentMetaData);
+	// TODO for SW //void TraverseConfigurationSet(const A3DAsmProductOccurrence* ConfigurationSet, FEntityData& ParentMetaData);
+	// TODO for SW //void TraverseConfiguration(const A3DAsmProductOccurrence* OccurrConfigurationence, FEntityData& ParentMetaData);
+	FCadId TraverseOccurrence(const A3DAsmProductOccurrence* Occurrence);
+	void TraversePrototype(const A3DAsmProductOccurrence* InPrototype, FEntityMetaData& OutMetaData, FMatrix& OutPrototypeMatrix);
+	void TraversePartDefinition(const A3DAsmPartDefinition* PartDefinition, FArchiveComponent& Component);
+	FCadId TraverseRepresentationSet(const A3DRiSet* pSet, FEntityMetaData& PartMetaData);
+	FCadId TraverseRepresentationItem(A3DRiRepresentationItem* RepresentationItem, FEntityMetaData& PartMetaData);
+	FCadId TraverseBRepModel(A3DRiBrepModel* BrepModel, FEntityMetaData& PartMetaData);
+	FCadId TraversePolyBRepModel(const A3DRiPolyBrepModel* PolygonalBrepModel, FEntityMetaData& PartMetaData);
+
+	// Tessellation methods
+	void MeshRepresentationWithTechSoft(A3DRiRepresentationItem* RepresentationItem, FArchiveBody& Body);
+	void TraverseRepresentationContent(const A3DRiRepresentationItem* RepresentationItem, FArchiveBody& Body);
+	void TraverseTessellationBase(const A3DTessBase* Tessellation, FArchiveBody& Body);
+	void TraverseTessellation3D(const A3DTess3D* Tessellation, FArchiveBody& Body);
+
+	// MetaData
+	void TraverseMetaData(const A3DEntity* Entity, FEntityMetaData& OutMetaData);
+	void TraverseSpecificMetaData(const A3DAsmProductOccurrence* Occurrence, FEntityMetaData& OutMetaData);
+
+	void BuildInstanceName(TMap<FString, FString>& MetaData);
+	void BuildReferenceName(TMap<FString, FString>& MetaData);
+	void BuildPartName(TMap<FString, FString>& MetaData);
+	void BuildBodyName(TMap<FString, FString>& MetaData);
+
+	// Graphic properties
+	void TraverseGraphics(const A3DGraphics* Graphics, FEntityMetaData& OutMetaData);
+	void TraverseGraphStyleData(uint32 StyleIndex, FCADUUID& ColorName, FCADUUID& MaterialName);
+	void TraverseMaterialProperties(const A3DEntity* Entity);
+	void TraverseLayer(const A3DAsmProductOccurrence* Occurrence);
+	FArchiveColor& FindOrAddColor(uint32 ColorIndex, uint8 Alpha);
+	FArchiveMaterial& FindOrAddMaterial(uint32 MaterialIndex);
+
+	// Transform
+	FMatrix TraverseCoordinateSystem(const A3DRiCoordinateSystem* CoordinateSystem);
+	FMatrix TraverseTransformation(const A3DMiscCartesianTransformation* Transformation3d);
+	FMatrix TraverseGeneralTransformation(const A3DMiscGeneralTransformation* GeneralTransformation);
+	FMatrix TraverseTransformation3D(const A3DMiscCartesianTransformation* CartesianTransformation);
+
+	FFileDescriptor GetOccurrenceFileName(const A3DAsmProductOccurrence* OccurrencePtr);
+
+	// Archive methodes
+	FArchiveInstance& AddInstance(FEntityMetaData& InstanceMetaData);
+	FArchiveComponent& AddComponent(FEntityMetaData& ComponentMetaData, FArchiveInstance& Instance);
+	FArchiveUnloadedComponent& AddUnloadedComponent(FEntityMetaData& ComponentMetaData, FArchiveInstance& Instance);
+	FArchiveComponent& AddOccurence(FEntityMetaData& InstanceMetaData, FEntityMetaData& ReferenceMetaData, FCadId& OutComponentId);
+	FArchiveComponent& AddOccurence(FEntityMetaData& InstanceMetaData, FCadId& OutComponentId);
+	FArchiveBody& AddBody(FEntityMetaData& BodyMetaData);
 #endif
 
-	private:
+private:
 
-		TSet<const A3DAsmProductOccurrence*> PrototypeCounted;
-		uint32 ComponentCount[EComponentType::LastType] = {0};
+	TSet<const A3DAsmProductOccurrence*> PrototypeCounted;
+	uint32 ComponentCount[EComponentType::LastType] = { 0 };
 
-		FCADFileData& CADFileData;
-		FTechSoftInterface& TechSoftInterface;
+	FCADFileData& CADFileData;
+	FTechSoftInterface& TechSoftInterface;
 
-		ECADFormat Format;
+	ECADFormat Format;
 
-		EModellerType ModellerType;
-		double FileUnit = 1;
-		FCadId LastEntityId = 1;
-	};
+	EModellerType ModellerType;
+	double FileUnit = 1;
+	FCadId LastEntityId = 1;
+};
 
 namespace TechSoftFileParserImpl
 {
-	// Methodes used by TraverseTessellation3D
-	
-	typedef double A3DDouble;
-	inline bool AddFace(int32 FaceIndex[3], CADLibrary::FTessellationData& Tessellation, int32& VertexIndex)
+// Methodes used by TraverseTessellation3D
+
+typedef double A3DDouble;
+inline bool AddFace(int32 FaceIndex[3], CADLibrary::FTessellationData& Tessellation, int32& VertexIndex)
+{
+	if (FaceIndex[0] == FaceIndex[1] || FaceIndex[0] == FaceIndex[2] || FaceIndex[1] == FaceIndex[2])
 	{
-		if (FaceIndex[0] == FaceIndex[1] || FaceIndex[0] == FaceIndex[2] || FaceIndex[1] == FaceIndex[2])
-		{
-			return false;
-		}
-
-		for (int32 Index = 0; Index < 3; ++Index)
-		{
-			FaceIndex[Index] /= 3;
-			Tessellation.VertexIndices.Add(VertexIndex++);
-		}
-		Tessellation.PositionIndices.Append(FaceIndex, 3);
-		return true;
-	};
-
-	inline void AddNormals(const A3DDouble* Normals, const int32 Indices[3], TArray<FVector>& TessellationNormals)
+		return false;
+	}
+
+	for (int32 Index = 0; Index < 3; ++Index)
+	{
+		FaceIndex[Index] /= 3;
+		Tessellation.VertexIndices.Add(VertexIndex++);
+	}
+	Tessellation.PositionIndices.Append(FaceIndex, 3);
+	return true;
+};
+
+inline void AddNormals(const A3DDouble* Normals, const int32 Indices[3], TArray<FVector>& TessellationNormals)
+{
+	for (int32 Index = 0; Index < 3; ++Index)
 	{
-		for (int32 Index = 0; Index < 3; ++Index)
-		{
-			int32 NormalIndex = Indices[Index];
-			TessellationNormals.Emplace(Normals[NormalIndex], Normals[NormalIndex + 1], Normals[NormalIndex + 2]);
-		}
-	};
-
-	inline void AddTextureCoordinates(const A3DDouble* TextureCoords, const int32 Indices[3], TArray<FVector2D>& TessellationTextures)
+		int32 NormalIndex = Indices[Index];
+		TessellationNormals.Emplace(Normals[NormalIndex], Normals[NormalIndex + 1], Normals[NormalIndex + 2]);
+	}
+};
+
+inline void AddTextureCoordinates(const A3DDouble* TextureCoords, const int32 Indices[3], TArray<FVector2D>& TessellationTextures)
+{
+	for (int32 Index = 0; Index < 3; ++Index)
 	{
-		for (int32 Index = 0; Index < 3; ++Index)
-		{
-			int32 TextureIndex = Indices[Index];
-			TessellationTextures.Emplace(TextureCoords[TextureIndex], TextureCoords[TextureIndex + 1]);
-		}
-	};
-
-	inline void Reserve(CADLibrary::FTessellationData& Tessellation, int32 InTrinangleCount, bool bWithTexture)
+		int32 TextureIndex = Indices[Index];
+		TessellationTextures.Emplace(TextureCoords[TextureIndex], TextureCoords[TextureIndex + 1]);
+	}
+};
+
+inline void Reserve(CADLibrary::FTessellationData& Tessellation, int32 InTrinangleCount, bool bWithTexture)
+{
+	Tessellation.PositionIndices.Reserve(3 * InTrinangleCount);
+	Tessellation.VertexIndices.Reserve(3 * InTrinangleCount);
+	Tessellation.NormalArray.Reserve(3 * InTrinangleCount);
+	if (bWithTexture)
 	{
-		Tessellation.PositionIndices.Reserve(3 * InTrinangleCount);
-		Tessellation.VertexIndices.Reserve(3 * InTrinangleCount);
-		Tessellation.NormalArray.Reserve(3 * InTrinangleCount);
-		if (bWithTexture)
-		{
-			Tessellation.TexCoordArray.Reserve(3 * InTrinangleCount);
-		}
-	};
+		Tessellation.TexCoordArray.Reserve(3 * InTrinangleCount);
+	}
+};
+
+} // ns TechSoftFileParserImpl
 
-}
 } // ns CADLibrary
\ No newline at end of file
diff --git a/Engine/Plugins/Enterprise/DatasmithCADImporter/Source/CADInterfaces/Private/TechSoftInterface.cpp b/Engine/Plugins/Enterprise/DatasmithCADImporter/Source/CADInterfaces/Private/TechSoftInterface.cpp
--- a/Engine/Plugins/Enterprise/DatasmithCADImporter/Source/CADInterfaces/Private/TechSoftInterface.cpp
+++ b/Engine/Plugins/Enterprise/DatasmithCADImporter/Source/CADInterfaces/Private/TechSoftInterface.cpp
@@ -12,358 +12,358 @@
 
 namespace CADLibrary
 {
-	FTechSoftInterface& GetTechSoftInterface()
+FTechSoftInterface& GetTechSoftInterface()
+{
+	static FTechSoftInterface TechSoftInterface;
+	return TechSoftInterface;
+}
+
+bool TECHSOFT_InitializeKernel(const TCHAR* InEnginePluginsPath)
+{
+	return GetTechSoftInterface().InitializeKernel(InEnginePluginsPath);
+}
+
+bool FTechSoftInterface::InitializeKernel(const TCHAR* InEnginePluginsPath)
+{
+#ifdef USE_TECHSOFT_SDK
+	if (bIsInitialize)
 	{
-		static FTechSoftInterface TechSoftInterface;
-		return TechSoftInterface;
+		return true;
 	}
 
-	bool TECHSOFT_InitializeKernel(const TCHAR* InEnginePluginsPath)
+	FString EnginePluginsPath(InEnginePluginsPath);
+	if (EnginePluginsPath.IsEmpty())
 	{
-		return GetTechSoftInterface().InitializeKernel(InEnginePluginsPath);
+		EnginePluginsPath = FPaths::EnginePluginsDir();
 	}
 
-	bool FTechSoftInterface::InitializeKernel(const TCHAR* InEnginePluginsPath)
+	FString TechSoftDllPath = FPaths::Combine(EnginePluginsPath, TEXT("Enterprise/DatasmithCADImporter"), TEXT("Binaries"), FPlatformProcess::GetBinariesSubdirectory(), "TechSoft");
+	TechSoftDllPath = FPaths::ConvertRelativePathToFull(TechSoftDllPath);
+	ExchangeLoader = MakeUnique<A3DSDKHOOPSExchangeLoader>(*TechSoftDllPath);
+
+	const A3DStatus IRet = ExchangeLoader->m_eSDKStatus;
+	if (IRet != A3D_SUCCESS)
 	{
-#ifdef USE_TECHSOFT_SDK
-		if (bIsInitialize)
-		{
-			return true;
-		}
-
-		FString EnginePluginsPath(InEnginePluginsPath);
-		if (EnginePluginsPath.IsEmpty())
-		{
-			EnginePluginsPath = FPaths::EnginePluginsDir();
-		}
-
-		FString TechSoftDllPath = FPaths::Combine(EnginePluginsPath, TEXT("Enterprise/DatasmithCADImporter"), TEXT("Binaries"), FPlatformProcess::GetBinariesSubdirectory(), "TechSoft");
-		TechSoftDllPath = FPaths::ConvertRelativePathToFull(TechSoftDllPath);
-		ExchangeLoader = MakeUnique<A3DSDKHOOPSExchangeLoader>(*TechSoftDllPath);
-
-		const A3DStatus IRet = ExchangeLoader->m_eSDKStatus;
-		if (IRet != A3D_SUCCESS)
-		{
-			UE_LOG(LogCADInterfaces, Warning, TEXT("Failed to load required library in %s. Plug-in will not be functional."), *TechSoftDllPath);
-		}
-		else
-		{
-			bIsInitialize = true;
-		}
-		return bIsInitialize;
+		UE_LOG(LogCADInterfaces, Warning, TEXT("Failed to load required library in %s. Plug-in will not be functional."), *TechSoftDllPath);
+	}
+	else
+	{
+		bIsInitialize = true;
+	}
+	return bIsInitialize;
 #else
-		return false;
+	return false;
 #endif
-	}
+}
 
 #ifdef USE_TECHSOFT_SDK
 
-	A3DStatus FTechSoftInterface::Import(const A3DImport& Import)
-	{
-		return ExchangeLoader->Import(Import);
-	}
+A3DStatus FTechSoftInterface::Import(const A3DImport& Import)
+{
+	return ExchangeLoader->Import(Import);
+}
 
-	A3DAsmModelFile* FTechSoftInterface::GetModelFile()
-	{
-		return ExchangeLoader->m_psModelFile;
-	}
+A3DAsmModelFile* FTechSoftInterface::GetModelFile()
+{
+	return ExchangeLoader->m_psModelFile;
+}
 
 
-	void TUniqueTSObj<A3DAsmModelFileData>::InitializeData()
-	{
-		A3D_INITIALIZE_DATA(A3DAsmModelFileData, Data);
-	}
+void TUniqueTSObj<A3DAsmModelFileData>::InitializeData()
+{
+	A3D_INITIALIZE_DATA(A3DAsmModelFileData, Data);
+}
 
-	void TUniqueTSObj<A3DAsmPartDefinitionData>::InitializeData()
-	{
-		A3D_INITIALIZE_DATA(A3DAsmPartDefinitionData, Data);
-	}
+void TUniqueTSObj<A3DAsmPartDefinitionData>::InitializeData()
+{
+	A3D_INITIALIZE_DATA(A3DAsmPartDefinitionData, Data);
+}
 
-	void TUniqueTSObj<A3DAsmProductOccurrenceData>::InitializeData()
-	{
-		A3D_INITIALIZE_DATA(A3DAsmProductOccurrenceData, Data);
-	}
+void TUniqueTSObj<A3DAsmProductOccurrenceData>::InitializeData()
+{
+	A3D_INITIALIZE_DATA(A3DAsmProductOccurrenceData, Data);
+}
 
-	void TUniqueTSObj<A3DAsmProductOccurrenceDataCV5>::InitializeData()
-	{
-		A3D_INITIALIZE_DATA(A3DAsmProductOccurrenceDataCV5, Data);
-	}
+void TUniqueTSObj<A3DAsmProductOccurrenceDataCV5>::InitializeData()
+{
+	A3D_INITIALIZE_DATA(A3DAsmProductOccurrenceDataCV5, Data);
+}
 
-	void TUniqueTSObj<A3DAsmProductOccurrenceDataSLW>::InitializeData()
-	{
-		A3D_INITIALIZE_DATA(A3DAsmProductOccurrenceDataSLW, Data);
-	}
+void TUniqueTSObj<A3DAsmProductOccurrenceDataSLW>::InitializeData()
+{
+	A3D_INITIALIZE_DATA(A3DAsmProductOccurrenceDataSLW, Data);
+}
 
-	void TUniqueTSObj<A3DAsmProductOccurrenceDataUg>::InitializeData()
-	{
-		A3D_INITIALIZE_DATA(A3DAsmProductOccurrenceDataUg, Data);
-	}
+void TUniqueTSObj<A3DAsmProductOccurrenceDataUg>::InitializeData()
+{
+	A3D_INITIALIZE_DATA(A3DAsmProductOccurrenceDataUg, Data);
+}
 
-	void TUniqueTSObj<A3DGraphicsData>::InitializeData()
-	{
-		A3D_INITIALIZE_DATA(A3DGraphicsData, Data);
-	}
+void TUniqueTSObj<A3DGraphicsData>::InitializeData()
+{
+	A3D_INITIALIZE_DATA(A3DGraphicsData, Data);
+}
 
-	void TUniqueTSObjFromIndex<A3DGraphMaterialData>::InitializeData()
-	{
-		A3D_INITIALIZE_DATA(A3DGraphMaterialData, Data);
-	}
+void TUniqueTSObjFromIndex<A3DGraphMaterialData>::InitializeData()
+{
+	A3D_INITIALIZE_DATA(A3DGraphMaterialData, Data);
+}
 
-	void TUniqueTSObjFromIndex<A3DGraphRgbColorData>::InitializeData()
-	{
-		A3D_INITIALIZE_DATA(A3DGraphRgbColorData, Data);
-	}
+void TUniqueTSObjFromIndex<A3DGraphRgbColorData>::InitializeData()
+{
+	A3D_INITIALIZE_DATA(A3DGraphRgbColorData, Data);
+}
 
-	void TUniqueTSObjFromIndex<A3DGraphStyleData>::InitializeData()
-	{
-		A3D_INITIALIZE_DATA(A3DGraphStyleData, Data);
-	}
+void TUniqueTSObjFromIndex<A3DGraphStyleData>::InitializeData()
+{
+	A3D_INITIALIZE_DATA(A3DGraphStyleData, Data);
+}
 
-	void TUniqueTSObj<A3DMiscAttributeData>::InitializeData()
-	{
-		A3D_INITIALIZE_DATA(A3DMiscAttributeData, Data);
-	}
+void TUniqueTSObj<A3DMiscAttributeData>::InitializeData()
+{
+	A3D_INITIALIZE_DATA(A3DMiscAttributeData, Data);
+}
 
-	void TUniqueTSObj<A3DMiscCartesianTransformationData>::InitializeData()
-	{
-		A3D_INITIALIZE_DATA(A3DMiscCartesianTransformationData, Data);
-	}
+void TUniqueTSObj<A3DMiscCartesianTransformationData>::InitializeData()
+{
+	A3D_INITIALIZE_DATA(A3DMiscCartesianTransformationData, Data);
+}
 
-	void TUniqueTSObj<A3DMiscGeneralTransformationData>::InitializeData()
-	{
-		A3D_INITIALIZE_DATA(A3DMiscGeneralTransformationData, Data);
-	}
+void TUniqueTSObj<A3DMiscGeneralTransformationData>::InitializeData()
+{
+	A3D_INITIALIZE_DATA(A3DMiscGeneralTransformationData, Data);
+}
 
-	void TUniqueTSObj<A3DMiscMaterialPropertiesData>::InitializeData()
-	{
-		A3D_INITIALIZE_DATA(A3DMiscMaterialPropertiesData, Data);
-	}
+void TUniqueTSObj<A3DMiscMaterialPropertiesData>::InitializeData()
+{
+	A3D_INITIALIZE_DATA(A3DMiscMaterialPropertiesData, Data);
+}
 
-	void TUniqueTSObj<A3DRiBrepModelData>::InitializeData()
-	{
-		A3D_INITIALIZE_DATA(A3DRiBrepModelData, Data);
-	}
+void TUniqueTSObj<A3DRiBrepModelData>::InitializeData()
+{
+	A3D_INITIALIZE_DATA(A3DRiBrepModelData, Data);
+}
 
-	void TUniqueTSObj<A3DRiCoordinateSystemData>::InitializeData()
-	{
-		A3D_INITIALIZE_DATA(A3DRiCoordinateSystemData, Data);
-	}
+void TUniqueTSObj<A3DRiCoordinateSystemData>::InitializeData()
+{
+	A3D_INITIALIZE_DATA(A3DRiCoordinateSystemData, Data);
+}
 
-	void TUniqueTSObj<A3DRiDirectionData>::InitializeData()
-	{
-		A3D_INITIALIZE_DATA(A3DRiDirectionData, Data);
-	}
+void TUniqueTSObj<A3DRiDirectionData>::InitializeData()
+{
+	A3D_INITIALIZE_DATA(A3DRiDirectionData, Data);
+}
 
-	void TUniqueTSObj<A3DRiPolyBrepModelData>::InitializeData()
-	{
-		A3D_INITIALIZE_DATA(A3DRiPolyBrepModelData, Data);
-	}
+void TUniqueTSObj<A3DRiPolyBrepModelData>::InitializeData()
+{
+	A3D_INITIALIZE_DATA(A3DRiPolyBrepModelData, Data);
+}
 
-	void TUniqueTSObj<A3DRiRepresentationItemData>::InitializeData()
-	{
-		A3D_INITIALIZE_DATA(A3DRiRepresentationItemData, Data);
-	}
+void TUniqueTSObj<A3DRiRepresentationItemData>::InitializeData()
+{
+	A3D_INITIALIZE_DATA(A3DRiRepresentationItemData, Data);
+}
 
-	void TUniqueTSObj<A3DRiSetData>::InitializeData()
-	{
-		A3D_INITIALIZE_DATA(A3DRiSetData, Data);
-	}
+void TUniqueTSObj<A3DRiSetData>::InitializeData()
+{
+	A3D_INITIALIZE_DATA(A3DRiSetData, Data);
+}
 
-	void TUniqueTSObj<A3DRootBaseData>::InitializeData()
-	{
-		A3D_INITIALIZE_DATA(A3DRootBaseData, Data);
-	}
+void TUniqueTSObj<A3DRootBaseData>::InitializeData()
+{
+	A3D_INITIALIZE_DATA(A3DRootBaseData, Data);
+}
 
-	void TUniqueTSObj<A3DRootBaseWithGraphicsData>::InitializeData()
-	{
-		A3D_INITIALIZE_DATA(A3DRootBaseWithGraphicsData, Data);
-	}
+void TUniqueTSObj<A3DRootBaseWithGraphicsData>::InitializeData()
+{
+	A3D_INITIALIZE_DATA(A3DRootBaseWithGraphicsData, Data);
+}
 
-	void TUniqueTSObj<A3DTess3DData>::InitializeData()
-	{
-		A3D_INITIALIZE_DATA(A3DTess3DData, Data);
-	}
+void TUniqueTSObj<A3DTess3DData>::InitializeData()
+{
+	A3D_INITIALIZE_DATA(A3DTess3DData, Data);
+}
 
-	void TUniqueTSObj<A3DTessBaseData>::InitializeData()
-	{
-		A3D_INITIALIZE_DATA(A3DTessBaseData, Data);
-	}
+void TUniqueTSObj<A3DTessBaseData>::InitializeData()
+{
+	A3D_INITIALIZE_DATA(A3DTessBaseData, Data);
+}
 
-	void TUniqueTSObj<A3DTopoBrepDataData>::InitializeData()
-	{
-		A3D_INITIALIZE_DATA(A3DTopoBrepDataData, Data);
-	}
+void TUniqueTSObj<A3DTopoBrepDataData>::InitializeData()
+{
+	A3D_INITIALIZE_DATA(A3DTopoBrepDataData, Data);
+}
 
-	void TUniqueTSObj<A3DTopoEdgeData>::InitializeData()
-	{
-		A3D_INITIALIZE_DATA(A3DTopoEdgeData, Data);
-	}
+void TUniqueTSObj<A3DTopoEdgeData>::InitializeData()
+{
+	A3D_INITIALIZE_DATA(A3DTopoEdgeData, Data);
+}
 
-	void TUniqueTSObj<A3DTopoFaceData>::InitializeData()
-	{
-		A3D_INITIALIZE_DATA(A3DTopoFaceData, Data);
-	}
+void TUniqueTSObj<A3DTopoFaceData>::InitializeData()
+{
+	A3D_INITIALIZE_DATA(A3DTopoFaceData, Data);
+}
 
-	void TUniqueTSObj<A3DTopoUniqueVertexData>::InitializeData()
-	{
-		A3D_INITIALIZE_DATA(A3DTopoUniqueVertexData, Data);
-	}
+void TUniqueTSObj<A3DTopoUniqueVertexData>::InitializeData()
+{
+	A3D_INITIALIZE_DATA(A3DTopoUniqueVertexData, Data);
+}
 
-	void TUniqueTSObj<A3DTopoWireEdgeData>::InitializeData()
-	{
-		A3D_INITIALIZE_DATA(A3DTopoWireEdgeData, Data);
-	}
+void TUniqueTSObj<A3DTopoWireEdgeData>::InitializeData()
+{
+	A3D_INITIALIZE_DATA(A3DTopoWireEdgeData, Data);
+}
 
-	A3DStatus TUniqueTSObj<A3DAsmModelFileData>::GetData(const A3DAsmModelFile* InAsmModelFilePtr)
-	{
-		return A3DAsmModelFileGet(InAsmModelFilePtr, &Data);
-	}
+A3DStatus TUniqueTSObj<A3DAsmModelFileData>::GetData(const A3DAsmModelFile* InAsmModelFilePtr)
+{
+	return A3DAsmModelFileGet(InAsmModelFilePtr, &Data);
+}
 
-	A3DStatus TUniqueTSObj<A3DAsmPartDefinitionData>::GetData(const A3DAsmPartDefinition* InAsmPartDefinitionPtr)
-	{
-		return A3DAsmPartDefinitionGet(InAsmPartDefinitionPtr, &Data);
-	}
+A3DStatus TUniqueTSObj<A3DAsmPartDefinitionData>::GetData(const A3DAsmPartDefinition* InAsmPartDefinitionPtr)
+{
+	return A3DAsmPartDefinitionGet(InAsmPartDefinitionPtr, &Data);
+}
 
-	A3DStatus TUniqueTSObj<A3DAsmProductOccurrenceData>::GetData(const A3DAsmProductOccurrence* InAsmProductOccurrencePtr)
-	{
-		return A3DAsmProductOccurrenceGet(InAsmProductOccurrencePtr, &Data);
-	}
+A3DStatus TUniqueTSObj<A3DAsmProductOccurrenceData>::GetData(const A3DAsmProductOccurrence* InAsmProductOccurrencePtr)
+{
+	return A3DAsmProductOccurrenceGet(InAsmProductOccurrencePtr, &Data);
+}
 
-	A3DStatus TUniqueTSObj<A3DAsmProductOccurrenceDataCV5>::GetData(const A3DAsmProductOccurrence* InAsmProductOccurrencePtr)
-	{
-		return A3DAsmProductOccurrenceGetCV5(InAsmProductOccurrencePtr, &Data);
-	}
+A3DStatus TUniqueTSObj<A3DAsmProductOccurrenceDataCV5>::GetData(const A3DAsmProductOccurrence* InAsmProductOccurrencePtr)
+{
+	return A3DAsmProductOccurrenceGetCV5(InAsmProductOccurrencePtr, &Data);
+}
 
-	A3DStatus TUniqueTSObj<A3DAsmProductOccurrenceDataSLW>::GetData(const A3DAsmProductOccurrence* InAsmProductOccurrencePtr)
-	{
-		return A3DAsmProductOccurrenceGetSLW(InAsmProductOccurrencePtr, &Data);
-	}
+A3DStatus TUniqueTSObj<A3DAsmProductOccurrenceDataSLW>::GetData(const A3DAsmProductOccurrence* InAsmProductOccurrencePtr)
+{
+	return A3DAsmProductOccurrenceGetSLW(InAsmProductOccurrencePtr, &Data);
+}
 
-	A3DStatus TUniqueTSObj<A3DAsmProductOccurrenceDataUg>::GetData(const A3DAsmProductOccurrence* InAsmProductOccurrencePtr)
-	{
-		return A3DAsmProductOccurrenceGetUg(InAsmProductOccurrencePtr, &Data);
-	}
+A3DStatus TUniqueTSObj<A3DAsmProductOccurrenceDataUg>::GetData(const A3DAsmProductOccurrence* InAsmProductOccurrencePtr)
+{
+	return A3DAsmProductOccurrenceGetUg(InAsmProductOccurrencePtr, &Data);
+}
 
-	A3DStatus TUniqueTSObj<A3DGraphicsData>::GetData(const A3DGraphics* InGraphicsPtr)
-	{
-		return A3DGraphicsGet(InGraphicsPtr, &Data);
-	}
+A3DStatus TUniqueTSObj<A3DGraphicsData>::GetData(const A3DGraphics* InGraphicsPtr)
+{
+	return A3DGraphicsGet(InGraphicsPtr, &Data);
+}
 
-	A3DStatus TUniqueTSObjFromIndex<A3DGraphMaterialData>::GetData(uint32 InGraphicsIndex)
-	{
-		return A3DGlobalGetGraphMaterialData(InGraphicsIndex, &Data);
-	}
+A3DStatus TUniqueTSObjFromIndex<A3DGraphMaterialData>::GetData(uint32 InGraphicsIndex)
+{
+	return A3DGlobalGetGraphMaterialData(InGraphicsIndex, &Data);
+}
 
-	A3DStatus TUniqueTSObjFromIndex<A3DGraphRgbColorData>::GetData(uint32 ObjectIndex)
-	{
-		return A3DGlobalGetGraphRgbColorData(ObjectIndex, &Data);
-	}
+A3DStatus TUniqueTSObjFromIndex<A3DGraphRgbColorData>::GetData(uint32 ObjectIndex)
+{
+	return A3DGlobalGetGraphRgbColorData(ObjectIndex, &Data);
+}
 
-	A3DStatus TUniqueTSObjFromIndex<A3DGraphStyleData>::GetData(uint32 ObjectIndex)
-	{
-		return A3DGlobalGetGraphStyleData(ObjectIndex, &Data);
-	}
+A3DStatus TUniqueTSObjFromIndex<A3DGraphStyleData>::GetData(uint32 ObjectIndex)
+{
+	return A3DGlobalGetGraphStyleData(ObjectIndex, &Data);
+}
 
-	A3DStatus TUniqueTSObj<A3DMiscAttributeData>::GetData(const A3DMiscAttribute* InMiscAttributePtr)
-	{
-		return A3DMiscAttributeGet(InMiscAttributePtr, &Data);
-	}
+A3DStatus TUniqueTSObj<A3DMiscAttributeData>::GetData(const A3DMiscAttribute* InMiscAttributePtr)
+{
+	return A3DMiscAttributeGet(InMiscAttributePtr, &Data);
+}
 
-	A3DStatus TUniqueTSObj<A3DMiscCartesianTransformationData>::GetData(const A3DMiscCartesianTransformation* InMiscCartesianTransformationPtr)
-	{
-		return A3DMiscCartesianTransformationGet(InMiscCartesianTransformationPtr, &Data);
-	}
+A3DStatus TUniqueTSObj<A3DMiscCartesianTransformationData>::GetData(const A3DMiscCartesianTransformation* InMiscCartesianTransformationPtr)
+{
+	return A3DMiscCartesianTransformationGet(InMiscCartesianTransformationPtr, &Data);
+}
 
-	A3DStatus TUniqueTSObj<A3DMiscGeneralTransformationData>::GetData(const A3DMiscGeneralTransformation* InMiscGeneralTransformationPtr)
-	{
-		return A3DMiscGeneralTransformationGet(InMiscGeneralTransformationPtr, &Data);
-	}
+A3DStatus TUniqueTSObj<A3DMiscGeneralTransformationData>::GetData(const A3DMiscGeneralTransformation* InMiscGeneralTransformationPtr)
+{
+	return A3DMiscGeneralTransformationGet(InMiscGeneralTransformationPtr, &Data);
+}
 
-	A3DStatus TUniqueTSObj<A3DMiscMaterialPropertiesData>::GetData(const A3DEntity* InMiscMaterialPropertiesPtr)
-	{
-		return A3DMiscGetMaterialProperties(InMiscMaterialPropertiesPtr, &Data);
-	}
+A3DStatus TUniqueTSObj<A3DMiscMaterialPropertiesData>::GetData(const A3DEntity* InMiscMaterialPropertiesPtr)
+{
+	return A3DMiscGetMaterialProperties(InMiscMaterialPropertiesPtr, &Data);
+}
 
-	A3DStatus TUniqueTSObj<A3DRiBrepModelData>::GetData(const A3DRiBrepModel* InRiBrepModelPtr)
-	{
-		return A3DRiBrepModelGet(InRiBrepModelPtr, &Data);
-	}
+A3DStatus TUniqueTSObj<A3DRiBrepModelData>::GetData(const A3DRiBrepModel* InRiBrepModelPtr)
+{
+	return A3DRiBrepModelGet(InRiBrepModelPtr, &Data);
+}
 
-	A3DStatus TUniqueTSObj<A3DRiCoordinateSystemData>::GetData(const A3DRiCoordinateSystem* InRiCoordinateSystemPtr)
-	{
-		return A3DRiCoordinateSystemGet(InRiCoordinateSystemPtr, &Data);
-	}
+A3DStatus TUniqueTSObj<A3DRiCoordinateSystemData>::GetData(const A3DRiCoordinateSystem* InRiCoordinateSystemPtr)
+{
+	return A3DRiCoordinateSystemGet(InRiCoordinateSystemPtr, &Data);
+}
 
-	A3DStatus TUniqueTSObj<A3DRiDirectionData>::GetData(const A3DRiDirection* InRiDirectionPtr)
-	{
-		return A3DRiDirectionGet(InRiDirectionPtr, &Data);
-	}
+A3DStatus TUniqueTSObj<A3DRiDirectionData>::GetData(const A3DRiDirection* InRiDirectionPtr)
+{
+	return A3DRiDirectionGet(InRiDirectionPtr, &Data);
+}
 
-	A3DStatus TUniqueTSObj<A3DRiPolyBrepModelData>::GetData(const A3DRiPolyBrepModel* InRiPolyBrepModelPtr)
-	{
-		return A3DRiPolyBrepModelGet(InRiPolyBrepModelPtr, &Data);
-	}
+A3DStatus TUniqueTSObj<A3DRiPolyBrepModelData>::GetData(const A3DRiPolyBrepModel* InRiPolyBrepModelPtr)
+{
+	return A3DRiPolyBrepModelGet(InRiPolyBrepModelPtr, &Data);
+}
 
-	A3DStatus TUniqueTSObj<A3DRiRepresentationItemData>::GetData(const A3DRiRepresentationItem* InRiRepresentationItemPtr)
-	{
-		return A3DRiRepresentationItemGet(InRiRepresentationItemPtr, &Data);
-	}
+A3DStatus TUniqueTSObj<A3DRiRepresentationItemData>::GetData(const A3DRiRepresentationItem* InRiRepresentationItemPtr)
+{
+	return A3DRiRepresentationItemGet(InRiRepresentationItemPtr, &Data);
+}
 
-	A3DStatus TUniqueTSObj<A3DRiSetData>::GetData(const A3DRiSet* InRiSetPtr)
-	{
-		return A3DRiSetGet(InRiSetPtr, &Data);
-	}
+A3DStatus TUniqueTSObj<A3DRiSetData>::GetData(const A3DRiSet* InRiSetPtr)
+{
+	return A3DRiSetGet(InRiSetPtr, &Data);
+}
 
-	A3DStatus TUniqueTSObj<A3DRootBaseData>::GetData(const A3DRootBase* InRootBasePtr)
-	{
-		return A3DRootBaseGet(InRootBasePtr, &Data);
-	}
+A3DStatus TUniqueTSObj<A3DRootBaseData>::GetData(const A3DRootBase* InRootBasePtr)
+{
+	return A3DRootBaseGet(InRootBasePtr, &Data);
+}
 
-	A3DStatus TUniqueTSObj<A3DRootBaseWithGraphicsData>::GetData(const A3DRootBaseWithGraphics* InRootBaseWithGraphicsPtr)
-	{
-		return A3DRootBaseWithGraphicsGet(InRootBaseWithGraphicsPtr, &Data);
-	}
+A3DStatus TUniqueTSObj<A3DRootBaseWithGraphicsData>::GetData(const A3DRootBaseWithGraphics* InRootBaseWithGraphicsPtr)
+{
+	return A3DRootBaseWithGraphicsGet(InRootBaseWithGraphicsPtr, &Data);
+}
 
-	A3DStatus TUniqueTSObj<A3DTess3DData>::GetData(const A3DTess3D* InTess3DPtr)
-	{
-		return A3DTess3DGet(InTess3DPtr, &Data);
-	}
+A3DStatus TUniqueTSObj<A3DTess3DData>::GetData(const A3DTess3D* InTess3DPtr)
+{
+	return A3DTess3DGet(InTess3DPtr, &Data);
+}
 
-	A3DStatus TUniqueTSObj<A3DTessBaseData>::GetData(const A3DTessBase* InTessBasePtr)
-	{
-		return A3DTessBaseGet(InTessBasePtr, &Data);
-	}
+A3DStatus TUniqueTSObj<A3DTessBaseData>::GetData(const A3DTessBase* InTessBasePtr)
+{
+	return A3DTessBaseGet(InTessBasePtr, &Data);
+}
 
-	A3DStatus TUniqueTSObj<A3DTopoBrepDataData>::GetData(const A3DTopoBrepData* InTopoBrepDataPtr)
-	{
-		return A3DTopoBrepDataGet(InTopoBrepDataPtr, &Data);
-	}
+A3DStatus TUniqueTSObj<A3DTopoBrepDataData>::GetData(const A3DTopoBrepData* InTopoBrepDataPtr)
+{
+	return A3DTopoBrepDataGet(InTopoBrepDataPtr, &Data);
+}
 
-	A3DStatus TUniqueTSObj<A3DTopoEdgeData>::GetData(const A3DTopoEdge* InTopoEdgePtr)	
-	{
-		return A3DTopoEdgeGet(InTopoEdgePtr, &Data);
-	}
+A3DStatus TUniqueTSObj<A3DTopoEdgeData>::GetData(const A3DTopoEdge* InTopoEdgePtr)
+{
+	return A3DTopoEdgeGet(InTopoEdgePtr, &Data);
+}
 
-	A3DStatus TUniqueTSObj<A3DTopoFaceData>::GetData(const A3DTopoFace* InTopoFacePtr)
-	{
-		return A3DTopoFaceGet(InTopoFacePtr, &Data);
-	}
+A3DStatus TUniqueTSObj<A3DTopoFaceData>::GetData(const A3DTopoFace* InTopoFacePtr)
+{
+	return A3DTopoFaceGet(InTopoFacePtr, &Data);
+}
 
-	A3DStatus TUniqueTSObj<A3DTopoUniqueVertexData>::GetData(const A3DTopoUniqueVertex* InTopoUniqueVertexPtr)
-	{
-		return A3DTopoUniqueVertexGet(InTopoUniqueVertexPtr, &Data);
-	}
+A3DStatus TUniqueTSObj<A3DTopoUniqueVertexData>::GetData(const A3DTopoUniqueVertex* InTopoUniqueVertexPtr)
+{
+	return A3DTopoUniqueVertexGet(InTopoUniqueVertexPtr, &Data);
+}
 
-	A3DStatus TUniqueTSObj<A3DTopoWireEdgeData>::GetData(const A3DTopoWireEdge* InTopoWireEdgePtr)
-	{
-		return A3DTopoWireEdgeGet(InTopoWireEdgePtr, &Data);
-	}
+A3DStatus TUniqueTSObj<A3DTopoWireEdgeData>::GetData(const A3DTopoWireEdge* InTopoWireEdgePtr)
+{
+	return A3DTopoWireEdgeGet(InTopoWireEdgePtr, &Data);
+}
 
 
-	uint32 TUniqueTSObjFromIndex<A3DGraphMaterialData>::DefaultValue = A3D_DEFAULT_MATERIAL_INDEX;
-	uint32 TUniqueTSObjFromIndex<A3DGraphPictureData>::DefaultValue = A3D_DEFAULT_MATERIAL_INDEX;
-	uint32 TUniqueTSObjFromIndex<A3DGraphRgbColorData>::DefaultValue = A3D_DEFAULT_COLOR_INDEX;
-	uint32 TUniqueTSObjFromIndex<A3DGraphStyleData>::DefaultValue = A3D_DEFAULT_STYLE_INDEX;
+uint32 TUniqueTSObjFromIndex<A3DGraphMaterialData>::DefaultValue = A3D_DEFAULT_MATERIAL_INDEX;
+uint32 TUniqueTSObjFromIndex<A3DGraphPictureData>::DefaultValue = A3D_DEFAULT_MATERIAL_INDEX;
+uint32 TUniqueTSObjFromIndex<A3DGraphRgbColorData>::DefaultValue = A3D_DEFAULT_COLOR_INDEX;
+uint32 TUniqueTSObjFromIndex<A3DGraphStyleData>::DefaultValue = A3D_DEFAULT_STYLE_INDEX;
 
 #endif
 
diff --git a/Engine/Plugins/Enterprise/DatasmithCADImporter/Source/CADInterfaces/Public/CADFileData.h b/Engine/Plugins/Enterprise/DatasmithCADImporter/Source/CADInterfaces/Public/CADFileData.h
--- a/Engine/Plugins/Enterprise/DatasmithCADImporter/Source/CADInterfaces/Public/CADFileData.h
+++ b/Engine/Plugins/Enterprise/DatasmithCADImporter/Source/CADInterfaces/Public/CADFileData.h
@@ -9,322 +9,323 @@
 
 namespace CADLibrary
 {
-	class FCADFileData
-	{
-	public:
-		FCADFileData(const FImportParameters& InImportParameters, const FFileDescriptor& InFileDescription, const FString& InCachePath)
-			: ImportParameters(InImportParameters)
-			, CachePath(InCachePath)
-			, bIsCacheDefined(!InCachePath.IsEmpty())
-			, FileDescription(InFileDescription)
-		{
-			SceneGraphArchive.FullPath = FileDescription.GetSourcePath();
-			SceneGraphArchive.CADFileName = FileDescription.GetFileName();
-		}
 
-		uint32 GetSceneFileHash() const
-		{
-			if (!SceneFileHash)
-			{
-				SceneFileHash = HashCombine(FileDescription.GetDescriptorHash(), ::GetTypeHash(ImportParameters.GetStitchingTechnique()));
-				SceneFileHash = HashCombine(SceneFileHash, ::GetTypeHash(CADLibrary::GCADLibrary));
-			}
-			return SceneFileHash;
-		}
+class FCADFileData
+{
+public:
+	FCADFileData(const FImportParameters& InImportParameters, const FFileDescriptor& InFileDescription, const FString& InCachePath)
+		: ImportParameters(InImportParameters)
+		, CachePath(InCachePath)
+		, bIsCacheDefined(!InCachePath.IsEmpty())
+		, FileDescription(InFileDescription)
+	{
+		SceneGraphArchive.FullPath = FileDescription.GetSourcePath();
+		SceneGraphArchive.CADFileName = FileDescription.GetFileName();
+	}
 
-		uint32 GetGeomFileHash() const
+	uint32 GetSceneFileHash() const
+	{
+		if (!SceneFileHash)
 		{
-			if (!GeomFileHash)
-			{
-				GeomFileHash = GetSceneFileHash();
-				GeomFileHash = HashCombine(GeomFileHash, GetTypeHash(ImportParameters));
-			}
-			return GeomFileHash;
+			SceneFileHash = HashCombine(FileDescription.GetDescriptorHash(), ::GetTypeHash(ImportParameters.GetStitchingTechnique()));
+			SceneFileHash = HashCombine(SceneFileHash, ::GetTypeHash(CADLibrary::GCADLibrary));
 		}
+		return SceneFileHash;
+	}
 
-		void SetArchiveNames()
+	uint32 GetGeomFileHash() const
+	{
+		if (!GeomFileHash)
 		{
-			SceneGraphArchive.ArchiveFileName = FString::Printf(TEXT("UEx%08x"), GetSceneFileHash());
-			MeshArchiveFile = FString::Printf(TEXT("UEx%08x"), GetGeomFileHash());
+			GeomFileHash = GetSceneFileHash();
+			GeomFileHash = HashCombine(GeomFileHash, GetTypeHash(ImportParameters));
 		}
+		return GeomFileHash;
+	}
 
-		const FString GetSceneGraphFilePath() const
-		{
-			if (IsCacheDefined())
-			{
-				return FPaths::Combine(CachePath, TEXT("scene"), SceneGraphArchive.ArchiveFileName + TEXT(".sg"));
-			}
-			return FString();
-		}
+	void SetArchiveNames()
+	{
+		SceneGraphArchive.ArchiveFileName = FString::Printf(TEXT("UEx%08x"), GetSceneFileHash());
+		MeshArchiveFile = FString::Printf(TEXT("UEx%08x"), GetGeomFileHash());
+	}
 
-		const FString GetMeshArchiveFilePath() const
+	const FString GetSceneGraphFilePath() const
+	{
+		if (IsCacheDefined())
 		{
-			if (IsCacheDefined())
-			{
-				return FPaths::Combine(CachePath, TEXT("mesh"), MeshArchiveFile + TEXT(".gm"));
-			}
-			return FString();
+			return FPaths::Combine(CachePath, TEXT("scene"), SceneGraphArchive.ArchiveFileName + TEXT(".sg"));
 		}
+		return FString();
+	}
 
-		const FString GetBodyCachePath(uint32 BodyHash) const
+	const FString GetMeshArchiveFilePath() const
+	{
+		if (IsCacheDefined())
 		{
-			if (IsCacheDefined())
-			{
-				return CADLibrary::BuildCacheFilePath(*CachePath, TEXT("body"), BodyHash);
-			}
-			return FString();
+			return FPaths::Combine(CachePath, TEXT("mesh"), MeshArchiveFile + TEXT(".gm"));
 		}
+		return FString();
+	}
 
-		/**
-		 * @return the path of CAD cache file
-		 */
-		const FString GetCADCachePath() const
+	const FString GetBodyCachePath(uint32 BodyHash) const
+	{
+		if (IsCacheDefined())
 		{
-			if (IsCacheDefined())
-			{
-				return CADLibrary::BuildCacheFilePath(*CachePath, TEXT("cad"), FileDescription.GetDescriptorHash());
-			}
-			return FString();
+			return CADLibrary::BuildCacheFilePath(*CachePath, TEXT("body"), BodyHash);
 		}
+		return FString();
+	}
 
-		const FString GetCachePath() const
+	/**
+	 * @return the path of CAD cache file
+	 */
+	const FString GetCADCachePath() const
+	{
+		if (IsCacheDefined())
 		{
-			if (IsCacheDefined())
-			{
-				return CachePath;
-			}
-			return FString();
+			return CADLibrary::BuildCacheFilePath(*CachePath, TEXT("cad"), FileDescription.GetDescriptorHash());
 		}
+		return FString();
+	}
 
-		bool IsCacheDefined() const
+	const FString GetCachePath() const
+	{
+		if (IsCacheDefined())
 		{
-			return bIsCacheDefined;
+			return CachePath;
 		}
+		return FString();
+	}
 
-		void AddWarningMessages(const FString& Message)
-		{
-			WarningMessages.Add(Message);
-		}
+	bool IsCacheDefined() const
+	{
+		return bIsCacheDefined;
+	}
 
-		void LoadSceneGraphArchive()
-		{
-			FString SceneGraphFilePath = GetSceneGraphFilePath();
-			SceneGraphArchive.DeserializeMockUpFile(*SceneGraphFilePath);
-		}
+	void AddWarningMessages(const FString& Message)
+	{
+		WarningMessages.Add(Message);
+	}
 
-		void ExportSceneGraphFile()
-		{
-			FString SceneGraphFilePath = GetSceneGraphFilePath();
-			SceneGraphArchive.SerializeMockUp(*SceneGraphFilePath);
-		}
+	void LoadSceneGraphArchive()
+	{
+		FString SceneGraphFilePath = GetSceneGraphFilePath();
+		SceneGraphArchive.DeserializeMockUpFile(*SceneGraphFilePath);
+	}
 
-		bool HasComponentOfId(FCadId ComponentId) const
-		{
-			return SceneGraphArchive.CADIdToComponentIndex.Find(ComponentId) != nullptr;
-		}
+	void ExportSceneGraphFile()
+	{
+		FString SceneGraphFilePath = GetSceneGraphFilePath();
+		SceneGraphArchive.SerializeMockUp(*SceneGraphFilePath);
+	}
 
+	bool HasComponentOfId(FCadId ComponentId) const
+	{
+		return SceneGraphArchive.CADIdToComponentIndex.Find(ComponentId) != nullptr;
+	}
 
-		int32 AddComponent(FCadId ComponentId)
-		{
-			int32 Index = SceneGraphArchive.Components.Emplace(ComponentId);
-			SceneGraphArchive.CADIdToComponentIndex.Add(ComponentId, Index);
-			return Index;
-		}
 
-		FArchiveComponent& GetComponentAt(int32 Index)
-		{
-			return SceneGraphArchive.Components[Index];
-		}
+	int32 AddComponent(FCadId ComponentId)
+	{
+		int32 Index = SceneGraphArchive.Components.Emplace(ComponentId);
+		SceneGraphArchive.CADIdToComponentIndex.Add(ComponentId, Index);
+		return Index;
+	}
 
+	FArchiveComponent& GetComponentAt(int32 Index)
+	{
+		return SceneGraphArchive.Components[Index];
+	}
 
-		bool HasInstanceOfId(FCadId InstanceId) const
-		{
-			return SceneGraphArchive.CADIdToInstanceIndex.Find(InstanceId) != nullptr;
-		}
 
-		int32 AddInstance(FCadId InstanceId)
-		{
-			int32 Index = SceneGraphArchive.Instances.Emplace(InstanceId);
-			SceneGraphArchive.CADIdToInstanceIndex.Add(InstanceId, Index);
-			return Index;
-		}
+	bool HasInstanceOfId(FCadId InstanceId) const
+	{
+		return SceneGraphArchive.CADIdToInstanceIndex.Find(InstanceId) != nullptr;
+	}
 
-		FArchiveInstance& GetInstanceAt(int32 Index)
-		{
-			return SceneGraphArchive.Instances[Index];
-		}
+	int32 AddInstance(FCadId InstanceId)
+	{
+		int32 Index = SceneGraphArchive.Instances.Emplace(InstanceId);
+		SceneGraphArchive.CADIdToInstanceIndex.Add(InstanceId, Index);
+		return Index;
+	}
 
+	FArchiveInstance& GetInstanceAt(int32 Index)
+	{
+		return SceneGraphArchive.Instances[Index];
+	}
 
-		bool HasBodyOfId(FCadId BodyId) const
-		{
-			return SceneGraphArchive.CADIdToBodyIndex.Find(BodyId) != nullptr;
-		}
 
-		int32 AddBody(FCadId BodyId)
-		{
-			int32 Index = SceneGraphArchive.Bodies.Emplace(BodyId);
-			SceneGraphArchive.CADIdToBodyIndex.Add(BodyId, Index);
-			return Index;
-		}
+	bool HasBodyOfId(FCadId BodyId) const
+	{
+		return SceneGraphArchive.CADIdToBodyIndex.Find(BodyId) != nullptr;
+	}
 
-		FArchiveBody& GetBodyAt(int32 Index)
-		{
-			return SceneGraphArchive.Bodies[Index];
-		}
+	int32 AddBody(FCadId BodyId)
+	{
+		int32 Index = SceneGraphArchive.Bodies.Emplace(BodyId);
+		SceneGraphArchive.CADIdToBodyIndex.Add(BodyId, Index);
+		return Index;
+	}
 
+	FArchiveBody& GetBodyAt(int32 Index)
+	{
+		return SceneGraphArchive.Bodies[Index];
+	}
 
-		int32* FindUnloadedComponentOfId(FCadId ComponentId)
-		{
-			return SceneGraphArchive.CADIdToUnloadedComponentIndex.Find(ComponentId);
-		}
 
-		bool HasUnloadedComponentOfId(FCadId ComponentId)
-		{
-			return FindUnloadedComponentOfId(ComponentId) != nullptr;
-		}
+	int32* FindUnloadedComponentOfId(FCadId ComponentId)
+	{
+		return SceneGraphArchive.CADIdToUnloadedComponentIndex.Find(ComponentId);
+	}
 
-		int32 AddUnloadedComponent(FCadId ComponentId)
-		{
-			int32 Index = SceneGraphArchive.UnloadedComponents.Emplace(ComponentId);
-			SceneGraphArchive.CADIdToUnloadedComponentIndex.Add(ComponentId, Index);
-			return Index;
-		}
+	bool HasUnloadedComponentOfId(FCadId ComponentId)
+	{
+		return FindUnloadedComponentOfId(ComponentId) != nullptr;
+	}
 
-		FArchiveUnloadedComponent& GetUnloadedComponentAt(int32 Index)
-		{
-			return SceneGraphArchive.UnloadedComponents[Index];
-		}
+	int32 AddUnloadedComponent(FCadId ComponentId)
+	{
+		int32 Index = SceneGraphArchive.UnloadedComponents.Emplace(ComponentId);
+		SceneGraphArchive.CADIdToUnloadedComponentIndex.Add(ComponentId, Index);
+		return Index;
+	}
 
-		FFileDescriptor& GetExternalReferences(int32 Index)
-		{
-			return SceneGraphArchive.ExternalReferences[Index];
-		}
+	FArchiveUnloadedComponent& GetUnloadedComponentAt(int32 Index)
+	{
+		return SceneGraphArchive.UnloadedComponents[Index];
+	}
 
-		FFileDescriptor& AddExternalRef(const TCHAR* InFilePath, const TCHAR* InConfiguration, const TCHAR* InRootFilePath)
-		{
-			return SceneGraphArchive.ExternalReferences.Emplace_GetRef(InFilePath, InConfiguration, InRootFilePath);
-		}
+	FFileDescriptor& GetExternalReferences(int32 Index)
+	{
+		return SceneGraphArchive.ExternalReferences[Index];
+	}
 
-		FFileDescriptor& AddExternalRef(const FFileDescriptor& InFileDescription)
-		{
-			return SceneGraphArchive.ExternalReferences.Emplace_GetRef(InFileDescription);
-		}
+	FFileDescriptor& AddExternalRef(const TCHAR* InFilePath, const TCHAR* InConfiguration, const TCHAR* InRootFilePath)
+	{
+		return SceneGraphArchive.ExternalReferences.Emplace_GetRef(InFilePath, InConfiguration, InRootFilePath);
+	}
 
-		/** return a unique value that will be used to define the static mesh name */
-		uint32 GetStaticMeshHash(const int32 BodyId)
-		{
-			return HashCombine(GetSceneFileHash(), ::GetTypeHash(BodyId));
-		}
+	FFileDescriptor& AddExternalRef(const FFileDescriptor& InFileDescription)
+	{
+		return SceneGraphArchive.ExternalReferences.Emplace_GetRef(InFileDescription);
+	}
 
-		FBodyMesh& AddBodyMesh(FCadId BodyId, FArchiveBody& Body)
-		{
-			FBodyMesh& BodyMesh = BodyMeshes.Emplace_GetRef(BodyId);
-			BodyMesh.MeshActorName = GetStaticMeshHash(BodyId);
-			Body.MeshActorName = BodyMesh.MeshActorName;
-			return BodyMesh;
-		}
+	/** return a unique value that will be used to define the static mesh name */
+	uint32 GetStaticMeshHash(const int32 BodyId)
+	{
+		return HashCombine(GetSceneFileHash(), ::GetTypeHash(BodyId));
+	}
 
-		void ExportMeshArchiveFile()
-		{
-			FString MeshArchiveFilePath = GetMeshArchiveFilePath();
-			SerializeBodyMeshSet(*MeshArchiveFilePath, BodyMeshes);
-		}
+	FBodyMesh& AddBodyMesh(FCadId BodyId, FArchiveBody& Body)
+	{
+		FBodyMesh& BodyMesh = BodyMeshes.Emplace_GetRef(BodyId);
+		BodyMesh.MeshActorName = GetStaticMeshHash(BodyId);
+		Body.MeshActorName = BodyMesh.MeshActorName;
+		return BodyMesh;
+	}
 
-		const TArray<FFileDescriptor>& GetExternalRefSet() const 
-		{
-			return SceneGraphArchive.ExternalReferences;
-		}
+	void ExportMeshArchiveFile()
+	{
+		FString MeshArchiveFilePath = GetMeshArchiveFilePath();
+		SerializeBodyMeshSet(*MeshArchiveFilePath, BodyMeshes);
+	}
 
-		const FString& GetSceneGraphFileName() const
-		{
-			return SceneGraphArchive.ArchiveFileName;
-		}
+	const TArray<FFileDescriptor>& GetExternalRefSet() const
+	{
+		return SceneGraphArchive.ExternalReferences;
+	}
 
-		const FString& GetMeshFileName() const
-		{
-			return MeshArchiveFile;
-		}
+	const FString& GetSceneGraphFileName() const
+	{
+		return SceneGraphArchive.ArchiveFileName;
+	}
 
-		const TArray<FString>& GetWarningMessages() const
-		{
-			return WarningMessages;
-		}
+	const FString& GetMeshFileName() const
+	{
+		return MeshArchiveFile;
+	}
 
-		const FArchiveSceneGraph& GetSceneGraphArchive() const
-		{
-			return SceneGraphArchive;
-		}
+	const TArray<FString>& GetWarningMessages() const
+	{
+		return WarningMessages;
+	}
 
-		FArchiveSceneGraph& GetSceneGraphArchive()
-		{
-			return SceneGraphArchive;
-		}
+	const FArchiveSceneGraph& GetSceneGraphArchive() const
+	{
+		return SceneGraphArchive;
+	}
 
-		FArchiveMaterial* FindMaterial(uint32 MaterialId)
-		{
-			return SceneGraphArchive.MaterialHIdToMaterial.Find(MaterialId);
-		}
+	FArchiveSceneGraph& GetSceneGraphArchive()
+	{
+		return SceneGraphArchive;
+	}
 
-		FArchiveMaterial& AddMaterial(uint32 MaterialId)
-		{
-			return SceneGraphArchive.MaterialHIdToMaterial.Emplace(MaterialId, MaterialId);
-		}
+	FArchiveMaterial* FindMaterial(uint32 MaterialId)
+	{
+		return SceneGraphArchive.MaterialHIdToMaterial.Find(MaterialId);
+	}
 
-		FArchiveColor* FindColor(uint32 ColorId)
-		{
-			return SceneGraphArchive.ColorHIdToColor.Find(ColorId);
-		}
+	FArchiveMaterial& AddMaterial(uint32 MaterialId)
+	{
+		return SceneGraphArchive.MaterialHIdToMaterial.Emplace(MaterialId, MaterialId);
+	}
 
-		FArchiveColor& AddColor(uint32 ColorId)
-		{
-			return SceneGraphArchive.ColorHIdToColor.Emplace(ColorId, ColorId);
-		}
+	FArchiveColor* FindColor(uint32 ColorId)
+	{
+		return SceneGraphArchive.ColorHIdToColor.Find(ColorId);
+	}
 
-		const TArray<FBodyMesh>& GetBodyMeshes() const
-		{
-			return BodyMeshes;
-		}
+	FArchiveColor& AddColor(uint32 ColorId)
+	{
+		return SceneGraphArchive.ColorHIdToColor.Emplace(ColorId, ColorId);
+	}
 
-		TArray<FBodyMesh>& GetBodyMeshes()
-		{
-			return BodyMeshes;
-		}
+	const TArray<FBodyMesh>& GetBodyMeshes() const
+	{
+		return BodyMeshes;
+	}
 
-		const FFileDescriptor& GetCADFileDescription() const
-		{
-			return FileDescription;
-		}
+	TArray<FBodyMesh>& GetBodyMeshes()
+	{
+		return BodyMeshes;
+	}
 
-		FFileDescriptor& GetCADFileDescription()
-		{
-			return FileDescription;
-		}
+	const FFileDescriptor& GetCADFileDescription() const
+	{
+		return FileDescription;
+	}
 
-		void ReserveBodyMeshes(int32 MaxBodyCount)
-		{
-			BodyMeshes.Reserve(MaxBodyCount);
-		}
+	FFileDescriptor& GetCADFileDescription()
+	{
+		return FileDescription;
+	}
 
-		const FImportParameters& GetImportParameters() const
-		{
-			return ImportParameters;
-		}
+	void ReserveBodyMeshes(int32 MaxBodyCount)
+	{
+		BodyMeshes.Reserve(MaxBodyCount);
+	}
+
+	const FImportParameters& GetImportParameters() const
+	{
+		return ImportParameters;
+	}
 
-	private:
-		const FImportParameters ImportParameters;
-		const FString CachePath;
-		const bool bIsCacheDefined;
+private:
+	const FImportParameters ImportParameters;
+	const FString CachePath;
+	const bool bIsCacheDefined;
 
-		FFileDescriptor FileDescription;
+	FFileDescriptor FileDescription;
 
-		FString MeshArchiveFile;
+	FString MeshArchiveFile;
 
-		FArchiveSceneGraph SceneGraphArchive;
-		TArray<FBodyMesh> BodyMeshes;
+	FArchiveSceneGraph SceneGraphArchive;
+	TArray<FBodyMesh> BodyMeshes;
 
-		TArray<FString> WarningMessages;
+	TArray<FString> WarningMessages;
 
-		mutable uint32 SceneFileHash = 0;
-		mutable uint32 GeomFileHash = 0;
-	};
+	mutable uint32 SceneFileHash = 0;
+	mutable uint32 GeomFileHash = 0;
+};
 }
\ No newline at end of file
diff --git a/Engine/Plugins/Enterprise/DatasmithCADImporter/Source/CADInterfaces/Public/TechSoftInterface.h b/Engine/Plugins/Enterprise/DatasmithCADImporter/Source/CADInterfaces/Public/TechSoftInterface.h
--- a/Engine/Plugins/Enterprise/DatasmithCADImporter/Source/CADInterfaces/Public/TechSoftInterface.h
+++ b/Engine/Plugins/Enterprise/DatasmithCADImporter/Source/CADInterfaces/Public/TechSoftInterface.h
@@ -19,44 +19,44 @@ THIRD_PARTY_INCLUDES_END
 
 namespace CADLibrary
 {
-	CADINTERFACES_API bool TECHSOFT_InitializeKernel(const TCHAR* = TEXT(""));
+CADINTERFACES_API bool TECHSOFT_InitializeKernel(const TCHAR* = TEXT(""));
 
-	CADINTERFACES_API class FTechSoftInterface
+CADINTERFACES_API class FTechSoftInterface
+{
+public:
+	/*
+	* Returns true if the object has been created outside of the memory pool of the running process
+	* This is the case when the object has been created by the DatasmithRuntime plugin
+	*/
+	bool IsExternal()
 	{
-	public:
-		/*
-		* Returns true if the object has been created outside of the memory pool of the running process
-		* This is the case when the object has been created by the DatasmithRuntime plugin
-		*/
-		bool IsExternal() 
-		{
-			return bIsExternal;
-		}
+		return bIsExternal;
+	}
 
-		void SetExternal(bool Value) 
-		{
-			bIsExternal = Value;
-		}
+	void SetExternal(bool Value)
+	{
+		bIsExternal = Value;
+	}
 
-		bool InitializeKernel(const TCHAR* = TEXT(""));
+	bool InitializeKernel(const TCHAR* = TEXT(""));
 
 
 #ifdef USE_TECHSOFT_SDK
-		A3DStatus Import(const A3DImport& Import);
-		A3DAsmModelFile* GetModelFile();
+	A3DStatus Import(const A3DImport& Import);
+	A3DAsmModelFile* GetModelFile();
 #endif
 
-	private:
+private:
 
-		bool bIsExternal = false;
-		bool bIsInitialize = false;
+	bool bIsExternal = false;
+	bool bIsInitialize = false;
 
 #ifdef USE_TECHSOFT_SDK
-		TUniquePtr<A3DSDKHOOPSExchangeLoader> ExchangeLoader;
+	TUniquePtr<A3DSDKHOOPSExchangeLoader> ExchangeLoader;
 #endif
-	};
+};
 
-	CADINTERFACES_API FTechSoftInterface& GetTechSoftInterface();
+CADINTERFACES_API FTechSoftInterface& GetTechSoftInterface();
 
 }
 
diff --git a/Engine/Source/Runtime/Datasmith/CADKernel/Private/CADKernel/UI/Message.cpp b/Engine/Source/Runtime/Datasmith/CADKernel/Private/CADKernel/UI/Message.cpp
--- a/Engine/Source/Runtime/Datasmith/CADKernel/Private/CADKernel/UI/Message.cpp
+++ b/Engine/Source/Runtime/Datasmith/CADKernel/Private/CADKernel/UI/Message.cpp
@@ -11,174 +11,174 @@
 
 namespace CADKernel
 {
-	const char* VerboseLevelConstNames[] = {
-		("NO_VERBOSE"),
-		("SPY"),
-		("LOG"),
-		("DBG"),
-		nullptr
-	};
-
-	const char* VerboseConstDescHelp[] = {
-		": nothing is printed",
-		": only commands and arguments will be showed",
-		": SPY + messages",
-		": debug mode, for developers only",
-		nullptr
-	};
-
-	int32 FMessage::NumberOfIndentation = 0;
-	int32 FMessage::OldPercent = -1;
-
-	void FMessage::VPrintf(EVerboseLevel Level, const TCHAR* Text, ...)
+const char* VerboseLevelConstNames[] = {
+	("NO_VERBOSE"),
+	("SPY"),
+	("LOG"),
+	("DBG"),
+	nullptr
+};
+
+const char* VerboseConstDescHelp[] = {
+	": nothing is printed",
+	": only commands and arguments will be showed",
+	": SPY + messages",
+	": debug mode, for developers only",
+	nullptr
+};
+
+int32 FMessage::NumberOfIndentation = 0;
+int32 FMessage::OldPercent = -1;
+
+void FMessage::VPrintf(EVerboseLevel Level, const TCHAR* Text, ...)
+{
+	EVerboseLevel VerboseLevel = (EVerboseLevel)(int32)FSystem::Get().GetVerboseLevel();
+	TSharedPtr<FArchive> LogFile = FSystem::Get().GetLogFile();
+	EVerboseLevel LogLevel = FSystem::Get().GetLogLevel();
+	TSharedPtr<FArchive> SpyFile = FSystem::Get().GetSpyFile();
+
+	FString Indentation;
+	TCHAR* Buffer = NULL;
+	int32 Result = -1;
+	if (VerboseLevel >= Level || (LogFile.IsValid() && Level <= EVerboseLevel::Debug) || (SpyFile.IsValid() && Level == EVerboseLevel::Spy))
 	{
-		EVerboseLevel VerboseLevel = (EVerboseLevel)(int32)FSystem::Get().GetVerboseLevel();
-		TSharedPtr<FArchive> LogFile = FSystem::Get().GetLogFile();
-		EVerboseLevel LogLevel = FSystem::Get().GetLogLevel();
-		TSharedPtr<FArchive> SpyFile = FSystem::Get().GetSpyFile();
-
-		FString Indentation;
-		TCHAR* Buffer = NULL;
-		int32 Result = -1;
-		if (VerboseLevel >= Level || (LogFile.IsValid() && Level <= EVerboseLevel::Debug) || (SpyFile.IsValid() && Level == EVerboseLevel::Spy))
-		{
-			int32 BufferSize = 1024;
-
-			// do the usual VARARGS shenanigans
-			while (Result == -1)
-			{
-				FMemory::Free(Buffer);
-				Buffer = (TCHAR*)FMemory::Malloc(BufferSize * sizeof(TCHAR));
-				GET_VARARGS_RESULT(Buffer, BufferSize, BufferSize - 1, Text, Text, Result);
-				BufferSize *= 2;
-			};
-			Buffer[Result] = 0;
-
-			Indentation.Reserve(NumberOfIndentation * 3);
-			for (int32 iIndent = 0; iIndent < NumberOfIndentation; iIndent++)
-			{
-				Indentation += TEXT(" - ");
-			}
-		}
+		int32 BufferSize = 1024;
 
-		if (VerboseLevel >= Level)
+		// do the usual VARARGS shenanigans
+		while (Result == -1)
 		{
-			FSystem::Get().GetConsole().Print(*Indentation, Level);
-			FSystem::Get().GetConsole().Print(Buffer, Level);
+			FMemory::Free(Buffer);
+			Buffer = (TCHAR*)FMemory::Malloc(BufferSize * sizeof(TCHAR));
+			GET_VARARGS_RESULT(Buffer, BufferSize, BufferSize - 1, Text, Text, Result);
+			BufferSize *= 2;
+		};
+		Buffer[Result] = 0;
+
+		Indentation.Reserve(NumberOfIndentation * 3);
+		for (int32 iIndent = 0; iIndent < NumberOfIndentation; iIndent++)
+		{
+			Indentation += TEXT(" - ");
 		}
+	}
 
-		if ((LogFile.IsValid() && Level <= LogLevel) || (SpyFile.IsValid() && Level <= Spy))
+	if (VerboseLevel >= Level)
+	{
+		FSystem::Get().GetConsole().Print(*Indentation, Level);
+		FSystem::Get().GetConsole().Print(Buffer, Level);
+	}
+
+	if ((LogFile.IsValid() && Level <= LogLevel) || (SpyFile.IsValid() && Level <= Spy))
+	{
+		FTCHARToUTF8 UTF8Indentation(*Indentation, Indentation.Len());
+		FTCHARToUTF8 UTF8String(Buffer, Result);
+		if (LogFile.IsValid() && Level <= LogLevel)
 		{
-			FTCHARToUTF8 UTF8Indentation(*Indentation, Indentation.Len());
-			FTCHARToUTF8 UTF8String(Buffer, Result);
-			if (LogFile.IsValid() && Level <= LogLevel)
-			{
-				LogFile->Serialize((UTF8CHAR*)UTF8Indentation.Get(), UTF8Indentation.Length() * sizeof(UTF8CHAR));
-				LogFile->Serialize((UTF8CHAR*)UTF8String.Get(), UTF8String.Length() * sizeof(UTF8CHAR));
-				LogFile->Flush();
-			}
-
-			if (SpyFile.IsValid() && Level <= Spy)
-			{
-				SpyFile->Serialize((UTF8CHAR*)UTF8Indentation.Get(), UTF8Indentation.Length() * sizeof(UTF8CHAR));
-				SpyFile->Serialize((UTF8CHAR*)UTF8String.Get(), UTF8String.Length() * sizeof(UTF8CHAR));
-				SpyFile->Flush();
-			}
+			LogFile->Serialize((UTF8CHAR*)UTF8Indentation.Get(), UTF8Indentation.Length() * sizeof(UTF8CHAR));
+			LogFile->Serialize((UTF8CHAR*)UTF8String.Get(), UTF8String.Length() * sizeof(UTF8CHAR));
+			LogFile->Flush();
 		}
 
-		if (Buffer)
+		if (SpyFile.IsValid() && Level <= Spy)
 		{
-			FMemory::Free(Buffer);
+			SpyFile->Serialize((UTF8CHAR*)UTF8Indentation.Get(), UTF8Indentation.Length() * sizeof(UTF8CHAR));
+			SpyFile->Serialize((UTF8CHAR*)UTF8String.Get(), UTF8String.Length() * sizeof(UTF8CHAR));
+			SpyFile->Flush();
 		}
 	}
 
+	if (Buffer)
+	{
+		FMemory::Free(Buffer);
+	}
+}
+
 #if defined(CADKERNEL_DEV) || defined(CADKERNEL_STDA)
-	void FMessage::VQaPrintF(const TCHAR* Header, const TCHAR* Text, ...)
+void FMessage::VQaPrintF(const TCHAR* Header, const TCHAR* Text, ...)
+{
+	TSharedPtr<FArchive> QaDataFile = FSystem::Get().GetQaDataFile();
+	TSharedPtr<FArchive> QaHeaderFile = FSystem::Get().GetQaHeaderFile();
+	if (!QaDataFile.IsValid())
 	{
-		TSharedPtr<FArchive> QaDataFile = FSystem::Get().GetQaDataFile();
-		TSharedPtr<FArchive> QaHeaderFile = FSystem::Get().GetQaHeaderFile();
-		if (!QaDataFile.IsValid())
-		{
-			return;
-		}
+		return;
+	}
 
-		FString Indentation;
-		TCHAR* Buffer = NULL;
-		int32 Result = -1;
+	FString Indentation;
+	TCHAR* Buffer = NULL;
+	int32 Result = -1;
 
+	{
+		int32 BufferSize = 1024;
+
+		// do the usual VARARGS shenanigans
+		while (Result == -1)
+		{
+			FMemory::Free(Buffer);
+			Buffer = (TCHAR*)FMemory::Malloc(BufferSize * sizeof(TCHAR));
+			GET_VARARGS_RESULT(Buffer, BufferSize, BufferSize - 1, Text, Text, Result);
+			BufferSize *= 2;
+		};
+		Buffer[Result] = 0;
+
+		Indentation.Reserve(NumberOfIndentation * 3);
+		for (int32 iIndent = 0; iIndent < NumberOfIndentation; iIndent++)
 		{
-			int32 BufferSize = 1024;
-
-			// do the usual VARARGS shenanigans
-			while (Result == -1)
-			{
-				FMemory::Free(Buffer);
-				Buffer = (TCHAR*)FMemory::Malloc(BufferSize * sizeof(TCHAR));
-				GET_VARARGS_RESULT(Buffer, BufferSize, BufferSize - 1, Text, Text, Result);
-				BufferSize *= 2;
-			};
-			Buffer[Result] = 0;
-
-			Indentation.Reserve(NumberOfIndentation * 3);
-			for (int32 iIndent = 0; iIndent < NumberOfIndentation; iIndent++)
-			{
-				Indentation += TEXT(" - ");
-			}
+			Indentation += TEXT(" - ");
 		}
+	}
 
-		FSystem::Get().GetConsole().Print(Header, Log);
-		FSystem::Get().GetConsole().Print(TEXT("\n"), Log);
-		FSystem::Get().GetConsole().Print(Buffer, Log);
-		FSystem::Get().GetConsole().Print(TEXT("\n"), Log);
+	FSystem::Get().GetConsole().Print(Header, Log);
+	FSystem::Get().GetConsole().Print(TEXT("\n"), Log);
+	FSystem::Get().GetConsole().Print(Buffer, Log);
+	FSystem::Get().GetConsole().Print(TEXT("\n"), Log);
 
-		FString HeaderStr = Header;
-		FTCHARToUTF8 UTF8String(Header, HeaderStr.Len());
-		QaDataFile->Serialize((UTF8CHAR*)UTF8String.Get(), UTF8String.Length() * sizeof(UTF8CHAR));
-		QaDataFile->Flush();
+	FString HeaderStr = Header;
+	FTCHARToUTF8 UTF8String(Header, HeaderStr.Len());
+	QaDataFile->Serialize((UTF8CHAR*)UTF8String.Get(), UTF8String.Length() * sizeof(UTF8CHAR));
+	QaDataFile->Flush();
 
-		if(QaHeaderFile.IsValid())
-		{
-			FTCHARToUTF8 UTF8String(Buffer, Result);
-			QaHeaderFile->Serialize((UTF8CHAR*)UTF8String.Get(), UTF8String.Length() * sizeof(UTF8CHAR));
-			QaHeaderFile->Flush();
-		}
+	if (QaHeaderFile.IsValid())
+	{
+		FTCHARToUTF8 UTF8String(Buffer, Result);
+		QaHeaderFile->Serialize((UTF8CHAR*)UTF8String.Get(), UTF8String.Length() * sizeof(UTF8CHAR));
+		QaHeaderFile->Flush();
+	}
 
-		if (Buffer)
-		{
-			FMemory::Free(Buffer);
-		}
+	if (Buffer)
+	{
+		FMemory::Free(Buffer);
 	}
+}
 #endif
 
 
-	void FChrono::PrintClockElapse(EVerboseLevel Level, const TCHAR* Indent, const TCHAR* Process, FDuration Duration, ETimeUnit Unit)
-	{
+void FChrono::PrintClockElapse(EVerboseLevel Level, const TCHAR* Indent, const TCHAR* Process, FDuration Duration, ETimeUnit Unit)
+{
 #ifdef CADKERNEL_DEV
-		if (FSystem::Get().GetVerboseLevel() < Level)
-		{
-			return;
-		}
+	if (FSystem::Get().GetVerboseLevel() < Level)
+	{
+		return;
+	}
 
-		TCHAR* SUnit = TEXT("seconds");
-		long long Time = 0;
-		switch (Unit)
-		{
-		case ETimeUnit::NanoSeconds:
-			SUnit = TEXT("nanoseconds");
-			Time = FChrono::ConvertInto<std::chrono::nanoseconds>(Duration);
-			break;
-		case ETimeUnit::MicroSeconds:
-			SUnit = TEXT("microseconds");
-			Time = FChrono::ConvertInto<std::chrono::microseconds>(Duration);
-			break;
-		case ETimeUnit::MilliSeconds:
-			SUnit = TEXT("milliseconds");
-			Time = FChrono::ConvertInto<std::chrono::milliseconds>(Duration);
-			break;
-		default:
-			Time = FChrono::ConvertInto<std::chrono::seconds>(Duration);
-		}
-		FMessage::Printf(Level, TEXT("%sSpeed test (%s) -----> %lld %s\n"), Indent, Process, Time, SUnit);
-#endif
+	TCHAR* SUnit = TEXT("seconds");
+	long long Time = 0;
+	switch (Unit)
+	{
+	case ETimeUnit::NanoSeconds:
+		SUnit = TEXT("nanoseconds");
+		Time = FChrono::ConvertInto<std::chrono::nanoseconds>(Duration);
+		break;
+	case ETimeUnit::MicroSeconds:
+		SUnit = TEXT("microseconds");
+		Time = FChrono::ConvertInto<std::chrono::microseconds>(Duration);
+		break;
+	case ETimeUnit::MilliSeconds:
+		SUnit = TEXT("milliseconds");
+		Time = FChrono::ConvertInto<std::chrono::milliseconds>(Duration);
+		break;
+	default:
+		Time = FChrono::ConvertInto<std::chrono::seconds>(Duration);
 	}
+	FMessage::Printf(Level, TEXT("%sSpeed test (%s) -----> %lld %s\n"), Indent, Process, Time, SUnit);
+#endif
+}
 }
\ No newline at end of file
diff --git a/Engine/Source/Runtime/Datasmith/CADKernel/Public/CADKernel/Geo/Surfaces/NURBSSurface.h b/Engine/Source/Runtime/Datasmith/CADKernel/Public/CADKernel/Geo/Surfaces/NURBSSurface.h
--- a/Engine/Source/Runtime/Datasmith/CADKernel/Public/CADKernel/Geo/Surfaces/NURBSSurface.h
+++ b/Engine/Source/Runtime/Datasmith/CADKernel/Public/CADKernel/Geo/Surfaces/NURBSSurface.h
@@ -7,210 +7,210 @@
 namespace CADKernel
 {
 
-	struct FNurbsSurfaceHomogeneousData
-	{
-		bool bSwapUV;
-		int32 PoleUCount;
-		int32 PoleVCount;
+struct FNurbsSurfaceHomogeneousData
+{
+	bool bSwapUV;
+	int32 PoleUCount;
+	int32 PoleVCount;
 
-		int32 UDegree;
-		int32 VDegree;
+	int32 UDegree;
+	int32 VDegree;
 
-		TArray<double> UNodalVector;
-		TArray<double> VNodalVector;
+	TArray<double> UNodalVector;
+	TArray<double> VNodalVector;
 
-		bool bIsRational;
-		TArray<double> HomogeneousPoles;
-	};
+	bool bIsRational;
+	TArray<double> HomogeneousPoles;
+};
 
-	class CADKERNEL_API FNURBSSurface : public FSurface
+class CADKERNEL_API FNURBSSurface : public FSurface
+{
+	friend FEntity;
+
+protected:
+	int32 PoleUCount;
+	int32 PoleVCount;
+
+	int32 UDegree;
+	int32 VDegree;
+
+	TArray<double> UNodalVector;
+	TArray<double> VNodalVector;
+
+	TArray<double> Weights;
+
+	TArray<FPoint> Poles;
+
+	bool bIsRational;
+
+	/**
+	 * Data generated at initialization which are not serialized
+	 */
+	TArray<double> HomogeneousPoles;
+
+	/**
+	 * Build a Non uniform B-Spline surface
+	 * @param NodalVectorU Its size is the number of poles in U + the surface degree in U + 1 (PoleUNum + UDegre + 1)
+	 * @param NodalVectorV Its size is the number of poles in V + the surface degree in V + 1 (PoleVNum + VDegre + 1)
+	 */
+	FNURBSSurface(const double InToleranceGeometric, int32 InPoleUCount, int32 InPoleVCount, int32 InDegreU, int32 InDegreV, const TArray<double>& InNodalVectorU, const TArray<double>& InNodalVectorV, const TArray<FPoint>& InPoles)
+		: FSurface(InToleranceGeometric)
+		, PoleUCount(InPoleUCount)
+		, PoleVCount(InPoleVCount)
+		, UDegree(InDegreU)
+		, VDegree(InDegreV)
+		, UNodalVector(InNodalVectorU)
+		, VNodalVector(InNodalVectorV)
+		, Poles(InPoles)
+		, bIsRational(false)
 	{
-		friend FEntity;
-
-	protected:
-		int32 PoleUCount;
-		int32 PoleVCount;
-
-		int32 UDegree;
-		int32 VDegree;
-
-		TArray<double> UNodalVector;
-		TArray<double> VNodalVector;
-
-		TArray<double> Weights;
-
-		TArray<FPoint> Poles;
-
-		bool bIsRational;
-
-		/**
-		 * Data generated at initialization which are not serialized
-		 */
-		TArray<double> HomogeneousPoles;
-
-		/**
-		 * Build a Non uniform B-Spline surface
-		 * @param NodalVectorU Its size is the number of poles in U + the surface degree in U + 1 (PoleUNum + UDegre + 1)
-		 * @param NodalVectorV Its size is the number of poles in V + the surface degree in V + 1 (PoleVNum + VDegre + 1)
-		 */
-		FNURBSSurface(const double InToleranceGeometric, int32 InPoleUCount, int32 InPoleVCount, int32 InDegreU, int32 InDegreV, const TArray<double>& InNodalVectorU, const TArray<double>& InNodalVectorV, const TArray<FPoint>& InPoles)
-			: FSurface(InToleranceGeometric)
-			, PoleUCount(InPoleUCount)
-			, PoleVCount(InPoleVCount)
-			, UDegree(InDegreU)
-			, VDegree(InDegreV)
-			, UNodalVector(InNodalVectorU)
-			, VNodalVector(InNodalVectorV)
-			, Poles(InPoles)
-			, bIsRational(false)
-		{
-			Finalize();
-		}
-
-		/**
-		 * Build a Non uniform rational B-Spline surface
-		 * @param NodalVectorU Its size is the number of poles in U + the surface degree in U + 1 (PoleUNum + UDegre + 1)
-		 * @param NodalVectorV Its size is the number of poles in V + the surface degree in V + 1 (PoleVNum + VDegre + 1)
-		 */
-		FNURBSSurface(const double InToleranceGeometric, int32 InPoleUCount, int32 InPoleVCount, int32 InDegreU, int32 InDegreV, const TArray<double>& InNodalVectorU, const TArray<double>& InNodalVectorV, const TArray<FPoint>& InPoles, const TArray<double>& InWeights)
-			: FSurface(InToleranceGeometric)
-			, PoleUCount(InPoleUCount)
-			, PoleVCount(InPoleVCount)
-			, UDegree(InDegreU)
-			, VDegree(InDegreV)
-			, UNodalVector(InNodalVectorU)
-			, VNodalVector(InNodalVectorV)
-			, Weights(InWeights)
-			, Poles(InPoles)
-			, bIsRational(true)
-		{
-			ComputeMinToleranceIso();
-			Finalize();
-		}
-
-		/**
-		 * Build a Non uniform B-Spline surface
-		 * @param NodalVectorU Its size is the number of poles in U + the surface degree in U + 1 (PoleUNum + UDegre + 1)
-		 * @param NodalVectorV Its size is the number of poles in V + the surface degree in V + 1 (PoleVNum + VDegre + 1)
-		 */
-		FNURBSSurface(const double InToleranceGeometric, FNurbsSurfaceHomogeneousData NurbsData)
-			: FSurface(InToleranceGeometric)
-		{
-			FillNurbs(NurbsData);
-		}
+		Finalize();
+	}
+
+	/**
+	 * Build a Non uniform rational B-Spline surface
+	 * @param NodalVectorU Its size is the number of poles in U + the surface degree in U + 1 (PoleUNum + UDegre + 1)
+	 * @param NodalVectorV Its size is the number of poles in V + the surface degree in V + 1 (PoleVNum + VDegre + 1)
+	 */
+	FNURBSSurface(const double InToleranceGeometric, int32 InPoleUCount, int32 InPoleVCount, int32 InDegreU, int32 InDegreV, const TArray<double>& InNodalVectorU, const TArray<double>& InNodalVectorV, const TArray<FPoint>& InPoles, const TArray<double>& InWeights)
+		: FSurface(InToleranceGeometric)
+		, PoleUCount(InPoleUCount)
+		, PoleVCount(InPoleVCount)
+		, UDegree(InDegreU)
+		, VDegree(InDegreV)
+		, UNodalVector(InNodalVectorU)
+		, VNodalVector(InNodalVectorV)
+		, Weights(InWeights)
+		, Poles(InPoles)
+		, bIsRational(true)
+	{
+		ComputeMinToleranceIso();
+		Finalize();
+	}
+
+	/**
+	 * Build a Non uniform B-Spline surface
+	 * @param NodalVectorU Its size is the number of poles in U + the surface degree in U + 1 (PoleUNum + UDegre + 1)
+	 * @param NodalVectorV Its size is the number of poles in V + the surface degree in V + 1 (PoleVNum + VDegre + 1)
+	 */
+	FNURBSSurface(const double InToleranceGeometric, FNurbsSurfaceHomogeneousData NurbsData)
+		: FSurface(InToleranceGeometric)
+	{
+		FillNurbs(NurbsData);
+	}
 
-		FNURBSSurface() = default;
+	FNURBSSurface() = default;
 
-	public:
+public:
 
-		virtual void Serialize(FCADKernelArchive& Ar) override
+	virtual void Serialize(FCADKernelArchive& Ar) override
+	{
+		FSurface::Serialize(Ar);
+		Ar << PoleUCount;
+		Ar << PoleVCount;
+		Ar << UDegree;
+		Ar << VDegree;
+		Ar.Serialize(UNodalVector);
+		Ar.Serialize(VNodalVector);
+		Ar.Serialize(Weights);
+		Ar.Serialize(Poles);
+		Ar << bIsRational;
+
+		if (Ar.IsLoading())
 		{
-			FSurface::Serialize(Ar);
-			Ar << PoleUCount;
-			Ar << PoleVCount;
-			Ar << UDegree;
-			Ar << VDegree;
-			Ar.Serialize(UNodalVector);
-			Ar.Serialize(VNodalVector);
-			Ar.Serialize(Weights);
-			Ar.Serialize(Poles);
-			Ar << bIsRational;
-
-			if (Ar.IsLoading())
-			{
-				Finalize();
-			}
+			Finalize();
 		}
+	}
 
-		ESurface GetSurfaceType() const
-		{
-			return ESurface::Nurbs;
-		}
+	ESurface GetSurfaceType() const
+	{
+		return ESurface::Nurbs;
+	}
 
 #ifdef CADKERNEL_DEV
-		virtual FInfoEntity& GetInfo(FInfoEntity&) const override;
+	virtual FInfoEntity& GetInfo(FInfoEntity&) const override;
 #endif
 
-		constexpr const int32 GetDegree(EIso Iso) const
+	constexpr const int32 GetDegree(EIso Iso) const
+	{
+		switch (Iso)
 		{
-			switch (Iso)
-			{
-			case EIso::IsoU:
-				return UDegree;
-			case EIso::IsoV:
-			default:
-				return VDegree;
-			}
+		case EIso::IsoU:
+			return UDegree;
+		case EIso::IsoV:
+		default:
+			return VDegree;
 		}
+	}
 
-		constexpr const int32 GetPoleCount(EIso Iso) const
+	constexpr const int32 GetPoleCount(EIso Iso) const
+	{
+		switch (Iso)
 		{
-			switch (Iso)
-			{
-			case EIso::IsoU:
-				return PoleUCount;
-			case EIso::IsoV:
-			default:
-				return PoleVCount;
-			}
+		case EIso::IsoU:
+			return PoleUCount;
+		case EIso::IsoV:
+		default:
+			return PoleVCount;
 		}
+	}
 
-		const TArray<FPoint>& GetPoles() const
-		{
-			return Poles;
-		}
+	const TArray<FPoint>& GetPoles() const
+	{
+		return Poles;
+	}
 
-		const TArray<double>& GetWeights() const
-		{
-			return Weights;
-		}
+	const TArray<double>& GetWeights() const
+	{
+		return Weights;
+	}
 
-		TArray<double> GetHPoles() const
-		{
-			return HomogeneousPoles;
-		}
+	TArray<double> GetHPoles() const
+	{
+		return HomogeneousPoles;
+	}
 
-		constexpr const TArray<double>& GetNodalVector(EIso Iso) const
+	constexpr const TArray<double>& GetNodalVector(EIso Iso) const
+	{
+		switch (Iso)
 		{
-			switch (Iso)
-			{
-			case EIso::IsoU:
-				return UNodalVector;
-			case EIso::IsoV:
-			default:
-				return VNodalVector;
-			}
+		case EIso::IsoU:
+			return UNodalVector;
+		case EIso::IsoV:
+		default:
+			return VNodalVector;
 		}
+	}
 
-		bool IsRational() const
-		{
-			return bIsRational;
-		}
+	bool IsRational() const
+	{
+		return bIsRational;
+	}
 
 
-		virtual TSharedPtr<FEntityGeom> ApplyMatrix(const FMatrixH& InMatrix) const override;
+	virtual TSharedPtr<FEntityGeom> ApplyMatrix(const FMatrixH& InMatrix) const override;
 
-		virtual void EvaluatePoint(const FPoint2D& InSurfacicCoordinate, FSurfacicPoint& OutPoint3D, int32 InDerivativeOrder = 0) const override
-		{
-			BSpline::EvaluatePoint(*this, InSurfacicCoordinate, OutPoint3D, InDerivativeOrder);
-		}
+	virtual void EvaluatePoint(const FPoint2D& InSurfacicCoordinate, FSurfacicPoint& OutPoint3D, int32 InDerivativeOrder = 0) const override
+	{
+		BSpline::EvaluatePoint(*this, InSurfacicCoordinate, OutPoint3D, InDerivativeOrder);
+	}
 
-		void EvaluatePointGrid(const FCoordinateGrid& InSurfacicCoordinates, FSurfacicSampling& OutPoints, bool bComputeNormals) const override
-		{
-			BSpline::EvaluatePointGrid(*this, InSurfacicCoordinates, OutPoints, bComputeNormals);
-		}
+	void EvaluatePointGrid(const FCoordinateGrid& InSurfacicCoordinates, FSurfacicSampling& OutPoints, bool bComputeNormals) const override
+	{
+		BSpline::EvaluatePointGrid(*this, InSurfacicCoordinates, OutPoints, bComputeNormals);
+	}
 
-		virtual void LinesNotDerivables(const FSurfacicBoundary& InBoundary, int32 InDerivativeOrder, FCoordinateGrid& OutNotDerivableCoordinates) const override
-		{
-			BSpline::FindNotDerivableParameters(*this, InDerivativeOrder, InBoundary, OutNotDerivableCoordinates);
-		}
+	virtual void LinesNotDerivables(const FSurfacicBoundary& InBoundary, int32 InDerivativeOrder, FCoordinateGrid& OutNotDerivableCoordinates) const override
+	{
+		BSpline::FindNotDerivableParameters(*this, InDerivativeOrder, InBoundary, OutNotDerivableCoordinates);
+	}
 
-		void ComputeMinToleranceIso();
+	void ComputeMinToleranceIso();
 
-	private:
-		void Finalize();
-		void FillNurbs(FNurbsSurfaceHomogeneousData& NurbsData);
-	};
+private:
+	void Finalize();
+	void FillNurbs(FNurbsSurfaceHomogeneousData& NurbsData);
+};
 
 } // namespace CADKernel
 
diff --git a/Engine/Source/Runtime/Datasmith/CADKernel/Public/CADKernel/Mesh/Meshers/MesherReport.h b/Engine/Source/Runtime/Datasmith/CADKernel/Public/CADKernel/Mesh/Meshers/MesherReport.h
--- a/Engine/Source/Runtime/Datasmith/CADKernel/Public/CADKernel/Mesh/Meshers/MesherReport.h
+++ b/Engine/Source/Runtime/Datasmith/CADKernel/Public/CADKernel/Mesh/Meshers/MesherReport.h
@@ -9,18 +9,51 @@
 namespace CADKernel
 {
 
-struct CADKERNEL_API FMesherLog
+class CADKERNEL_API FMesherLog
 {
-	int32 CountOfIntersectingLoopCorrectionFailures = 0;
-	int32 CountOfSelfIntersectingCycles = 0;
-	int32 CountOfDegeneratedGrids = 0;
-	int32 CountOfDegeneratedLoops = 0;
+private:
 
+	int32 CountOfRemoveLoopSelfIntersectionFailure = 0;
+	int32 CountOfCrossingLoopsFailure = 0;
+
+	int32 CountOfDegeneratedGrid = 0;
+	int32 CountOfDegeneratedLoop = 0;
+	int32 CountOfMeshingFailure = 0;
+	int32 CountOfCycleMeshingFailure = 0;
+
+public:
 	FMesherLog()
 	{}
 
-	void PrintBilan() const
+	void PrintReport() const
+	{
+	}
+
+	void AddRemoveCrossingLoopsFailure()
+	{
+		CountOfCrossingLoopsFailure++;
+	}
+
+	void AddRemoveSelfIntersectionFailure()
+	{
+		CountOfRemoveLoopSelfIntersectionFailure++;
+	}
+
+	void AddCycleMeshingFailure()
 	{
+		CountOfCycleMeshingFailure++;
+	}
+
+	void AddDegeneratedLoop()
+	{
+		CountOfDegeneratedLoop++;
+		CountOfMeshingFailure++;
+	}
+
+	void AddDegeneratedGrid()
+	{
+		CountOfDegeneratedGrid++;
+		CountOfMeshingFailure++;
 	}
 };
 
@@ -39,7 +72,6 @@ struct CADKERNEL_API FMesherChronos
 	FDuration GlobalThinZones;
 
 	FDuration GlobalFindThinZones;
-	//FDuration GlobalScaleGrid;
 	FDuration GlobalMeshThinZones;
 
 	FMesherChronos()
@@ -55,7 +87,6 @@ struct CADKERNEL_API FMesherChronos
 		, GlobalMeshEdges(FChrono::Init())
 		, GlobalThinZones(FChrono::Init())
 		, GlobalFindThinZones(FChrono::Init())
-		//, GlobalScaleGrid(FChrono::Init())
 		, GlobalMeshThinZones(FChrono::Init())
 	{}
 
@@ -74,12 +105,22 @@ struct CADKERNEL_API FMesherChronos
 		FChrono::PrintClockElapse(Log, TEXT("  |   |  "), TEXT("TriangulateDuration "), GlobalTriangulateDuration);
 		FChrono::PrintClockElapse(Log, TEXT("  |   |   |  "), TEXT("Delaunay Duration "), GlobalDelaunayDuration);
 	}
+
+	void PrintReport() const
+	{
+	}
 };
 
 struct CADKERNEL_API FMesherReport
 {
 	FMesherLog Logs;
 	FMesherChronos Chronos;
+
+	void Print()
+	{
+		Logs.PrintReport();
+		Chronos.PrintReport();
+	}
 };
 
 }
diff --git a/Engine/Source/Runtime/Datasmith/CADKernel/Public/CADKernel/UI/Message.h b/Engine/Source/Runtime/Datasmith/CADKernel/Public/CADKernel/UI/Message.h
--- a/Engine/Source/Runtime/Datasmith/CADKernel/Public/CADKernel/UI/Message.h
+++ b/Engine/Source/Runtime/Datasmith/CADKernel/Public/CADKernel/UI/Message.h
@@ -7,75 +7,76 @@
 
 namespace CADKernel
 {
-	extern const char* VerboseLevelConstNames[];
-	extern const char* VerboseConstDescHelp[];
+extern const char* VerboseLevelConstNames[];
+extern const char* VerboseConstDescHelp[];
 
 
-	class CADKERNEL_API FMessage
-	{
-		friend class FProgressManager;
+class CADKERNEL_API FMessage
+{
+	friend class FProgressManager;
 
-	private:
-		static int32 NumberOfIndentation;
-		static int32 OldPercent;
+private:
+	static int32 NumberOfIndentation;
+	static int32 OldPercent;
 
-		static void VPrintf(EVerboseLevel Level, const TCHAR* Text, ...);
+	static void VPrintf(EVerboseLevel Level, const TCHAR* Text, ...);
 
 #if defined(CADKERNEL_DEV) || defined(CADKERNEL_STDA)
-		static void VQaPrintF(const TCHAR* Header, const TCHAR* Text, ...);
+	static void VQaPrintF(const TCHAR* Header, const TCHAR* Text, ...);
 #endif
 
-	public:
-
-		template <typename FmtType, typename... Types>
-		static void Printf(EVerboseLevel Level, const FmtType& Text, Types... Args)
-		{
-			static_assert(TAnd<TIsValidVariadicFunctionArg<Types>...>::Value, "Invalid argument(s) passed to FString::Printf");
-			VPrintf(Level, Text, Args...);
-		}
+public:
 
-		template <typename FmtType, typename... Types>
-		static void Error(const FmtType& Text, Types... Args)
-		{
-			static_assert(TAnd<TIsValidVariadicFunctionArg<Types>...>::Value, "Invalid argument(s) passed to FString::Printf");
-			FString CompleteText = TEXT("ERROR: ");
-			CompleteText += Text;
-			VPrintf(Log, *CompleteText, Args...);
-		}
+	template <typename FmtType, typename... Types>
+	static void Printf(EVerboseLevel Level, const FmtType& Text, Types... Args)
+	{
+		static_assert(TAnd<TIsValidVariadicFunctionArg<Types>...>::Value, "Invalid argument(s) passed to FString::Printf");
+		VPrintf(Level, Text, Args...);
+	}
 
-		template <typename FmtType, typename... Types>
-		static void Warning(const FmtType& Text, Types... Args)
-		{
-			static_assert(TAnd<TIsValidVariadicFunctionArg<Types>...>::Value, "Invalid argument(s) passed to FString::Printf");
-			FString CompleteText = TEXT("WARNING: ");
-			CompleteText += Text;
-			VPrintf(Log, *CompleteText, Args...);
-		}
+	template <typename FmtType, typename... Types>
+	static void Error(const FmtType& Text, Types... Args)
+	{
+		static_assert(TAnd<TIsValidVariadicFunctionArg<Types>...>::Value, "Invalid argument(s) passed to FString::Printf");
+		FString CompleteText = TEXT("ERROR: ");
+		CompleteText += Text;
+		VPrintf(Log, *CompleteText, Args...);
+	}
+
+	template <typename FmtType, typename... Types>
+	static void Warning(const FmtType& Text, Types... Args)
+	{
+		static_assert(TAnd<TIsValidVariadicFunctionArg<Types>...>::Value, "Invalid argument(s) passed to FString::Printf");
+		FString CompleteText = TEXT("WARNING: ");
+		CompleteText += Text;
+		VPrintf(Log, *CompleteText, Args...);
+	}
 
 #if defined(CADKERNEL_DEV) || defined(CADKERNEL_STDA)
-		template <typename FmtType, typename... Types>
-		static void FillQaDataFile(const FmtType& Header, const FmtType& Text, Types... Args)
-		{
-			static_assert(TAnd<TIsValidVariadicFunctionArg<Types>...>::Value, "Invalid argument(s) passed to FString::Printf");
-			VQaPrintF(Header, Text, Args...);
-		}
+	template <typename FmtType, typename... Types>
+	static void FillQaDataFile(const FmtType& Header, const FmtType& Text, Types... Args)
+	{
+		static_assert(TAnd<TIsValidVariadicFunctionArg<Types>...>::Value, "Invalid argument(s) passed to FString::Printf");
+		VQaPrintF(Header, Text, Args...);
+	}
 #endif
 
-		static void Indent(int32 InNumberOfIndent = 1)
-		{
-			NumberOfIndentation += InNumberOfIndent;
-		}
+	static void Indent(int32 InNumberOfIndent = 1)
+	{
+		NumberOfIndentation += InNumberOfIndent;
+	}
 
-		static void Deindent(int32 InNumberOfIndent = 1)
+	static void Deindent(int32 InNumberOfIndent = 1)
+	{
+		NumberOfIndentation -= InNumberOfIndent;
+		if (NumberOfIndentation < 0)
 		{
-			NumberOfIndentation -= InNumberOfIndent;
-			if (NumberOfIndentation < 0)
-			{
-				NumberOfIndentation = 0;
-			}
+			NumberOfIndentation = 0;
 		}
+	}
+
+};
 
-	};
 } // namespace CADKernel
 
 #define ERROR_FUNCTION_CALL_NOT_EXPECTED \
