CADTranslator : Add xas & xpr extension (Creo)

diff --git a/Engine/Plugins/Enterprise/DatasmithCADImporter/Source/DatasmithCADTranslator/Private/DatasmithCADTranslator.cpp b/Engine/Plugins/Enterprise/DatasmithCADImporter/Source/DatasmithCADTranslator/Private/DatasmithCADTranslator.cpp
--- a/Engine/Plugins/Enterprise/DatasmithCADImporter/Source/DatasmithCADTranslator/Private/DatasmithCADTranslator.cpp
+++ b/Engine/Plugins/Enterprise/DatasmithCADImporter/Source/DatasmithCADTranslator/Private/DatasmithCADTranslator.cpp
@@ -43,6 +43,8 @@ void FDatasmithCADTranslator::Initialize(FDatasmithTranslatorCapabilities& OutCa
 	OutCapabilities.SupportedFileFormats.Add(FFileFormatInfo{ TEXT("neu.*"), TEXT("Creo Assembly files") });
 	OutCapabilities.SupportedFileFormats.Add(FFileFormatInfo{ TEXT("neu"), TEXT("Creo Assembly files") });
 	OutCapabilities.SupportedFileFormats.Add(FFileFormatInfo{ TEXT("prt.*"), TEXT("Creo Part files") });
+	OutCapabilities.SupportedFileFormats.Add(FFileFormatInfo{ TEXT("xas"), TEXT("Creo Assembly files") });
+	OutCapabilities.SupportedFileFormats.Add(FFileFormatInfo{ TEXT("xpr"), TEXT("Creo Part files") });
 	OutCapabilities.SupportedFileFormats.Add(FFileFormatInfo{ TEXT("prt"), TEXT("Creo, NX Part files") });
 
 	OutCapabilities.SupportedFileFormats.Add(FFileFormatInfo{ TEXT("iam"), TEXT("Inventor Assembly files") });
