Process JT in mm (no more m)

diff --git a/Engine/Plugins/Enterprise/DatasmithCADImporter/Source/DatasmithCADTranslator/Private/DatasmithCADTranslator.cpp b/Engine/Plugins/Enterprise/DatasmithCADImporter/Source/DatasmithCADTranslator/Private/DatasmithCADTranslator.cpp
--- a/Engine/Plugins/Enterprise/DatasmithCADImporter/Source/DatasmithCADTranslator/Private/DatasmithCADTranslator.cpp
+++ b/Engine/Plugins/Enterprise/DatasmithCADImporter/Source/DatasmithCADTranslator/Private/DatasmithCADTranslator.cpp
@@ -108,22 +108,6 @@ bool FDatasmithCADTranslator::LoadScene(TSharedRef<IDatasmithScene> DatasmithSce
 
 	switch (FileDescriptor.GetFileFormat())
 	{
-	case CADLibrary::ECADFormat::JT:
-	{
-#if WITH_EDITOR
-		// Do not change the model unit when translator is called by the Datasmith runtime plugin.
-		if (IsInGameThread())
-		{
-			// In UE Editor, for historical QA reason (i.e. due to some validation files) the unit of JT file is set to meter as it seemed to be the default unit
-			// But setting the meter as the default unit requires resetting KernelIO that crash at the runtime...
-			// so for runtime use, metric unit is mm.
-			// this will change with the next release of CAD importer
-			ImportParameters.SetMetricUnit(1.);
-		}
-#endif
-		break;
-	}
-
 	case CADLibrary::ECADFormat::NX:
 	{
 		ImportParameters.SetDisplayPreference(CADLibrary::EDisplayPreference::ColorOnly);
