reduce wire missing dll spam

diff --git a/Engine/Plugins/Enterprise/DatasmithCADImporter/Source/DatasmithWireTranslator/Private/DatasmithWireTranslator.cpp b/Engine/Plugins/Enterprise/DatasmithCADImporter/Source/DatasmithWireTranslator/Private/DatasmithWireTranslator.cpp
--- a/Engine/Plugins/Enterprise/DatasmithCADImporter/Source/DatasmithWireTranslator/Private/DatasmithWireTranslator.cpp
+++ b/Engine/Plugins/Enterprise/DatasmithCADImporter/Source/DatasmithWireTranslator/Private/DatasmithWireTranslator.cpp
@@ -1401,7 +1401,7 @@ TSharedPtr<IDatasmithMeshElement> FWireTranslatorImpl::FindOrAddMeshElement(cons
 	MeshElement->SetLabel(*ShellNodeInfo.Label);
 	MeshElement->SetLightmapSourceUV(-1);
 
-	// Set MeshElement FileHash used for re-import task 
+	// Set MeshElement FileHash used for re-import task
 	FMD5 MD5; // unique Value that define the mesh
 	MD5.Update(reinterpret_cast<const uint8*>(&SceneFileHash), sizeof SceneFileHash);
 	// MeshActor Name
@@ -1847,7 +1847,7 @@ TOptional<FMeshDescription> FWireTranslatorImpl::MeshDagNodeWithExternalMesher(A
 
 	if (MeshParameters.bIsSymmetric)
 	{
-		// All actors of a Alias symmetric layer are defined in the world Reference i.e. they have identity transform. So Mesh actor has to be defined in the world reference. 
+		// All actors of a Alias symmetric layer are defined in the world Reference i.e. they have identity transform. So Mesh actor has to be defined in the world reference.
 		ObjectReference = EAliasObjectReference::WorldReference;
 	}
 
@@ -1873,7 +1873,7 @@ TOptional<FMeshDescription> FWireTranslatorImpl::MeshDagNodeWithExternalMesher(T
 	EAliasObjectReference ObjectReference = EAliasObjectReference::LocalReference;
 	if (MeshParameters.bIsSymmetric)
 	{
-		// All actors of a Alias symmetric layer are defined in the world Reference i.e. they have identity transform. So Mesh actor has to be defined in the world reference. 
+		// All actors of a Alias symmetric layer are defined in the world Reference i.e. they have identity transform. So Mesh actor has to be defined in the world reference.
 		ObjectReference = EAliasObjectReference::WorldReference;
 	}
 	else if (TessellationOptions.StitchingTechnique != EDatasmithCADStitchingTechnique::StitchingSew)
@@ -2125,7 +2125,8 @@ void FDatasmithWireTranslator::Initialize(FDatasmithTranslatorCapabilities& OutC
 	if (GIsEditor && !GEditor->PlayWorld && !GIsPlayInEditorWorld)
 	{
 #ifdef USE_OPENMODEL
-		if (FPlatformProcess::GetDllHandle(TEXT("libalias_api.dll")))
+ 		static void* DllHandle = FPlatformProcess::GetDllHandle(TEXT("libalias_api.dll"));
+ 		if (DllHandle)
 		{
 			// Check installed version of Alias Tools because binaries before 2021.3 are not compatible with Alias 2022
 			uint64 FileVersion = FPlatformMisc::GetFileVersion(TEXT("libalias_api.dll"));
@@ -2172,7 +2173,7 @@ bool FDatasmithWireTranslator::LoadScene(TSharedRef<IDatasmithScene> OutScene)
 
 	Translator = MakeShared<FWireTranslatorImpl>(Filename, OutScene);
 	if (!Translator)
-	{ 
+	{
 		return false;
 	}
 
