TechSoft Interface
- Improve Jt & Sw import
- Improve FTechSoftFileParser::TraverseTessellation3D add Strip and Fan
- Start CADFileReport in CoreTechBridge

diff --git a/Engine/Plugins/Enterprise/DatasmithCADImporter/Source/CADInterfaces/Private/CADInterfacesModule.cpp b/Engine/Plugins/Enterprise/DatasmithCADImporter/Source/CADInterfaces/Private/CADInterfacesModule.cpp
--- a/Engine/Plugins/Enterprise/DatasmithCADImporter/Source/CADInterfaces/Private/CADInterfacesModule.cpp
+++ b/Engine/Plugins/Enterprise/DatasmithCADImporter/Source/CADInterfaces/Private/CADInterfacesModule.cpp
@@ -45,7 +45,7 @@ ECADInterfaceAvailability ICADInterfacesModule::GetAvailability()
 		}
 		else if (CADLibrary::GCADLibrary == TEXT("TechSoft"))
 		{
-			if (CADLibrary::TECHSOFT_InitializeKernel())
+			if (CADLibrary::TechSoftUtils::TECHSOFT_InitializeKernel())
 			{
 				return ECADInterfaceAvailability::Available;
 			}
diff --git a/Engine/Plugins/Enterprise/DatasmithCADImporter/Source/CADInterfaces/Private/CoreTechBridge.h b/Engine/Plugins/Enterprise/DatasmithCADImporter/Source/CADInterfaces/Private/CoreTechBridge.h
--- a/Engine/Plugins/Enterprise/DatasmithCADImporter/Source/CADInterfaces/Private/CoreTechBridge.h
+++ b/Engine/Plugins/Enterprise/DatasmithCADImporter/Source/CADInterfaces/Private/CoreTechBridge.h
@@ -105,8 +105,23 @@ private:
 	{
 		AddMetadata(NodeId, (TSharedRef<FMetadataDictionary>) Entity);
 	}
-};
 
+//#ifdef CADKERNEL_DEV
+//	TMap<const uint32, FMatrixH> SurfaceToMatrix;
+//
+//	const FMatrixH& GetParamSpaceTransform(TSharedPtr<FSurface>& Surface)
+//	{
+//		FMatrixH* Matrix = SurfaceToMatrix.Find(Surface->GetId());
+//		return Matrix == nullptr ? FMatrixH::Identity : *Matrix;
+//	}
+//
+//	void SetParamSpaceTransorm(TSharedPtr<FSurface>& Surface, const FMatrixH Matrix)
+//	{
+//		SurfaceToMatrix.Add(Surface->GetId(), Matrix);
+//	}
+//#endif
+
+};
 }
 
 #endif // USE_KERNEL_IO_SDK
diff --git a/Engine/Plugins/Enterprise/DatasmithCADImporter/Source/CADInterfaces/Private/TUniqueTechSoftObj.h b/Engine/Plugins/Enterprise/DatasmithCADImporter/Source/CADInterfaces/Private/TUniqueTechSoftObj.h
--- a/Engine/Plugins/Enterprise/DatasmithCADImporter/Source/CADInterfaces/Private/TUniqueTechSoftObj.h
+++ b/Engine/Plugins/Enterprise/DatasmithCADImporter/Source/CADInterfaces/Private/TUniqueTechSoftObj.h
@@ -3,8 +3,6 @@
 
 #ifdef USE_TECHSOFT_SDK
 
-#include "CADData.h"
-
 #include "TechSoftInterface.h"
 
 namespace CADLibrary
@@ -85,6 +83,28 @@ public:
 		return Status;
 	}
 
+	template<typename... InArgTypes>
+	A3DStatus FillWith(A3DStatus (*Getter)(const A3DEntity*, ObjectType*, InArgTypes&&... ), const A3DEntity* DataPtr, InArgTypes&&... Args)
+	{
+		if (IsValid())
+		{
+			Status = ResetData();
+		}
+		else
+		{
+			Status = A3DStatus::A3D_SUCCESS;
+		}
+
+		if (!IsValid() || (DataPtr == DefaultValue))
+		{
+			Status = A3DStatus::A3D_ERROR;
+			return Status;
+		}
+
+		Status = Getter(DataPtr, &Data, Forward<InArgTypes>(Args)...);
+		return Status;
+	}
+
 	/**
 	 * Empty the structure
 	 */
diff --git a/Engine/Plugins/Enterprise/DatasmithCADImporter/Source/CADInterfaces/Private/TechSoftBridge.cpp b/Engine/Plugins/Enterprise/DatasmithCADImporter/Source/CADInterfaces/Private/TechSoftBridge.cpp
new file mode 100644
--- /dev/null
+++ b/Engine/Plugins/Enterprise/DatasmithCADImporter/Source/CADInterfaces/Private/TechSoftBridge.cpp
@@ -0,0 +1,1954 @@
+// Copyright Epic Games, Inc. All Rights Reserved.
+
+#include "TechSoftBridge.h"
+
+#ifdef USE_TECHSOFT_SDK
+
+#include "TechSoftInterface.h"
+#include "TUniqueTechSoftObj.h"
+#include "CADFileReport.h"
+
+#include "CADKernel/Core/Session.h"
+
+#include "CADKernel/Geo/Curves/BezierCurve.h"
+#include "CADKernel/Geo/Curves/BoundedCurve.h"
+#include "CADKernel/Geo/Curves/CompositeCurve.h"
+#include "CADKernel/Geo/Curves/Curve.h"
+#include "CADKernel/Geo/Curves/EllipseCurve.h"
+#include "CADKernel/Geo/Curves/HyperbolaCurve.h"
+#include "CADKernel/Geo/Curves/NURBSCurve.h"
+#include "CADKernel/Geo/Curves/NURBSCurve.h"
+#include "CADKernel/Geo/Curves/ParabolaCurve.h"
+#include "CADKernel/Geo/Curves/PolylineCurve.h"
+#include "CADKernel/Geo/Curves/RestrictionCurve.h"
+#include "CADKernel/Geo/Curves/RestrictionCurve.h"
+#include "CADKernel/Geo/Curves/SegmentCurve.h"
+#include "CADKernel/Geo/Curves/SurfacicCurve.h"
+
+#include "CADKernel/Geo/Surfaces/BezierSurface.h"
+#include "CADKernel/Geo/Surfaces/CompositeSurface.h"
+#include "CADKernel/Geo/Surfaces/ConeSurface.h"
+#include "CADKernel/Geo/Surfaces/CoonsSurface.h"
+#include "CADKernel/Geo/Surfaces/CylinderSurface.h"
+#include "CADKernel/Geo/Surfaces/NURBSSurface.h"
+#include "CADKernel/Geo/Surfaces/OffsetSurface.h"
+#include "CADKernel/Geo/Surfaces/PlaneSurface.h"
+#include "CADKernel/Geo/Surfaces/RevolutionSurface.h"
+#include "CADKernel/Geo/Surfaces/RuledSurface.h"
+#include "CADKernel/Geo/Surfaces/SphericalSurface.h"
+#include "CADKernel/Geo/Surfaces/Surface.h"
+#include "CADKernel/Geo/Surfaces/TabulatedCylinderSurface.h"
+#include "CADKernel/Geo/Surfaces/TorusSurface.h"
+#include "CADKernel/Topo/Body.h"
+#include "CADKernel/Topo/Shell.h"
+#include "CADKernel/Topo/TopologicalEdge.h"
+#include "CADKernel/Topo/TopologicalFace.h"
+#include "CADKernel/Topo/TopologicalLink.h"
+#include "CADKernel/Topo/TopologicalVertex.h"
+
+#include "CADKernel/Utils/StringUtil.h"
+
+#include "CADKernel/UI/Display.h"
+
+namespace CADKernel
+{
+namespace TechSoftUtils
+{
+
+const FUVReparameterization FUVReparameterization::Identity(1., 0, 1., 0);
+
+template<typename... InArgTypes>
+A3DStatus GetCurveAsNurbs(const A3DCrvBase* A3DCurve, A3DCrvNurbsData* DataPtr, InArgTypes&&... Args)
+{
+	return CADLibrary::TechSoftUtils::GetCurveAsNurbs(A3DCurve, DataPtr, Forward<InArgTypes>(Args)...);
+};
+
+template<typename... InArgTypes>
+A3DStatus GetSurfaceAsNurbs(const A3DSurfBase* A3DSurface, A3DSurfNurbsData* DataPtr, InArgTypes&&... Args)
+{
+	return CADLibrary::TechSoftUtils::GetSurfaceAsNurbs(A3DSurface, DataPtr, Forward<InArgTypes>(Args)...);
+}; 
+
+FString GetFileName(const FString& InFilePath)
+{
+	int32 Position = -1;
+	if (!InFilePath.FindLastChar(TEXT('\\'), Position))
+	{
+		InFilePath.FindLastChar(TEXT('/'), Position);
+	}
+
+	if (Position >= 0)
+	{
+		return InFilePath.Right(Position + 1);
+	}
+
+	return InFilePath;
+}
+
+FString GetOriginalName(const FString& Name)
+{
+	int32 Position;
+	if (Name.FindLastChar(TEXT('['), Position))
+	{
+		return Name.Left(Position);
+	}
+	return Name;
+}
+
+FString Get3dxmlOriginalName(const FString& Name)
+{
+	int32 Position;
+	if (Name.FindLastChar(TEXT('('), Position))
+	{
+		FString NewName = Name.Right(Position + 1);
+		if (NewName.FindLastChar(TEXT(')'), Position))
+		{
+			NewName[Position] = 0;
+		}
+		return NewName;
+	}
+	return Name;
+}
+
+FString GetSWOriginalName(const FString& Name)
+{
+	int32 Position;
+	if (Name.FindLastChar(TEXT('-'), Position))
+	{
+		FString NewName = Name.Left(Position) + TEXT("<") + Name.Right(Position + 1) + TEXT(">");
+		return NewName;
+	}
+	return Name;
+}
+
+int TraverseAttribute(const A3DMiscAttribute* Attribute, TMap<FString, FString>& OutMetaData)
+{
+	A3DInt32 Ret = A3D_SUCCESS;
+
+	CADLibrary::TUniqueTSObj<A3DMiscAttributeData> AttributeData(Attribute);
+	if (AttributeData.IsValid())
+	{
+		FString AttributeName;
+		if (AttributeData->m_bTitleIsInt)
+		{
+			A3DUns32 uiVal = 0;
+			memcpy(&uiVal, AttributeData->m_pcTitle, sizeof(A3DUns32));
+			AttributeName = Utils::ToFString((uint32)uiVal);
+		}
+		else if (AttributeData->m_pcTitle && AttributeData->m_pcTitle[0] != '\0')
+		{
+			AttributeName = UTF8_TO_TCHAR(AttributeData->m_pcTitle);
+		}
+		else
+		{
+			return A3D_ERROR;
+		}
+
+		for (A3DUns32 Index = 0; Index < AttributeData->m_uiSize; ++Index)
+		{
+			//single->SetAttribute("m_bTitleIsInt", (int)AttributeData->m_asSingleAttributesData[Index].m_bTitleIsInt);
+			if (AttributeData->m_asSingleAttributesData[Index].m_pcTitle == NULL)
+			{
+				//single->SetAttribute("m_pcTitle", "NULL");
+			}
+			else if (AttributeData->m_asSingleAttributesData[Index].m_bTitleIsInt)
+			{
+				A3DUns32 uiVal;
+				memcpy(&uiVal, AttributeData->m_asSingleAttributesData[Index].m_pcTitle, sizeof(A3DUns32));
+				//single->SetAttribute("m_pcTitle", (int)uiVal);
+			}
+			else
+			{
+				if (AttributeData->m_asSingleAttributesData[Index].m_pcTitle && AttributeData->m_asSingleAttributesData[Index].m_pcTitle[0] != '\0')
+				{
+					//single->SetAttribute("m_pcTitle", AttributeData->m_asSingleAttributesData[Index].m_pcTitle);
+				}
+			}
+
+			FString AttributeValue;
+
+			switch (AttributeData->m_asSingleAttributesData[Index].m_eType)
+			{
+			case kA3DModellerAttributeTypeTime:
+			case kA3DModellerAttributeTypeInt:
+			{
+				A3DInt32 Value;
+				memcpy(&Value, AttributeData->m_asSingleAttributesData[Index].m_pcData, sizeof(A3DInt32));
+				AttributeValue = Utils::ToFString((int32) Value);
+				break;
+			}
+
+			case kA3DModellerAttributeTypeReal:
+			{
+				A3DDouble Value;
+				memcpy(&Value, AttributeData->m_asSingleAttributesData[Index].m_pcData, sizeof(A3DDouble));
+				AttributeValue = Utils::ToFString((double) Value);
+				break;
+			}
+
+			case kA3DModellerAttributeTypeString:
+				if (AttributeData->m_asSingleAttributesData[Index].m_pcData && AttributeData->m_asSingleAttributesData[Index].m_pcData[0] != '\0')
+				{
+					AttributeValue = UTF8_TO_TCHAR(AttributeData->m_asSingleAttributesData[Index].m_pcData);
+				}
+				break;
+
+			default:
+				break;
+			}
+
+			if (AttributeName.Len())
+			{
+				if (Index)
+				{
+					FString Key = FString::Printf(TEXT("%s_%u"), *AttributeName, Index);
+					OutMetaData.Emplace(Key, AttributeValue);
+				}
+				else
+				{
+					OutMetaData.Emplace(AttributeName, AttributeValue);
+				}
+			}
+			else
+			{
+				Ret = A3D_ERROR;
+			}
+		}
+	}
+
+	return Ret;
+}
+
+int TraverseBase(const ECADType FileType, const A3DEntity* Entity, FEntityData& OutEntityData, bool bIsOccurrence)
+{
+	CADLibrary::TUniqueTSObj<A3DRootBaseData> RootBaseData(Entity);
+	if (RootBaseData.IsValid())
+	{
+		if (RootBaseData->m_uiPersistentId > 0)
+		{
+			FString PersistentId = Utils::ToFString((uint32)RootBaseData->m_uiPersistentId);
+			OutEntityData.MetaData.Emplace(TEXT("PersistentId"), PersistentId);
+		}
+
+		if (RootBaseData->m_pcPersistentId)
+		{
+			FString PersistentUuid = UTF8_TO_TCHAR(RootBaseData->m_pcPersistentId);
+			Utils::RemoveUnwantedChar(PersistentUuid, L'-');
+			OutEntityData.MetaData.Emplace(TEXT("UUID"), PersistentUuid);
+		}
+
+		for (A3DUns32 Index = 0; Index < RootBaseData->m_uiSize; ++Index)
+		{
+			TraverseAttribute(RootBaseData->m_ppAttributes[Index], OutEntityData.MetaData);
+		}
+
+		CADLibrary::TUniqueTSObj<A3DBoundingBoxData> sBoundingBox(Entity);
+		if (sBoundingBox.IsValid())
+		{
+			if ((sBoundingBox->m_sMin.m_dX + sBoundingBox->m_sMin.m_dY + sBoundingBox->m_sMin.m_dZ + sBoundingBox->m_sMax.m_dX + sBoundingBox->m_sMax.m_dY + sBoundingBox->m_sMax.m_dZ) != 0.0)
+			{
+				FString BBox = FString::Printf(TEXT("Min(%f,%f,%f) Max(%f,%f,%f)"), sBoundingBox->m_sMin.m_dX, sBoundingBox->m_sMin.m_dY, sBoundingBox->m_sMin.m_dZ, sBoundingBox->m_sMax.m_dX, sBoundingBox->m_sMax.m_dY, sBoundingBox->m_sMax.m_dZ);
+				OutEntityData.MetaData.Emplace(TEXT("BBox"), BBox);
+			}
+		}
+	}
+
+	if (bIsOccurrence)
+	{
+		FString* InstanceName = OutEntityData.MetaData.Find(TEXT("InstanceName"));
+		FString* Name = OutEntityData.MetaData.Find(TEXT("Name"));
+		FString* OriginalName = OutEntityData.MetaData.Find(TEXT("OriginalName"));
+		FString* PartNumber = OutEntityData.MetaData.Find(TEXT("PartNumber"));
+
+		if (!Name && !OriginalName &&
+			RootBaseData->m_pcName && RootBaseData->m_pcName[0] != '\0')
+		{
+			FString NameStr = UTF8_TO_TCHAR(RootBaseData->m_pcName);
+			FString PartNumberStr = GetOriginalName(NameStr);
+			if (FileType == ECADType::CATIA_3DXML)
+			{
+				PartNumberStr = Get3dxmlOriginalName(PartNumberStr);
+			}
+			if (FileType == ECADType::SOLIDWORKS)
+			{
+				PartNumberStr = GetSWOriginalName(PartNumberStr);
+			}
+			OutEntityData.MetaData.Add(TEXT("PartNumber"), PartNumberStr);
+			PartNumber = OutEntityData.MetaData.Find(TEXT("PartNumber"));
+		}
+
+		if (InstanceName)
+		{
+			OutEntityData.MetaData.Add(TEXT("Name"), *InstanceName);
+		}
+		else if (PartNumber)
+		{
+			OutEntityData.MetaData.Add(TEXT("Name"), *PartNumber);
+		}
+		else if (OriginalName)
+		{
+			OutEntityData.MetaData.Add(TEXT("Name"), *OriginalName);
+		}
+		OutEntityData.Name = OutEntityData.MetaData[L"Name"];
+	}
+	else
+	{
+		FString* Name = OutEntityData.MetaData.Find(TEXT("Name"));
+		FString* OriginalName = OutEntityData.MetaData.Find(TEXT("OriginalName"));
+		FString* PartNumber = OutEntityData.MetaData.Find(TEXT("PartNumber"));
+
+		if (!OriginalName && !PartNumber &&
+			RootBaseData->m_pcName && RootBaseData->m_pcName[0] != '\0')
+		{
+			FString NameStr = UTF8_TO_TCHAR(RootBaseData->m_pcName);
+			FString OriginalNameStr = GetOriginalName(NameStr);
+			OutEntityData.MetaData.Emplace(TEXT("OriginalName"), OriginalNameStr);
+			OriginalName = OutEntityData.MetaData.Find(TEXT("OriginalName"));
+		}
+
+		if (PartNumber)
+		{
+			OutEntityData.MetaData.Add(TEXT("Name"), *PartNumber);
+		}
+		else if (OriginalName)
+		{
+			OutEntityData.MetaData.Add(TEXT("Name"), *OriginalName);
+		}
+
+		Name = OutEntityData.MetaData.Find(TEXT("Name"));
+		if (Name)
+		{
+			OutEntityData.Name = *Name;
+			if (!OriginalName)
+			{
+				OutEntityData.MetaData.Add(TEXT("OriginalName"), OutEntityData.Name);
+			}
+		}
+	}
+
+	return A3D_SUCCESS;
+}
+
+int TraverseGraphics(const A3DGraphics* Graphics)
+{
+	CADLibrary::TUniqueTSObj<A3DGraphicsData> GraphicsData(Graphics);
+	if (GraphicsData.IsValid())
+	{
+		if (GraphicsData->m_uiLayerIndex != A3D_DEFAULT_LAYER)
+		{
+			//graphics->SetAttribute("m_uiLayerIndex", (int)GraphicsData.m_uiLayerIndex);
+		}
+		if (GraphicsData->m_uiStyleIndex != A3D_DEFAULT_STYLE_INDEX)
+		{
+			//graphics->SetAttribute("m_uiStyleIndex", (int)GraphicsData.m_uiStyleIndex);
+		}
+
+		//graphics->SetAttribute("m_usBehaviour", (int)GraphicsData.m_usBehaviour);
+	}
+
+	return A3D_SUCCESS;
+}
+
+int TraverseBaseWithGraphics(const ECADType FileType, const A3DEntity* Entity, FEntityData& OutEntityData, bool bIsOccurrence)
+{
+	CADLibrary::TUniqueTSObj<A3DRootBaseWithGraphicsData> BaseWithGraphicsData(Entity);
+	if (BaseWithGraphicsData.IsValid())
+	{
+		TraverseBase(FileType, Entity, OutEntityData, bIsOccurrence);
+		if (BaseWithGraphicsData->m_pGraphics != NULL)
+		{
+			TraverseGraphics(BaseWithGraphicsData->m_pGraphics);
+		}
+	}
+
+	return A3D_SUCCESS;
+}
+
+int TraverseSource(const A3DEntity* Entity, FEntityData& OutEntityData, bool bIsOccurrence, const ECADType FileType)
+{
+	if (CADLibrary::TechSoftUtils::IsEntityBaseWithGraphicsType(Entity))
+	{
+		return TraverseBaseWithGraphics(FileType, Entity, OutEntityData, bIsOccurrence);
+	}
+	else if (CADLibrary::TechSoftUtils::IsEntityBaseType(Entity))
+	{
+		return TraverseBase(FileType, Entity, OutEntityData, bIsOccurrence);
+	}
+	return A3D_SUCCESS;
+}
+
+int TraverseSpecificMetaData(A3DEModellerType ModellerType, const A3DAsmProductOccurrence* Occurrence, FEntityData& MetaData)
+{
+	//----------- Export Specific information per CAD format -----------
+	switch (ModellerType)
+	{
+	case kA3DModellerSlw:
+	{
+		CADLibrary::TUniqueTSObj<A3DAsmProductOccurrenceDataSLW> SolidWorksSpecificData(Occurrence);
+		if (SolidWorksSpecificData.IsValid())
+		{
+			//A3DUns16 m_usStructSize;									// Reserved: will be initialized by \ref A3D_INITIALIZE_DATA.
+			//A3DUTF8Char* SolidWorksSpecificData->m_psNodeSlwID;									/*!< ID comming from Solidworks, to use to resolve the "PathsInAssemblyTree" */
+
+			//A3DInt32 SolidWorksSpecificData->m_iIndexCfg;										/*!< Cfg Index */
+			//A3DUTF8Char* SolidWorksSpecificData->m_psCfgName;									/*!< Cfg Name in the file*/
+			if (SolidWorksSpecificData->m_psCfgName)
+			{
+				FString ConfigurationName = UTF8_TO_TCHAR(SolidWorksSpecificData->m_psCfgName);
+				MetaData.MetaData.Emplace(TEXT("ConfigurationName"), ConfigurationName);
+				FString ConfigurationIndex = Utils::ToFString((int32) SolidWorksSpecificData->m_iIndexCfg);
+				MetaData.MetaData.Emplace(TEXT("ConfigurationIndex"), ConfigurationIndex);
+				FString Configuration = TEXT("Configuration_") + ConfigurationIndex;
+				MetaData.MetaData.Emplace(Configuration, ConfigurationName);
+			}
+
+			//A3DEProductOccurrenceTypeSLW SolidWorksSpecificData->m_usType;						/*!< Product Occurrence Type */
+
+			//A3DUns32 SolidWorksSpecificData->m_uiAttachementsSize;								/*!< The size of \ref m_psAttachements. */
+			//A3DAsmAttachmentsInfosSLW* SolidWorksSpecificData->m_psAttachements;				/*!< Storage to put assembly attachments information on product occurrence in assembly tree. */
+
+		}
+		break;
+	}
+	case kA3DModellerUnigraphics:
+	{
+		CADLibrary::TUniqueTSObj<A3DAsmProductOccurrenceDataUg> UnigraphicsSpecificData(Occurrence);
+		if (UnigraphicsSpecificData.IsValid())
+		{
+			if (UnigraphicsSpecificData->m_psPartUID)
+			{
+				FString PartUID = UTF8_TO_TCHAR(UnigraphicsSpecificData->m_psPartUID);
+				MetaData.MetaData.Emplace(TEXT("UnigraphicsPartUID"), PartUID);
+			}
+			if (UnigraphicsSpecificData->m_psInstanceFileName)
+			{
+				FString InstanceFileName = UTF8_TO_TCHAR(UnigraphicsSpecificData->m_psInstanceFileName);
+				MetaData.MetaData.Emplace(TEXT("UnigraphicsPartUID"), InstanceFileName);
+			}
+
+			if (UnigraphicsSpecificData->m_uiInstanceTag)
+			{
+				FString InstanceTag = Utils::ToFString((uint32)UnigraphicsSpecificData->m_uiInstanceTag);
+				MetaData.MetaData.Emplace(TEXT("UnigraphicsInstanceTag"), InstanceTag);
+			}
+
+			//// Treat RefSets
+			//if (UnigraphicsSpecificData.m_uiChildrenByRefsetsSize || UnigraphicsSpecificData.m_uiSolidsByRefsetsSize)
+			//{
+			//	// Children by RefSet
+			//	for (A3DUns32 iRefSet = 0; iRefSet < UnigraphicsSpecificData.m_uiChildrenByRefsetsSize; ++iRefSet)
+			//	{
+			//		const A3DElementsByRefsetUg& pRefsetData = UnigraphicsSpecificData.m_asChildrenByRefsets[iRefSet];
+			//		// Init Xml Element if first occurrence
+			//		//refset->SetAttribute("Name", pRefsetData.m_psRefset);
+			//		// Add children
+			//		for (A3DUns32 i = 0; i < pRefsetData.m_uiElementsSize; ++i)
+			//		{
+			//			//refsetChild->SetAttribute("id", pRefsetData.m_auiElements[i]);
+			//		}
+			//	}
+			//	// Solid by RefSet
+			//	for (A3DUns32 iRefSet = 0; iRefSet < UnigraphicsSpecificData.m_uiSolidsByRefsetsSize; ++iRefSet)
+			//	{
+			//		const A3DElementsByRefsetUg& pRefsetData = UnigraphicsSpecificData.m_asSolidsByRefsets[iRefSet];
+			//		//refset->SetAttribute("Name", pRefsetData.m_psRefset);
+			//		// Add Solid
+			//		for (A3DUns32 i = 0; i < pRefsetData.m_uiElementsSize; ++i)
+			//		{
+			//			//refsetChild->SetAttribute("id", pRefsetData.m_auiElements[i]);
+			//		}
+			//	}
+			//}
+		}
+		break;
+	}
+
+	case kA3DModellerCatiaV5:
+	{
+		CADLibrary::TUniqueTSObj<A3DAsmProductOccurrenceDataCV5> CatiaV5SpecificData(Occurrence);
+		if (CatiaV5SpecificData.IsValid())
+		{
+			if (CatiaV5SpecificData->m_iCurrentAsmProductIdentifier)
+			{
+				FString InstanceTag = Utils::ToFString((int32) CatiaV5SpecificData->m_iCurrentAsmProductIdentifier);
+				MetaData.MetaData.Emplace(TEXT("CatiaAsmProductIdentifier"), InstanceTag);
+			}
+
+			if (CatiaV5SpecificData->m_iNotUpdatedAsmProductIdentifier)
+			{
+				FString InstanceTag = Utils::ToFString((int32) CatiaV5SpecificData->m_iNotUpdatedAsmProductIdentifier);
+				MetaData.MetaData.Emplace(TEXT("CatiaNotUpdatedAsmProductIdentifier"), InstanceTag);
+			}
+
+			if (CatiaV5SpecificData->m_aiCLSID[0])
+			{
+				FString InstanceTag = FString::Printf(TEXT("%d %d %d %d"), CatiaV5SpecificData->m_aiCLSID[0], CatiaV5SpecificData->m_aiCLSID[1], CatiaV5SpecificData->m_aiCLSID[2], CatiaV5SpecificData->m_aiCLSID[3]);
+				MetaData.MetaData.Emplace(TEXT("CatiaCLSID"), InstanceTag);
+			}
+
+			if (CatiaV5SpecificData->m_psVersion)
+			{
+				FString Version = UTF8_TO_TCHAR(CatiaV5SpecificData->m_psVersion);
+				MetaData.MetaData.Emplace(TEXT("CatiaVersion"), Version);
+			}
+
+			if (CatiaV5SpecificData->m_psPartNumber)
+			{
+				FString PartNumber = UTF8_TO_TCHAR(CatiaV5SpecificData->m_psPartNumber);
+				MetaData.MetaData.Emplace(TEXT("CatiaPartNumber"), PartNumber);
+			}
+		}
+		break;
+	}
+
+	default:
+		break;
+	}
+	return A3D_SUCCESS;
+}
+
+FMatrixH CreateCoordinateSystem(const A3DMiscCartesianTransformationData& Transformation, double UnitScale = 1.0)
+{
+	FPoint Origin(&Transformation.m_sOrigin.m_dX);
+	FPoint Ox(&Transformation.m_sXVector.m_dX);
+	FPoint Oy(&Transformation.m_sYVector.m_dX);
+
+	Ox.Normalize();
+	Oy.Normalize();
+
+	if (!FMath::IsNearlyEqual(UnitScale, 1.))
+	{
+		Origin *= UnitScale;
+	}
+	FPoint Oz = Ox ^ Oy;
+
+	FMatrixH Matrix = FMatrixH(Origin, Ox, Oy, Oz);
+
+	if (!FMath::IsNearlyEqual(Transformation.m_sScale.m_dX, 1.) || !FMath::IsNearlyEqual(Transformation.m_sScale.m_dY, 1.) || !FMath::IsNearlyEqual(Transformation.m_sScale.m_dZ, 1.))
+	{
+		FMatrixH Scale = FMatrixH::MakeScaleMatrix(Transformation.m_sScale.m_dX, Transformation.m_sScale.m_dY, Transformation.m_sScale.m_dZ);
+		Matrix *= Scale;
+	}
+	return Matrix;
+}
+
+FMatrixH TraverseTransformation3D(const A3DMiscCartesianTransformation* CartesianTransformation, double UnitScale)
+{
+	FMatrixH Matrix;
+	if (CartesianTransformation)
+	{
+		CADLibrary::TUniqueTSObj<A3DMiscCartesianTransformationData> CartesianTransformationData(CartesianTransformation);
+		if (CartesianTransformationData.IsValid())
+		{
+			Matrix = CreateCoordinateSystem(*CartesianTransformationData, UnitScale);
+		}
+	}
+	return Matrix;
+}
+
+FMatrixH TraverseCoordinateSystem(const A3DRiCoordinateSystem* CoordinateSystem, double UnitScale)
+{
+
+	FEntityData MetaData;
+	TechSoftUtils::TraverseSource(CoordinateSystem, MetaData);
+
+	FMatrixH Matrix;
+	CADLibrary::TUniqueTSObj<A3DRiCoordinateSystemData> CoordinateSystemData(CoordinateSystem);
+	if (CoordinateSystemData.IsValid())
+	{
+		Matrix = TraverseTransformation3D(CoordinateSystemData->m_pTransformation, UnitScale);
+	}
+	return Matrix;
+}
+
+FMatrixH  TraverseGeneralTransformation(const A3DMiscGeneralTransformation* GeneralTransformation, double UnitScale)
+{
+
+	FMatrixH Matrix;
+	CADLibrary::TUniqueTSObj<A3DMiscGeneralTransformationData> GeneralTransformationData(GeneralTransformation);
+	if (GeneralTransformationData.IsValid())
+	{
+		for (int32 Index = 0; Index < 16; ++Index)
+		{
+			Matrix[Index] = GeneralTransformationData->m_adCoeff[Index];
+		}
+		ensureCADKernel(!FMath::IsNearlyEqual(UnitScale, 1.));
+	}
+	return Matrix;
+}
+
+FMatrixH TraverseTransformation(const A3DMiscCartesianTransformation* Transformation3D, double UnitScale)
+{
+	if (Transformation3D == NULL)
+	{
+		return FMatrixH();
+	}
+
+	FEntityData MetaData;
+	TraverseSource(Transformation3D, MetaData, false);
+
+	A3DEEntityType Type = kA3DTypeUnknown;
+	CADLibrary::TechSoftUtils::GetEntityType(Transformation3D, &Type);
+
+	if (Type == kA3DTypeMiscCartesianTransformation)
+	{
+		return TraverseTransformation3D(Transformation3D, UnitScale);
+	}
+	else if (Type == kA3DTypeMiscGeneralTransformation)
+	{
+		return TraverseGeneralTransformation(Transformation3D, UnitScale);
+	}
+	return FMatrixH();
+}
+
+void FillInt32Array(const int32 Count, const A3DInt32* Values, TArray<int32>& OutInt32Array)
+{
+	OutInt32Array.Reserve(Count);
+	for (int32 Index = 0; Index < Count; Index++)
+	{
+		OutInt32Array.Add(Values[Index]);
+	}
+};
+
+void FillDoubleArray(const int32 Count, const double* Values, TArray<double>& OutDoubleArray)
+{
+	OutDoubleArray.Reserve(Count);
+	for (int32 Index = 0; Index < Count; Index++)
+	{
+		OutDoubleArray.Add(Values[Index]);
+	}
+};
+
+void FillDoubleArray(const int32 UCount, const int32 VCount, const double* Values, TArray<double>& OutDoubleArray)
+{
+	OutDoubleArray.SetNum(UCount * VCount);
+	for (int32 Undex = 0, ValueIndex = 0; Undex < UCount; ++Undex)
+	{
+		int32 Index = Undex;
+		for (int32 Vndex = 0; Vndex < VCount; ++Vndex, Index += UCount, ++ValueIndex)
+		{
+			OutDoubleArray[Index] = Values[ValueIndex];
+		}
+	}
+}
+
+void FillPointArray(const int32 Count, const A3DVector3dData* Points, TArray<FPoint>& OutPointsArray, double UnitScale = 1.0)
+{
+	OutPointsArray.Reserve(Count);
+	for (int32 Index = 0; Index < Count; Index++)
+	{
+		OutPointsArray.Emplace(&Points[Index].m_dX);
+	}
+
+	if (!FMath::IsNearlyEqual(UnitScale, 1.))
+	{
+		for (FPoint& Point : OutPointsArray)
+		{
+			Point *= UnitScale;
+		}
+	}
+};
+
+void FillPointArray(const int32 UCount, const int32 VCount, const A3DVector3dData* Points, TArray<FPoint>& OutPointsArray, double UnitScale = 1.0)
+{
+	OutPointsArray.SetNum(UCount * VCount);
+	for (int32 Undex = 0, PointIndex = 0; Undex < UCount; ++Undex)
+	{
+		int32 Index = Undex;
+		for (int32 Vndex = 0; Vndex < VCount; ++Vndex, Index += UCount, ++PointIndex)
+		{
+			OutPointsArray[Index].Set((double*)&(Points[PointIndex].m_dX));
+		}
+	}
+
+	if (!FMath::IsNearlyEqual(UnitScale, 1.))
+	{
+		for (FPoint& Point : OutPointsArray)
+		{
+			Point *= UnitScale;
+		}
+	}
+};
+
+FSurfacicBoundary GetSurfacicBoundary(A3DDomainData& Domain, double UnitScale = 1.0, bool bSwapUV = false)
+{
+	EIso UIndex = bSwapUV ? EIso::IsoV : EIso::IsoU;
+	EIso VIndex = bSwapUV ? EIso::IsoU : EIso::IsoV;
+
+	FSurfacicBoundary Boundary;
+	Boundary[UIndex].Min = Domain.m_sMin.m_dX * UnitScale;
+	Boundary[UIndex].Max = Domain.m_sMax.m_dX * UnitScale;
+	Boundary[VIndex].Min = Domain.m_sMin.m_dY * UnitScale;
+	Boundary[VIndex].Max = Domain.m_sMax.m_dY * UnitScale;
+	return Boundary;
+}
+
+FLinearBoundary GetLinearBoundary(A3DIntervalData& A3DDomain)
+{
+	FLinearBoundary Domain(A3DDomain.m_dMin, A3DDomain.m_dMax);
+	return Domain;
+}
+
+FLinearBoundary GetLinearBoundary(const A3DCrvBase* A3DCurve)
+{
+	CADLibrary::TUniqueTSObj<A3DIntervalData> A3DDomain(A3DCurve);
+	FLinearBoundary Domain(A3DDomain->m_dMin, A3DDomain->m_dMax);
+	return Domain;
+}
+
+} // ns TechSoftUtils
+
+FTechSoftBridge::FTechSoftBridge(FSession& InSession, FCADFileReport& InReport)
+	: Session(InSession)
+	, Report(InReport)
+	, GeometricTolerance(Session.GetGeometricTolerance())
+	, SquareGeometricTolerance(FMath::Square(Session.GetGeometricTolerance()))
+	, SquareJoiningVertexTolerance(SquareGeometricTolerance * 2)
+{
+}
+
+TSharedRef<FBody> FTechSoftBridge::AddBody(A3DRiBrepModel* BrepModel, FEntityData& ParentData, double FileUnit)
+{
+	Report.BodyCount++;
+
+	UnitScale = FileUnit;
+
+	FEntityData BrepMetaData;
+	TechSoftUtils::TraverseSource(BrepModel, BrepMetaData, false);
+
+	TSharedPtr<FEntity>* EntityPtr = TechSoftToEntity.Find(BrepModel);
+	if (EntityPtr)
+	{
+		TSharedPtr<FBody> Body = StaticCastSharedPtr<FBody>(*EntityPtr);
+		if (!Body.IsValid())
+		{
+			return Body.ToSharedRef();
+		}
+	}
+
+	TSharedRef<FBody> Body = FEntity::MakeShared<FBody>();
+	TechSoftToEntity.Add(BrepModel, Body);
+
+	AddMetadata(BrepMetaData, Body);
+
+	CADLibrary::TUniqueTSObj<A3DRiBrepModelData> BRepModelData(BrepModel);
+	if (BRepModelData.IsValid())
+	{
+		TraverseBrepData(BRepModelData->m_pBrepData, Body);
+	}
+
+	return Body;
+}
+
+void FTechSoftBridge::TraverseBrepData(const A3DTopoBrepData* A3DBrepData, TSharedRef<FBody>& Body)
+{
+	FEntityData MetaData;
+	TechSoftUtils::TraverseSource(A3DBrepData, MetaData, false);
+
+	{
+		CADLibrary::TUniqueTSObj<A3DTopoBodyData> TopoBodyData(A3DBrepData);
+		if (TopoBodyData.IsValid())
+		{
+			if (TopoBodyData->m_pContext)
+			{
+				CADLibrary::TUniqueTSObj<A3DTopoContextData> TopoContextData(TopoBodyData->m_pContext);
+				if (TopoContextData.IsValid())
+				{
+					if (TopoContextData->m_bHaveScale)
+					{
+						UnitScale *= TopoContextData->m_dScale;
+					}
+				}
+			}
+		}
+	}
+
+	CADLibrary::TUniqueTSObj<A3DTopoBrepDataData> TopoBrepData(A3DBrepData);
+	if (TopoBrepData.IsValid())
+	{
+		for (A3DUns32 Index = 0; Index < TopoBrepData->m_uiConnexSize; ++Index)
+		{
+			TraverseConnex(TopoBrepData->m_ppConnexes[Index], Body);
+		}
+	}
+}
+
+void FTechSoftBridge::TraverseConnex(const A3DTopoConnex* A3DTopoConnex, TSharedRef<FBody>& Body)
+{
+	FEntityData MetaData;
+	TechSoftUtils::TraverseSource(A3DTopoConnex, MetaData);
+
+	CADLibrary::TUniqueTSObj<A3DTopoConnexData> TopoConnexData(A3DTopoConnex);
+	if (TopoConnexData.IsValid())
+	{
+		for (A3DUns32 Index = 0; Index < TopoConnexData->m_uiShellSize; ++Index)
+		{
+			TraverseShell(TopoConnexData->m_ppShells[Index], Body);
+		}
+	}
+}
+
+void FTechSoftBridge::TraverseShell(const A3DTopoShell* A3DShell, TSharedRef<FBody>& Body)
+{
+	FEntityData MetaData;
+	TechSoftUtils::TraverseSource(A3DShell, MetaData);
+
+	TSharedRef<FShell> Shell = FEntity::MakeShared<FShell>();
+	Body->AddShell(Shell);
+	Report.ShellCount++;
+
+	AddMetadata(MetaData, Shell);
+
+	CADLibrary::TUniqueTSObj<A3DTopoShellData> ShellData(A3DShell);
+	if (ShellData.IsValid())
+	{
+		for (A3DUns32 ui = 0; ui < ShellData->m_uiFaceSize; ++ui)
+		{
+			AddFace(ShellData->m_ppFaces[ui], Shell);
+		}
+	}
+}
+
+static bool bUseCurveAsNurbs = true;
+TSharedPtr<FCurve> FTechSoftBridge::AddCurve(const A3DCrvBase* A3DCurve, const TechSoftUtils::FUVReparameterization& UVReparameterization)
+{
+	TSharedPtr<FCurve> Curve = TSharedPtr<FCurve>();
+	A3DEEntityType eType;
+	A3DInt32 iRet = CADLibrary::TechSoftUtils::GetEntityType(A3DCurve, &eType);
+	if (iRet == A3D_SUCCESS)
+	{
+		Report.CurveCount++;
+
+		switch (eType)
+		{
+		case kA3DTypeCrvNurbs:
+			Curve = AddCurveNurbs(A3DCurve, UVReparameterization);
+			break;
+		case kA3DTypeCrvLine:
+			Curve = AddCurveLine(A3DCurve, UVReparameterization);
+			Report.CurveNurbsCount++;
+			break;
+		case kA3DTypeCrvCircle:
+			Curve = AddCurveCircle(A3DCurve, UVReparameterization);
+			break;
+		case kA3DTypeCrvEllipse:
+			Curve = AddCurveEllipse(A3DCurve, UVReparameterization);
+			break;
+		case kA3DTypeCrvParabola:
+			Curve = AddCurveParabola(A3DCurve, UVReparameterization);
+			break;
+		case kA3DTypeCrvHyperbola:
+			Curve = AddCurveHyperbola(A3DCurve, UVReparameterization);
+			Report.CurveNurbsCount++;
+			break;
+		case kA3DTypeCrvHelix:
+			Curve = AddCurveHelix(A3DCurve, UVReparameterization);
+			Report.CurveNurbsCount++;
+			break;
+		case kA3DTypeCrvPolyLine:
+			Curve = AddCurvePolyLine(A3DCurve, UVReparameterization);
+			break;
+		case kA3DTypeCrvComposite:
+			Curve = AddCurveComposite(A3DCurve, UVReparameterization);
+			break;
+		default:
+			Curve = AddCurveAsNurbs(A3DCurve, UVReparameterization);
+			break;
+		}
+	}
+
+	FLinearBoundary Boundary = TechSoftUtils::GetLinearBoundary(A3DCurve);
+
+	return Curve;
+}
+
+TSharedPtr<FCurve> FTechSoftBridge::AddCurveLine(const A3DCrvLine* A3DCurve, const TechSoftUtils::FUVReparameterization& UVReparameterization)
+{
+	Report.CurveLineCount++;
+
+	if (bUseCurveAsNurbs)
+	{
+		return AddCurveAsNurbs(A3DCurve, UVReparameterization);
+	}
+
+
+	CADLibrary::TUniqueTSObj<A3DCrvLineData> CrvLineData(A3DCurve);
+	if (!CrvLineData.IsValid())
+	{
+		return TSharedPtr<FCurve>();
+	}
+
+	bool bIs2D = (bool)CrvLineData->m_bIs2D;
+
+	//setting->SetAttribute("m_bIs2D", (int)sData.m_bIs2D);
+	//traverseTransfo3d(sData.m_sTrsf, crv);
+	//traverseParam(&sData.m_sParam, crv);
+	//A3DCrvLineGet(NULL, &sData);
+
+	return TSharedPtr<FCurve>();
+}
+
+TSharedPtr<FCurve> FTechSoftBridge::AddCurveCircle(const A3DCrvCircle* A3DCurve, const TechSoftUtils::FUVReparameterization& UVReparameterization)
+{
+	Report.CurveCircleCount++;
+
+	if (bUseCurveAsNurbs)
+	{
+		return AddCurveAsNurbs(A3DCurve, UVReparameterization);
+	}
+
+	CADLibrary::TUniqueTSObj<A3DCrvCircleData> CrvCircleData(A3DCurve);
+	if (!CrvCircleData.IsValid())
+	{
+		return TSharedPtr<FCurve>();
+	}
+
+	bool bIs2D = (bool)CrvCircleData->m_bIs2D;
+
+	//setting->SetAttribute("m_bIs2D", (int)sData.m_bIs2D);
+	//traverseTransfo3d(sData.m_sTrsf, crv);
+	//traverseParam(&sData.m_sParam, crv);
+	//_SetDoubleAttribute(crv, "m_dRadius", sData.m_dRadius);
+
+	return TSharedPtr<FCurve>();
+}
+
+TSharedPtr<FCurve> FTechSoftBridge::AddCurveEllipse(const A3DCrvEllipse* A3DCurve, const TechSoftUtils::FUVReparameterization& UVReparameterization)
+{
+	Report.CurveEllipseCount++;
+
+	if (bUseCurveAsNurbs)
+	{
+		return AddCurveAsNurbs(A3DCurve, UVReparameterization);
+	}
+
+	CADLibrary::TUniqueTSObj<A3DCrvEllipseData> CrvEllipseData(A3DCurve);
+	if (!CrvEllipseData.IsValid())
+	{
+		return TSharedPtr<FCurve>();
+	}
+
+	bool bIs2D = (bool)CrvEllipseData->m_bIs2D;
+
+	//setting->SetAttribute("m_bIs2D", (int)sData.m_bIs2D);
+	//traverseTransfo3d(sData.m_sTrsf, crv);
+	//traverseParam(&sData.m_sParam, crv);
+	//_SetDoubleAttribute(crv, "m_dXRadius", sData.m_dXRadius);
+	//_SetDoubleAttribute(crv, "m_dYRadius", sData.m_dYRadius);
+
+	return TSharedPtr<FCurve>();
+}
+
+TSharedPtr<FCurve> FTechSoftBridge::AddCurveParabola(const A3DCrvParabola* A3DCurve, const TechSoftUtils::FUVReparameterization& UVReparameterization)
+{
+	Report.CurveParabolaCount++;
+
+	if (bUseCurveAsNurbs)
+	{
+		return AddCurveAsNurbs(A3DCurve, UVReparameterization);
+	}
+
+	CADLibrary::TUniqueTSObj<A3DCrvParabolaData> CrvParabolaData(A3DCurve);
+	if (!CrvParabolaData.IsValid())
+	{
+		return TSharedPtr<FCurve>();
+	}
+
+	bool bIs2D = (bool)CrvParabolaData->m_bIs2D;
+
+	//setting->SetAttribute("m_bIs2D", (int)sData.m_bIs2D);
+	//	traverseTransfo3d(sData.m_sTrsf, crv);
+	//	traverseParam(&sData.m_sParam, crv);
+	//	_SetDoubleAttribute(crv, "m_dFocal", sData.m_dFocal);
+	//	crv->SetAttribute("m_cParamType", (int)sData.m_cParamType);
+
+	return TSharedPtr<FCurve>();
+}
+
+TSharedPtr<FCurve> FTechSoftBridge::AddCurveHyperbola(const A3DCrvHyperbola* A3DCurve, const TechSoftUtils::FUVReparameterization& UVReparameterization)
+{
+	Report.CurveHyperbolaCount++;
+
+	if (bUseCurveAsNurbs)
+	{
+		return AddCurveAsNurbs(A3DCurve, UVReparameterization);
+	}
+
+	CADLibrary::TUniqueTSObj<A3DCrvHyperbolaData> CrvHyperbolaData(A3DCurve);
+	if (!CrvHyperbolaData.IsValid())
+	{
+		return TSharedPtr<FCurve>();
+	}
+
+	bool bIs2D = (bool)CrvHyperbolaData->m_bIs2D;
+
+	//setting->SetAttribute("m_bIs2D", (int)sData.m_bIs2D);
+	//traverseTransfo3d(sData.m_sTrsf, crv);
+	//traverseParam(&sData.m_sParam, crv);
+	//_SetDoubleAttribute(crv, "m_dSemiAxis", sData.m_dSemiAxis);
+	//_SetDoubleAttribute(crv, "m_dSemiImageAxis", sData.m_dSemiImageAxis);
+	//crv->SetAttribute("m_cParamType", (int)sData.m_cParamType);
+
+	return TSharedPtr<FCurve>();
+}
+
+TSharedPtr<FCurve> FTechSoftBridge::AddCurveHelix(const A3DCrvHelix* A3DCurve, const TechSoftUtils::FUVReparameterization& UVReparameterization)
+{
+	Report.CurveHelixCount++;
+
+	if (bUseCurveAsNurbs)
+	{
+		return AddCurveAsNurbs(A3DCurve, UVReparameterization);
+	}
+
+	CADLibrary::TUniqueTSObj<A3DCrvHelixData> CrvHelixData(A3DCurve);
+	if (!CrvHelixData.IsValid())
+	{
+		return TSharedPtr<FCurve>();
+	}
+
+	bool bIs2D = (bool)CrvHelixData->m_bIs2D;
+
+	return TSharedPtr<FCurve>();
+}
+
+TSharedPtr<FCurve> FTechSoftBridge::AddCurvePolyLine(const A3DCrvPolyLine* A3DCurve, const TechSoftUtils::FUVReparameterization& UVReparameterization)
+{
+	Report.CurvePolyLineCount++;
+
+	if (bUseCurveAsNurbs)
+	{
+		return AddCurveAsNurbs(A3DCurve, UVReparameterization);
+	}
+
+	CADLibrary::TUniqueTSObj<A3DCrvPolyLineData> CrvPolyLineData(A3DCurve);
+	if (!CrvPolyLineData.IsValid())
+	{
+		return TSharedPtr<FCurve>();
+	}
+
+	bool bIs2D = (bool)CrvPolyLineData->m_bIs2D;
+
+	//setting->SetAttribute("m_bIs2D", (int)sData.m_bIs2D);
+	//traverseTransfo3d(sData.m_sTrsf, crv);
+	//traverseParam(&sData.m_sParam, crv);
+	//traversePoints("m_pPts", sData.m_uiSize, sData.m_pPts, crv);
+
+
+	return TSharedPtr<FCurve>();
+}
+
+TSharedPtr<FCurve> FTechSoftBridge::AddCurveComposite(const A3DCrvComposite* A3DCurve, const TechSoftUtils::FUVReparameterization& UVReparameterization)
+{
+	Report.CurveCompositeCount++;
+
+	if (bUseCurveAsNurbs)
+	{
+		return AddCurveAsNurbs(A3DCurve, UVReparameterization);
+	}
+
+	CADLibrary::TUniqueTSObj<A3DCrvCompositeData> CrvCompositeData(A3DCurve);
+	if (!CrvCompositeData.IsValid())
+	{
+		return TSharedPtr<FCurve>();
+	}
+
+	bool bIs2D = (bool)CrvCompositeData->m_bIs2D;
+
+	//setting->SetAttribute("m_bIs2D", (int)sData.m_bIs2D);
+	//traverseTransfo3d(sData.m_sTrsf, crv);
+	//traverseParam(&sData.m_sParam, crv);
+	//crv->SetAttribute("m_bClosed", (int)sData.m_bClosed);
+	//traverseBools("m_pbSenses", sData.m_uiSize, sData.m_pbSenses, crv);
+
+	//A3DUns32 ui, uiSize = sData.m_uiSize;
+	//for (ui = 0; ui < uiSize; ++ui)
+	//	traverseCurve(sData.m_ppCurves[ui], crv);
+
+	return TSharedPtr<FCurve>();
+}
+
+TSharedPtr<FCurve> AddCurveNurbsFromData(A3DCrvNurbsData& A3DNurbs, const TechSoftUtils::FUVReparameterization& UVReparameterization)
+{
+	FNurbsCurveData Nurbs;
+	Nurbs.Dimension = A3DNurbs.m_bIs2D ? 2 : 3;
+	Nurbs.bIsRational = (bool)A3DNurbs.m_bRational;
+	Nurbs.Degree = A3DNurbs.m_uiDegree;
+
+	TechSoftUtils::FillPointArray(A3DNurbs.m_uiCtrlSize, A3DNurbs.m_pCtrlPts, Nurbs.Poles);
+	if (UVReparameterization.NeedApply())
+	{
+		for (FPoint& Point : Nurbs.Poles)
+		{
+			UVReparameterization.Apply(Point);
+		}
+	}
+
+	TechSoftUtils::FillDoubleArray(A3DNurbs.m_uiKnotSize, A3DNurbs.m_pdKnots, Nurbs.NodalVector);
+	if (Nurbs.bIsRational)
+	{
+		TechSoftUtils::FillDoubleArray(A3DNurbs.m_uiCtrlSize, A3DNurbs.m_pdWeights, Nurbs.Weights);
+	}
+
+	//A3DNurbs.m_eKnotType;
+	//kA3DKnotTypeUniformKnots,			/*!< Uniform. */
+	//kA3DKnotTypeUnspecified,			/*!< No particularity. */
+	//kA3DKnotTypeQuasiUniformKnots,	/*!< Quasi-uniform. */
+	//kA3DKnotTypePieceWiseBezierKnots	/*!< Extrema with multiplicities of degree + 1, internal is degree. */
+	ensureCADKernel(A3DNurbs.m_eKnotType != kA3DKnotTypePieceWiseBezierKnots);
+
+	//A3DNurbs.m_eCurveForm;
+	//kA3DBSplineCurveFormUnspecified,	/*!< No particularity. */
+	//kA3DBSplineCurveFormPolyline,		/*!< Polyline. */
+	//kA3DBSplineCurveFormCircularArc,	/*!< Circle arc. */
+	//kA3DBSplineCurveFormEllipticArc,	/*!< Elliptic arc. */
+	//kA3DBSplineCurveFormParabolicArc,	/*!< Parabolic arc. */
+	//kA3DBSplineCurveFormHyperbolicArc	/*!< Hyperbolic arc. */
+	ensureCADKernel(A3DNurbs.m_eCurveForm != kA3DBSplineCurveFormPolyline);
+	ensureCADKernel(A3DNurbs.m_eCurveForm != kA3DBSplineCurveFormCircularArc);
+	ensureCADKernel(A3DNurbs.m_eCurveForm != kA3DBSplineCurveFormEllipticArc);
+	ensureCADKernel(A3DNurbs.m_eCurveForm != kA3DBSplineCurveFormParabolicArc);
+	ensureCADKernel(A3DNurbs.m_eCurveForm != kA3DBSplineCurveFormHyperbolicArc);
+
+	A3DCrvNurbsGet(NULL, &A3DNurbs);
+
+	return FEntity::MakeShared<FNURBSCurve>(Nurbs);
+}
+
+TSharedPtr<FCurve> FTechSoftBridge::AddCurveNurbs(const A3DCrvNurbs* A3DNurbs, const TechSoftUtils::FUVReparameterization& UVReparameterization)
+{
+	Report.CurveLineCount++;
+
+	CADLibrary::TUniqueTSObj<A3DCrvNurbsData> CrvNurbsData(A3DNurbs);
+	if (!CrvNurbsData.IsValid())
+	{
+		return TSharedPtr<FCurve>();
+	}
+
+	return AddCurveNurbsFromData(*CrvNurbsData, UVReparameterization);
+}
+
+TSharedPtr<FCurve> FTechSoftBridge::AddCurveAsNurbs(const A3DCrvBase* A3DCurve, const TechSoftUtils::FUVReparameterization& UVReparameterization)
+{
+	Report.CurveAsNurbsCount++;
+
+	CADLibrary::TUniqueTSObj<A3DCrvNurbsData> NurbsData;
+
+	A3DDouble Tolerance = 1e-3;
+	A3DBool bUseSameParameterization = true;
+	NurbsData.FillWith(&TechSoftUtils::GetCurveAsNurbs, A3DCurve, Tolerance, bUseSameParameterization);
+
+	if (!NurbsData.IsValid())
+	{
+		return TSharedPtr<FCurve>();
+	}
+
+	return AddCurveNurbsFromData(*NurbsData, UVReparameterization);
+}
+
+TSharedPtr<FTopologicalEdge> FTechSoftBridge::AddEdge(const A3DTopoCoEdge* A3DCoedge, const TSharedRef<FSurface>& Surface, const TechSoftUtils::FUVReparameterization& UVReparameterization, EOrientation& OutOrientation)
+{
+	Report.EdgeCount++;
+
+	CADLibrary::TUniqueTSObj<A3DTopoCoEdgeData> CoEdgeData(A3DCoedge);
+	if (!CoEdgeData.IsValid())
+	{
+		return TSharedPtr<FTopologicalEdge>();
+	}
+
+	ensureCADKernel(CoEdgeData->m_pUVCurve);
+
+	TSharedPtr<FCurve> Curve = AddCurve(CoEdgeData->m_pUVCurve, UVReparameterization);
+	if (!Curve.IsValid())
+	{
+		return TSharedPtr<FTopologicalEdge>();
+	}
+
+	TSharedRef<FRestrictionCurve> RestrictionCurve = FEntity::MakeShared<FRestrictionCurve>(Surface, Curve.ToSharedRef());
+
+	TSharedPtr<FTopologicalEdge> Edge;
+	if (CoEdgeData->m_pEdge)
+	{
+		CADLibrary::TUniqueTSObj<A3DTopoEdgeData> A3DEdgeData(CoEdgeData->m_pEdge);
+		if (A3DEdgeData.IsValid())
+		{
+			double Tolerance = A3DEdgeData->m_dTolerance;
+			bool bHasTrimDomain = (bool)A3DEdgeData->m_bHasTrimDomain;
+
+			if (A3DEdgeData->m_p3dCurve)
+			{
+				//traverseCurve(A3DEdgeData.m_p3dCurve, edge);
+			}
+
+			if (A3DEdgeData->m_bHasTrimDomain)
+			{
+				FLinearBoundary Boundary = TechSoftUtils::GetLinearBoundary(A3DEdgeData->m_sInterval);
+				Edge = FTopologicalEdge::Make(RestrictionCurve, Boundary);
+			}
+		}
+	}
+
+	if (!Edge.IsValid())
+	{
+		Edge = FTopologicalEdge::Make(RestrictionCurve);
+	}
+
+	if (!Edge.IsValid())
+	{
+		return TSharedPtr<FTopologicalEdge>();
+	}
+
+	A3DEdgeToEdge.Emplace(A3DCoedge, Edge);
+
+	//CoEdgeData->m_ucOrientationUVWithLoop;
+	//CoEdgeData->m_ucOrientationWithLoop;
+	OutOrientation = CoEdgeData->m_ucOrientationUVWithLoop > 0 ? EOrientation::Front : EOrientation::Back;
+
+	// Link edges
+	if (CoEdgeData->m_pNeighbor)
+	{
+		const A3DTopoCoEdge* Neighbor = CoEdgeData->m_pNeighbor;
+		while (Neighbor && Neighbor != A3DCoedge)
+		{
+			TSharedPtr<FTopologicalEdge>* TwinEdge = A3DEdgeToEdge.Find(Neighbor);
+			if (TwinEdge != nullptr)
+			{
+				Edge->Link(*TwinEdge->Get(), SquareJoiningVertexTolerance);
+			}
+
+			CADLibrary::TUniqueTSObj<A3DTopoCoEdgeData> NeighborData(Neighbor);
+			if (NeighborData.IsValid())
+			{
+				Neighbor = NeighborData->m_pNeighbor;
+			}
+			else
+			{
+				break;
+			}
+		}
+	}
+
+	return Edge;
+}
+
+TSharedPtr<FTopologicalLoop> FTechSoftBridge::AddLoop(const A3DTopoLoop* A3DLoop, const TSharedRef<FSurface>& Surface, const TechSoftUtils::FUVReparameterization& UVReparameterization)
+{
+	Report.LoopCount++;
+
+	TArray<TSharedPtr<FTopologicalEdge>> Edges;
+	TArray<EOrientation> Directions;
+
+	CADLibrary::TUniqueTSObj<A3DTopoLoopData> TopoLoopData(A3DLoop);
+	if (!TopoLoopData.IsValid())
+	{
+		Report.DegeneratedLoopCount++;
+		return TSharedPtr<FTopologicalLoop>();
+	}
+
+	bool bLoopOrientation = (bool)TopoLoopData->m_ucOrientationWithSurface;
+	for (A3DUns32 Index = 0; Index < TopoLoopData->m_uiCoEdgeSize; ++Index)
+	{
+		EOrientation Orientation;
+		TSharedPtr<FTopologicalEdge> Edge = AddEdge(TopoLoopData->m_ppCoEdges[Index], Surface, UVReparameterization, Orientation);
+		if (!Edge.IsValid())
+		{
+			Report.DegeneratedEdgeCount++;
+			continue;
+		}
+
+		Edges.Emplace(Edge);
+		Directions.Emplace(Orientation);
+	}
+
+	//int32 CoedgeCount = CTCoedgeIds.Count();
+	//Edges.Reserve(CoedgeCount);
+	//Directions.Reserve(CoedgeCount);
+
+	//CTCoedgeIds.IteratorInitialize();
+
+	//CT_OBJECT_ID CoEdgeId;
+	//while ((CoEdgeId = CTCoedgeIds.IteratorIter()) != 0)
+	//{
+	//	TSharedPtr<FTopologicalEdge> Edge = AddEdge(CoEdgeId, Surface);
+	//	if (!Edge.IsValid())
+	//	{
+	//		continue;
+	//	}
+
+	//	Edges.Emplace(Edge);
+	//	Directions.Emplace(EOrientation::Front);
+	//}
+
+	if (Edges.Num() == 0)
+	{
+		Report.DegeneratedLoopCount++;
+		return TSharedPtr<FTopologicalLoop>();
+	}
+
+	return FTopologicalLoop::Make(Edges, Directions, GeometricTolerance);
+}
+
+void FTechSoftBridge::AddFace(const A3DTopoFace* A3DFace, TSharedRef<FShell>& Shell)
+{
+	Report.FaceCount++;
+
+	FEntityData MetaData;
+	TechSoftUtils::TraverseSource(A3DFace, MetaData);
+
+	CADLibrary::TUniqueTSObj<A3DTopoFaceData> TopoFaceData(A3DFace);
+	if (!TopoFaceData.IsValid())
+	{
+		Report.FailedFaceCount++;
+		return;
+	}
+
+	const A3DSurfBase* A3DSurface = TopoFaceData->m_pSurface;
+	TechSoftUtils::FUVReparameterization UVReparameterization = TechSoftUtils::FUVReparameterization::Identity;
+	TSharedPtr<FSurface> SurfacePtr = AddSurface(A3DSurface, UVReparameterization);
+	if (!SurfacePtr.IsValid())
+	{
+		Report.DegeneratedSurfaceCount++;
+		Report.FailedFaceCount++;
+		return;
+	}
+	TSharedRef<FSurface> Surface = SurfacePtr.ToSharedRef();
+
+	TSharedRef<FTopologicalFace> Face = FEntity::MakeShared<FTopologicalFace>(Surface);
+
+	if (TopoFaceData->m_bHasTrimDomain)
+	{
+		FSurfacicBoundary SurfaceBoundary = TechSoftUtils::GetSurfacicBoundary(TopoFaceData->m_sSurfaceDomain, UnitScale);
+		Surface->TrimBoundaryTo(SurfaceBoundary);
+	}
+
+	if (!TopoFaceData->m_uiLoopSize)
+	{
+		Face->ApplyNaturalLoops();
+	}
+	else
+	{
+		TArray<TSharedPtr<FTopologicalLoop>> Loops;
+
+		for (A3DUns32 Index = 0; Index < TopoFaceData->m_uiLoopSize; ++Index)
+		{
+			TSharedPtr<FTopologicalLoop> Loop = FTechSoftBridge::AddLoop(TopoFaceData->m_ppLoops[Index], Surface, UVReparameterization);
+			if (!Loop.IsValid())
+			{
+				continue;
+			}
+
+			TArray<FPoint2D> LoopSampling;
+			Loop->Get2DSampling(LoopSampling);
+			FAABB2D Boundary;
+			Boundary += LoopSampling;
+			Loop->Boundary.Set(Boundary.GetMin(), Boundary.GetMax());
+
+			// Check if the loop is not composed with only degenerated edge
+			bool bDegeneratedLoop = true;
+			for (const FOrientedEdge& Edge : Loop->GetEdges())
+			{
+				if (!Edge.Entity->IsDegenerated())
+				{
+					bDegeneratedLoop = false;
+					break;
+				}
+			}
+			if (bDegeneratedLoop)
+			{
+				continue;
+			}
+
+			Loops.Add(Loop);
+		}
+
+		if (Loops.Num() == 0)
+		{
+			FMessage::Printf(Log, TEXT("The Face has no loop, this face is bounded by its natural boundaries\n"));
+			//Face->ApplyNaturalLoops();
+			Report.FailedFaceCount++;
+			return; //?
+		}
+		else
+		{
+			Face->AddLoops(Loops, Report.DoubtfulLoopOrientationCount);
+		}
+	}
+
+	//AddMetadata(CTFaceId, Face);
+	//Face->SetPatchId((int32)CTFaceId);
+
+	EOrientation Orientation = EOrientation::Front; //CTOrientation == CT_FORWARD ? EOrientation::Front : EOrientation::Back;
+	Shell->Add(Face, Orientation);
+}
+
+static bool bUseSurfaceAsNurbs = true;
+
+TSharedPtr<FSurface> FTechSoftBridge::AddSurface(const A3DSurfBase* A3DSurface, TechSoftUtils::FUVReparameterization& OutUVReparameterization)
+{
+	Report.SurfaceCount++;
+
+	FEntityData MetaData;
+	TechSoftUtils::TraverseSource(A3DSurface, MetaData);
+
+	A3DEEntityType Type;
+	int32 Ret = CADLibrary::TechSoftUtils::GetEntityType(A3DSurface, &Type);
+	if (Ret == A3D_SUCCESS)
+	{
+		switch (Type)
+		{
+		case kA3DTypeSurfBlend01:
+			return AddBlend01Surface(A3DSurface, OutUVReparameterization);
+			
+		case kA3DTypeSurfBlend02:
+			return AddBlend02Surface(A3DSurface, OutUVReparameterization);
+
+		case kA3DTypeSurfBlend03:
+			return AddBlend03Surface(A3DSurface, OutUVReparameterization);
+
+		case kA3DTypeSurfNurbs:
+			return AddNurbsSurface(A3DSurface, OutUVReparameterization);
+
+		case kA3DTypeSurfCone:
+			OutUVReparameterization.SetScale(1, UnitScale);
+			return AddConeSurface(A3DSurface, OutUVReparameterization);
+
+		case kA3DTypeSurfCylinder:
+			OutUVReparameterization.SetScale(1, UnitScale);
+			return AddCylinderSurface(A3DSurface, OutUVReparameterization);
+
+		case kA3DTypeSurfCylindrical:
+			return AddCylindricalSurface(A3DSurface, OutUVReparameterization);
+
+		case kA3DTypeSurfOffset:
+			return AddOffsetSurface(A3DSurface, OutUVReparameterization);
+
+		case kA3DTypeSurfPipe:
+			return AddPipeSurface(A3DSurface, OutUVReparameterization);
+
+		case kA3DTypeSurfPlane:
+			OutUVReparameterization.SetScale(UnitScale, UnitScale);
+			return AddPlaneSurface(A3DSurface, OutUVReparameterization);
+
+		case kA3DTypeSurfRuled:
+			return AddRuledSurface(A3DSurface, OutUVReparameterization);
+
+		case kA3DTypeSurfSphere:
+			return AddSphereSurface(A3DSurface, OutUVReparameterization);
+
+		case kA3DTypeSurfRevolution:
+			return AddRevolutionSurface(A3DSurface, OutUVReparameterization);
+
+		case kA3DTypeSurfExtrusion:
+			return AddExtrusionSurface(A3DSurface, OutUVReparameterization);
+
+		case kA3DTypeSurfFromCurves:
+			return AddSurfaceFromCurves(A3DSurface, OutUVReparameterization);
+
+		case kA3DTypeSurfTorus:
+			return AddTorusSurface(A3DSurface, OutUVReparameterization);
+
+		case kA3DTypeSurfTransform:
+			return AddTransformSurface(A3DSurface, OutUVReparameterization);
+
+		default:
+			return AddSurfaceAsNurbs(A3DSurface, OutUVReparameterization);
+		}
+	}
+	else if (Ret == A3D_NOT_IMPLEMENTED)
+	{
+		return AddSurfaceAsNurbs(A3DSurface, OutUVReparameterization);
+	}
+	return TSharedPtr<FSurface>();
+}
+
+TSharedPtr<FSurface> FTechSoftBridge::AddConeSurface(const A3DSurfBase* Surface, TechSoftUtils::FUVReparameterization& OutUVReparameterization)
+{
+	Report.ConeSurfaceCount++;
+
+	CADLibrary::TUniqueTSObj<A3DSurfConeData> A3DConeData(Surface);
+	if (!A3DConeData.IsValid())
+	{
+		return TSharedPtr<FSurface>();
+	}
+
+	OutUVReparameterization.AddUVTransform(A3DConeData->m_sParam);
+	FMatrixH CoordinateSystem = TechSoftUtils::CreateCoordinateSystem(A3DConeData->m_sTrsf, UnitScale);
+	FSurfacicBoundary Boundary = TechSoftUtils::GetSurfacicBoundary(A3DConeData->m_sParam.m_sUVDomain, UnitScale);
+	TSharedPtr<FSurface> Cone = FEntity::MakeShared<FConeSurface>(GeometricTolerance, CoordinateSystem, A3DConeData->m_dRadius * UnitScale, A3DConeData->m_dSemiAngle, Boundary);
+
+	return Cone;
+}
+
+TSharedPtr<FSurface> FTechSoftBridge::AddCylinderSurface(const A3DSurfBase* Surface, TechSoftUtils::FUVReparameterization& OutUVReparameterization)
+{
+	Report.CylinderSurfaceCount++;
+
+	CADLibrary::TUniqueTSObj<A3DSurfCylinderData> A3DCylinderData(Surface);
+	if (!A3DCylinderData.IsValid())
+	{
+		return TSharedPtr<FSurface>();
+	}
+
+	OutUVReparameterization.AddUVTransform(A3DCylinderData->m_sParam);
+	FMatrixH CoordinateSystem = TechSoftUtils::CreateCoordinateSystem(A3DCylinderData->m_sTrsf, UnitScale);
+	FSurfacicBoundary Boundary = TechSoftUtils::GetSurfacicBoundary(A3DCylinderData->m_sParam.m_sUVDomain, UnitScale);
+	TSharedPtr<FSurface> Cylinder = FEntity::MakeShared<FCylinderSurface>(GeometricTolerance, CoordinateSystem, A3DCylinderData->m_dRadius * UnitScale, Boundary);
+
+	return Cylinder;
+}
+
+TSharedPtr<FSurface> FTechSoftBridge::AddLinearTransfoSurface(const A3DSurfBase* Surface, TechSoftUtils::FUVReparameterization& OutUVReparameterization)
+{
+	Report.LinearTransfoSurfaceCount++;
+
+	if (bUseSurfaceAsNurbs)
+	{
+		return AddSurfaceAsNurbs(Surface, OutUVReparameterization);
+	}
+
+	//TUniqueTSObj<A3DSurfLinearTransfoData A3DLinearTransfoData;
+	//A3D_INITIALIZE_DATA(A3DSurfLinearTransfoData, A3DLinearTransfoData);
+	//A3DInt32 Ret = A3DSurfLinearTransfoGet(Surface, &A3DLinearTransfoData);
+	//if (Ret != A3D_SUCCESS)
+	//{
+	//	return TSharedPtr<FSurface>();
+	//}
+
+	//A3DSurfLinearTransfoGet(NULL, &A3DLinearTransfoData);
+	return TSharedPtr<FSurface>();
+}
+
+TSharedPtr<FSurface> FTechSoftBridge::AddNurbsSurface(const A3DSurfNurbs* Nurbs, TechSoftUtils::FUVReparameterization& OutUVReparameterization)
+{
+	Report.NurbsSurfaceCount++;
+
+	CADLibrary::TUniqueTSObj<A3DSurfNurbsData> A3DNurbsData(Nurbs);
+	if (!A3DNurbsData.IsValid())
+	{
+		return TSharedPtr<FSurface>();
+	}
+
+	return AddSurfaceNurbs(*A3DNurbsData, OutUVReparameterization);
+}
+
+TSharedPtr<FSurface> FTechSoftBridge::AddOffsetSurface(const A3DSurfBase* Surface, TechSoftUtils::FUVReparameterization& OutUVReparameterization)
+{
+	Report.OffsetSurfaceCount++;
+
+	if (bUseSurfaceAsNurbs)
+	{
+		return AddSurfaceAsNurbs(Surface, OutUVReparameterization);
+	}
+
+	return TSharedPtr<FSurface>();
+}
+
+TSharedPtr<FSurface> FTechSoftBridge::AddPlaneSurface(const A3DSurfPlane* Surface, TechSoftUtils::FUVReparameterization& OutUVReparameterization)
+{
+	Report.PlaneSurfaceCount++;
+
+	CADLibrary::TUniqueTSObj<A3DSurfPlaneData> A3DPlaneData(Surface);
+	if (!A3DPlaneData.IsValid())
+	{
+		return TSharedPtr<FSurface>();
+	}
+
+	OutUVReparameterization.AddUVTransform(A3DPlaneData->m_sParam);
+	FMatrixH CoordinateSystem = TechSoftUtils::CreateCoordinateSystem(A3DPlaneData->m_sTrsf, UnitScale);
+	FSurfacicBoundary Boundary = TechSoftUtils::GetSurfacicBoundary(A3DPlaneData->m_sParam.m_sUVDomain, UnitScale);
+	TSharedPtr<FSurface> Plane = FEntity::MakeShared<FPlaneSurface>(GeometricTolerance, CoordinateSystem, Boundary);
+
+	return Plane;
+}
+
+TSharedPtr<FSurface> FTechSoftBridge::AddRevolutionSurface(const A3DSurfRevolution* Surface, TechSoftUtils::FUVReparameterization& OutUVReparameterization)
+{
+	Report.RevolutionSurfaceCount++;
+
+	if (bUseSurfaceAsNurbs)
+	{
+		return AddSurfaceAsNurbs(Surface, OutUVReparameterization);
+	}
+
+	CADLibrary::TUniqueTSObj<A3DSurfRevolutionData> A3DRevolutionData(Surface);
+	if (!A3DRevolutionData.IsValid())
+	{
+		return TSharedPtr<FSurface>();
+	}
+
+	return TSharedPtr<FSurface>();
+}
+
+TSharedPtr<FSurface> FTechSoftBridge::AddRuledSurface(const A3DSurfBase* Surface, TechSoftUtils::FUVReparameterization& OutUVReparameterization)
+{
+	Report.RuledSurfaceCount++;
+
+	if (bUseSurfaceAsNurbs)
+	{
+		return AddSurfaceAsNurbs(Surface, OutUVReparameterization);
+	}
+
+	CADLibrary::TUniqueTSObj<A3DSurfRuledData> A3DRuledData(Surface);
+	if (!A3DRuledData.IsValid())
+	{
+		return TSharedPtr<FSurface>();
+	}
+
+	return TSharedPtr<FSurface>();
+}
+
+TSharedPtr<FSurface> FTechSoftBridge::AddSphereSurface(const A3DSurfBase* Surface, TechSoftUtils::FUVReparameterization& OutUVReparameterization)
+{
+	Report.SphereSurfaceCount++;
+
+	if (bUseSurfaceAsNurbs)
+	{
+		return AddSurfaceAsNurbs(Surface, OutUVReparameterization);
+	}
+
+	CADLibrary::TUniqueTSObj<A3DSurfSphereData> A3DSphereData(Surface);
+	if (!A3DSphereData.IsValid())
+	{
+		return TSharedPtr<FSurface>();
+	}
+
+	OutUVReparameterization.AddUVTransform(A3DSphereData->m_sParam);
+	FMatrixH CoordinateSystem = TechSoftUtils::CreateCoordinateSystem(A3DSphereData->m_sTrsf, UnitScale);
+	FSurfacicBoundary Boundary = TechSoftUtils::GetSurfacicBoundary(A3DSphereData->m_sParam.m_sUVDomain, UnitScale);
+	TSharedRef<FSphericalSurface> Sphere = FEntity::MakeShared<FSphericalSurface>(GeometricTolerance, CoordinateSystem, A3DSphereData->m_dRadius * UnitScale, Boundary);
+
+	return Sphere;
+}
+
+TSharedPtr<FSurface> FTechSoftBridge::AddTorusSurface(const A3DSurfBase* Surface, TechSoftUtils::FUVReparameterization& OutUVReparameterization)
+{
+	Report.TorusSurfaceCount++;
+
+	CADLibrary::TUniqueTSObj<A3DSurfTorusData> A3DTorusData(Surface);
+	if (!A3DTorusData.IsValid())
+	{
+		return TSharedPtr<FSurface>();
+	}
+
+	OutUVReparameterization.AddUVTransform(A3DTorusData->m_sParam);
+	FMatrixH CoordinateSystem = TechSoftUtils::CreateCoordinateSystem(A3DTorusData->m_sTrsf, UnitScale);
+	FSurfacicBoundary Boundary = TechSoftUtils::GetSurfacicBoundary(A3DTorusData->m_sParam.m_sUVDomain, UnitScale);
+	TSharedPtr<FSurface> Torus = FEntity::MakeShared<FTorusSurface>(GeometricTolerance, CoordinateSystem, A3DTorusData->m_dMajorRadius * UnitScale, A3DTorusData->m_dMinorRadius * UnitScale, Boundary);
+
+	return Torus;
+}
+
+TSharedPtr<FSurface> FTechSoftBridge::AddBlend01Surface(const A3DSurfBase* Surface, TechSoftUtils::FUVReparameterization& OutUVReparameterization)
+{
+	Report.Blend01SurfaceCount++;
+
+	if (bUseSurfaceAsNurbs)
+	{
+		return AddSurfaceAsNurbs(Surface, OutUVReparameterization);
+	}
+
+	CADLibrary::TUniqueTSObj<A3DSurfBlend01Data> A3DBlend01Data(Surface);
+	if (!A3DBlend01Data.IsValid())
+	{
+		return TSharedPtr<FSurface>();
+	}
+
+	return TSharedPtr<FSurface>();
+}
+
+TSharedPtr<FSurface> FTechSoftBridge::AddBlend02Surface(const A3DSurfBase* Surface, TechSoftUtils::FUVReparameterization& OutUVReparameterization)
+{
+	Report.Blend02SurfaceCount++;
+
+	if (bUseSurfaceAsNurbs)
+	{
+		return AddSurfaceAsNurbs(Surface, OutUVReparameterization);
+	}
+
+	CADLibrary::TUniqueTSObj<A3DSurfBlend02Data> A3DBlend02Data(Surface);
+	if (!A3DBlend02Data.IsValid())
+	{
+		return TSharedPtr<FSurface>();
+	}
+
+	return TSharedPtr<FSurface>();
+}
+
+TSharedPtr<FSurface> FTechSoftBridge::AddBlend03Surface(const A3DSurfBase* Surface, TechSoftUtils::FUVReparameterization& OutUVReparameterization)
+{
+	Report.Blend03SurfaceCount++;
+
+	if (bUseSurfaceAsNurbs)
+	{
+		return AddSurfaceAsNurbs(Surface, OutUVReparameterization);
+	}
+
+#ifdef WIP
+	CADLibrary::TUniqueTSObj<A3DSurfBlend03Data> A3DBlend03Data(Surface);
+	if (!A3DBlend03Data.IsValid())
+	{
+		return TSharedPtr<FSurface>();
+	}
+
+	FBlend03Data BlendData;
+
+	F3DDebugSession _(TEXT("Blend03Surface"));
+	{
+
+		BlendData.CenterCurve.Reserve(A3DBlend03Data->m_uiNbOfElement);
+		BlendData.RailCurve1.Reserve(A3DBlend03Data->m_uiNbOfElement);
+		BlendData.RailCurve2.Reserve(A3DBlend03Data->m_uiNbOfElement);
+
+		BlendData.Parameters.Reserve(A3DBlend03Data->m_uiNbOfElement);
+		BlendData.Multiplicities.Reserve(A3DBlend03Data->m_uiNbOfElement);
+		BlendData.RailToAnglesV.Reserve(A3DBlend03Data->m_uiNbOfElement);
+		BlendData.Tangents.Reserve(A3DBlend03Data->m_uiNbOfElement);
+		BlendData.RailToDerivativesV.Reserve(A3DBlend03Data->m_uiNbOfElement);
+		BlendData.SecondDerivatives.Reserve(A3DBlend03Data->m_uiNbOfElement * 3);
+		BlendData.RailToSecondDerivatives.Reserve(A3DBlend03Data->m_uiNbOfElement);
+
+		BlendData.RailToParameterV = A3DBlend03Data->m_dRail2ParameterV;
+		BlendData.TrimVMin= A3DBlend03Data->m_dTrimVMin;
+		BlendData.TrimVMax= A3DBlend03Data->m_dTrimVMax;
+		BlendData.InitialMaxV = A3DBlend03Data->m_dInitialMaxV;
+
+		TArray<FPoint> Positions;
+		A3DBlend03Data->m_pPositions;
+		TechSoftUtils::FillPointArray(A3DBlend03Data->m_uiNbOfElement*3, A3DBlend03Data->m_pPositions, Positions, UnitScale);
+
+		for (uint32 Index = 0; Index < (3 * A3DBlend03Data->m_uiNbOfElement);)
+		{
+			BlendData.RailCurve1.Add(Positions[Index++]);
+			BlendData.RailCurve2.Add(Positions[Index++]);
+			BlendData.CenterCurve.Add(Positions[Index++]);
+		}
+
+		TechSoftUtils::FillPointArray(A3DBlend03Data->m_uiNbOfElement * 3, A3DBlend03Data->m_pTangent, Positions, 1);
+		TechSoftUtils::FillPointArray(A3DBlend03Data->m_uiNbOfElement * 3, A3DBlend03Data->m_pSecondDerivatives, Positions, 1);
+
+		TechSoftUtils::FillDoubleArray(A3DBlend03Data->m_uiNbOfElement, A3DBlend03Data->m_pdParameters, BlendData.Parameters);
+		TechSoftUtils::FillInt32Array(A3DBlend03Data->m_uiNbOfElement, A3DBlend03Data->m_piMultiplicities, BlendData.Multiplicities);
+		TechSoftUtils::FillDoubleArray(A3DBlend03Data->m_uiNbOfElement, A3DBlend03Data->m_pdRail2AnglesV, BlendData.RailToAnglesV);
+		TechSoftUtils::FillDoubleArray(A3DBlend03Data->m_uiNbOfElement, A3DBlend03Data->m_pdRail2DerivativesV, BlendData.RailToDerivativesV);
+		TechSoftUtils::FillDoubleArray(A3DBlend03Data->m_uiNbOfElement, A3DBlend03Data->m_pdRail2SecondDerivatives, BlendData.RailToSecondDerivatives);
+
+		DisplayPoints(TEXT("Center"), BlendData.CenterCurve);
+		DisplayPoints(TEXT("RailCurve1"), BlendData.RailCurve1, EVisuProperty::YellowPoint);
+		DisplayPoints(TEXT("RailCurve2"), BlendData.RailCurve2, EVisuProperty::GreenPoint);
+
+		TechSoftUtils::FillPointArray(A3DBlend03Data->m_uiNbOfElement * 3, A3DBlend03Data->m_pPositions, Positions, UnitScale);
+
+		FNurbsCurveData NurbsCurve;
+		NurbsCurve.Dimension = 3;
+		NurbsCurve.Degree = 5;
+		for (uint32 Index = 0; Index < A3DBlend03Data->m_uiNbOfElement; ++Index)
+		{
+			for (int32 Mndex = 0; Mndex < BlendData.Multiplicities[Index]; ++Mndex)
+			{
+				NurbsCurve.NodalVector.Add(BlendData.Parameters[Index]);
+			}
+		}
+		NurbsCurve.bIsRational = false;
+		NurbsCurve.Poles = BlendData.CenterCurve;
+
+		TSharedPtr<FNURBSCurve> CenterNurbs = FEntity::MakeShared<FNURBSCurve>(NurbsCurve);
+		{
+			F3DDebugSession _(TEXT("Center"));
+			Display(*CenterNurbs);
+		}
+
+		NurbsCurve.Poles = BlendData.RailCurve1;
+		TSharedPtr<FNURBSCurve> Rail1Nurbs = FEntity::MakeShared<FNURBSCurve>(NurbsCurve);
+		{
+			F3DDebugSession _(TEXT("Rail1Nurbs"));
+			Display(*Rail1Nurbs);
+		}
+
+		NurbsCurve.Poles = BlendData.RailCurve2;
+		TSharedPtr<FNURBSCurve> Rail2Nurbs = FEntity::MakeShared<FNURBSCurve>(NurbsCurve);
+		{
+			F3DDebugSession _(TEXT("Rail2Nurbs"));
+			Display(*Rail2Nurbs);
+		}
+
+	}
+#endif
+	return TSharedPtr<FSurface>();
+}
+
+TSharedPtr<FSurface> FTechSoftBridge::AddCylindricalSurface(const A3DSurfBase* Surface, TechSoftUtils::FUVReparameterization& OutUVReparameterization)
+{
+	Report.CylindricalSurfaceCount++;
+
+	if (bUseSurfaceAsNurbs)
+	{
+		return AddSurfaceAsNurbs(Surface, OutUVReparameterization);
+	}
+
+	CADLibrary::TUniqueTSObj<A3DSurfCylindricalData> A3DCylindricalData(Surface);
+	if (!A3DCylindricalData.IsValid())
+	{
+		return TSharedPtr<FSurface>();
+	}
+
+	return TSharedPtr<FSurface>();
+}
+
+TSharedPtr<FSurface> FTechSoftBridge::AddPipeSurface(const A3DSurfBase* Surface, TechSoftUtils::FUVReparameterization& OutUVReparameterization)
+{
+	Report.PipeSurfaceCount++;
+
+	if (bUseSurfaceAsNurbs)
+	{
+		return AddSurfaceAsNurbs(Surface, OutUVReparameterization);
+	}
+
+	CADLibrary::TUniqueTSObj<A3DSurfPipeData> A3DPipeData(Surface);
+	if (!A3DPipeData.IsValid())
+	{
+		return TSharedPtr<FSurface>();
+	}
+
+	return TSharedPtr<FSurface>();
+}
+
+TSharedPtr<FSurface> FTechSoftBridge::AddExtrusionSurface(const A3DSurfBase* Surface, TechSoftUtils::FUVReparameterization& OutUVReparameterization)
+{
+	Report.ExtrusionSurfaceCount++;
+
+	if (bUseSurfaceAsNurbs)
+	{
+		return AddSurfaceAsNurbs(Surface, OutUVReparameterization);
+	}
+
+	CADLibrary::TUniqueTSObj<A3DSurfExtrusionData> A3DExtrusionData(Surface);
+	if (!A3DExtrusionData.IsValid())
+	{
+		return TSharedPtr<FSurface>();
+	}
+
+	return TSharedPtr<FSurface>();
+}
+
+TSharedPtr<FSurface> FTechSoftBridge::AddSurfaceFromCurves(const A3DSurfBase* Surface, TechSoftUtils::FUVReparameterization& OutUVReparameterization)
+{
+	Report.SurfaceFromCurvesCount++;
+
+	if (bUseSurfaceAsNurbs)
+	{
+		return AddSurfaceAsNurbs(Surface, OutUVReparameterization);
+	}
+
+	CADLibrary::TUniqueTSObj<A3DSurfFromCurvesData> A3DFromCurvesData(Surface);
+	if (!A3DFromCurvesData.IsValid())
+	{
+		return TSharedPtr<FSurface>();
+	}
+
+	return TSharedPtr<FSurface>();
+}
+
+TSharedPtr<FSurface> FTechSoftBridge::AddTransformSurface(const A3DSurfBase* Surface, TechSoftUtils::FUVReparameterization& OutUVReparameterization)
+{
+	Report.TransformSurfaceCount++;
+
+	if (bUseSurfaceAsNurbs)
+	{
+		return AddSurfaceAsNurbs(Surface, OutUVReparameterization);
+	}
+
+	CADLibrary::TUniqueTSObj<A3DSurfFromCurvesData> A3DTransformData(Surface);
+	if (!A3DTransformData.IsValid())
+	{
+		return TSharedPtr<FSurface>();
+	}
+
+	return TSharedPtr<FSurface>();
+}
+
+TSharedPtr<FSurface> FTechSoftBridge::AddSurfaceNurbs(const A3DSurfNurbsData& A3DNurbsData, TechSoftUtils::FUVReparameterization& OutUVReparameterization)
+{
+	FNurbsSurfaceData NurbsData;
+
+	A3DNurbsData.m_eKnotType;
+	//kA3DKnotTypeUniformKnots,			/*!< Uniform. */
+	//kA3DKnotTypeUnspecified,			/*!< No particularity. */
+	//kA3DKnotTypeQuasiUniformKnots,	/*!< Quasi-uniform. */
+	//kA3DKnotTypePieceWiseBezierKnots	/*!< Extrema with multiplicities of degree + 1, internal is degree. */
+
+	A3DNurbsData.m_eSurfaceForm;
+	//kA3DBSplineSurfaceFormPlane,				/*!< Planar surface. */
+	//kA3DBSplineSurfaceFormCylindrical,		/*!< Cylindrical surface. */
+	//kA3DBSplineSurfaceFormConical,			/*!< Conical surface. */
+	//kA3DBSplineSurfaceFormSpherical,			/*!< Spherical surface. */
+	//kA3DBSplineSurfaceFormRevolution,			/*!< Surface of revolution. */
+	//kA3DBSplineSurfaceFormRuled,				/*!< Ruled surface. */
+	//kA3DBSplineSurfaceFormGeneralizedCone,	/*!< Cone. */
+	//kA3DBSplineSurfaceFormQuadric,			/*!< Quadric surface. */
+	//kA3DBSplineSurfaceFormLinearExtrusion,	/*!< Surface of extrusion. */
+	//kA3DBSplineSurfaceFormUnspecified,		/*!< Unspecified. */
+	//kA3DBSplineSurfaceFormPolynomial			/*!< Polynomial surface. */
+
+	NurbsData.PoleUCount = A3DNurbsData.m_uiUCtrlSize;
+	NurbsData.PoleVCount = A3DNurbsData.m_uiVCtrlSize;
+	int32 PoleCount = A3DNurbsData.m_uiUCtrlSize * A3DNurbsData.m_uiVCtrlSize;
+
+	NurbsData.UDegree = A3DNurbsData.m_uiUDegree;
+	NurbsData.VDegree = A3DNurbsData.m_uiVDegree;
+
+	TechSoftUtils::FillDoubleArray(A3DNurbsData.m_uiUKnotSize, A3DNurbsData.m_pdUKnots, NurbsData.UNodalVector);
+	TechSoftUtils::FillDoubleArray(A3DNurbsData.m_uiVKnotSize, A3DNurbsData.m_pdVKnots, NurbsData.VNodalVector);
+
+	TArray<FPoint> Poles;
+	TechSoftUtils::FillPointArray(NurbsData.PoleUCount, NurbsData.PoleVCount, A3DNurbsData.m_pCtrlPts, NurbsData.Poles);
+	if (!FMath::IsNearlyEqual(UnitScale, 1.))
+	{
+		for (FPoint& Point : NurbsData.Poles)
+		{
+			Point *= UnitScale;
+		}
+	}
+
+	bool bIsRational = false;
+	if (A3DNurbsData.m_pdWeights)
+	{
+		bIsRational = true;
+		TechSoftUtils::FillDoubleArray(NurbsData.PoleUCount, NurbsData.PoleVCount, A3DNurbsData.m_pdWeights, NurbsData.Weights);
+	}
+
+	return FEntity::MakeShared<FNURBSSurface>(GeometricTolerance, NurbsData);
+}
+
+TSharedPtr<FSurface> FTechSoftBridge::AddSurfaceAsNurbs(const A3DSurfBase* Surface, TechSoftUtils::FUVReparameterization& OutUVReparameterization)
+{
+	Report.SurfaceAsNurbsCount++;
+
+	CADLibrary::TUniqueTSObj<A3DSurfNurbsData> A3DNurbsData;
+
+	A3DDouble Tolerance = 1e-3;
+	A3DBool bUseSameParameterization = true;
+	A3DNurbsData.FillWith(&TechSoftUtils::GetSurfaceAsNurbs, Surface, Tolerance, bUseSameParameterization);
+
+	if (!A3DNurbsData.IsValid())
+	{
+		return TSharedPtr<FSurface>();
+	}
+
+	return AddSurfaceNurbs(*A3DNurbsData, OutUVReparameterization);
+
+}
+
+void FTechSoftBridge::AddMetadata(FEntityData& MetaData, TSharedRef<FMetadataDictionary> Entity)
+{
+	//Entity->SetHostId(CTNodeId);
+	Entity->SetName(MetaData.Name);
+	//Entity->SetLayer(LayerId);
+	//uint32 ColorHId = CADLibrary::BuildColorId(ColorId, Alpha);
+	//Entity->SetColorId(ColorHId);
+	//Entity->SetPatchId();
+}
+
+}
+
+
+#endif
\ No newline at end of file
diff --git a/Engine/Plugins/Enterprise/DatasmithCADImporter/Source/CADInterfaces/Private/TechSoftBridge.h b/Engine/Plugins/Enterprise/DatasmithCADImporter/Source/CADInterfaces/Private/TechSoftBridge.h
new file mode 100644
--- /dev/null
+++ b/Engine/Plugins/Enterprise/DatasmithCADImporter/Source/CADInterfaces/Private/TechSoftBridge.h
@@ -0,0 +1,315 @@
+// Copyright Epic Games, Inc. All Rights Reserved.
+#pragma once
+#include "CADKernel/Core/Types.h"
+
+#ifdef USE_TECHSOFT_SDK
+
+#include "CADEnum.h"
+#include "CADFileReport.h"
+
+#include "CADKernel/Geo/GeoEnum.h"
+#include "CADKernel/Math/MatrixH.h"
+
+THIRD_PARTY_INCLUDES_START
+#include "A3DSDKEnums.h"
+#include "A3DSDKGeometry.h"
+#include "A3DSDKGeometrySrf.h"
+THIRD_PARTY_INCLUDES_END
+
+typedef void A3DAsmProductOccurrence;
+typedef void A3DCrvBase;
+typedef void A3DCrvBase;
+typedef void A3DCrvBlend02Boundary;
+typedef void A3DCrvCircle;
+typedef void A3DCrvComposite;
+typedef void A3DCrvEllipse;
+typedef void A3DCrvEquation;
+typedef void A3DCrvHelix;
+typedef void A3DCrvHyperbola;
+typedef void A3DCrvIntersection;
+typedef void A3DCrvLine;
+typedef void A3DCrvNurbs;
+typedef void A3DCrvOffset;
+typedef void A3DCrvOnSurf;
+typedef void A3DCrvParabola;
+typedef void A3DCrvPolyLine;
+typedef void A3DCrvTransform;
+typedef void A3DEntity;
+typedef void A3DGraphics;
+typedef void A3DMiscCartesianTransformation;
+typedef void A3DMiscCartesianTransformation;
+typedef void A3DMiscGeneralTransformation;
+typedef void A3DRiBrepModel;
+typedef void A3DRiCoordinateSystem;
+typedef void A3DRiRepresentationItem;
+typedef void A3DSurfBase;
+typedef void A3DSurfBase;
+typedef void A3DSurfBase;
+typedef void A3DSurfBase;
+typedef void A3DSurfBlend01;
+typedef void A3DSurfBlend02;
+typedef void A3DSurfBlend03;
+typedef void A3DSurfCone;
+typedef void A3DSurfCone;
+typedef void A3DSurfCylinder;
+typedef void A3DSurfCylinder;
+typedef void A3DSurfCylindrical;
+typedef void A3DSurfExtrusion;
+typedef void A3DSurfFromCurves;
+typedef void A3DSurfNurbs;
+typedef void A3DSurfOffset;
+typedef void A3DSurfPipe;
+typedef void A3DSurfPlane;
+typedef void A3DSurfRevolution;
+typedef void A3DSurfRuled;
+typedef void A3DSurfSphere;
+typedef void A3DSurfTorus;
+typedef void A3DSurfTransform;
+typedef void A3DTopoBrepData;
+typedef void A3DTopoCoEdge;
+typedef void A3DTopoConnex;
+typedef void A3DTopoFace;
+typedef void A3DTopoLoop;
+typedef void A3DTopoShell;
+
+namespace CADKernel
+{
+
+enum class EEntityType
+{
+	Instance,
+	Component,
+	Part,
+	Assembly,
+	UnloadedComponent,
+	Body,
+	BRepModel,
+	Undefined
+};
+
+struct FEntityData
+{
+	unsigned int Uuid = 0;
+	EEntityType EntityType;
+	FString Name;
+	TMap<FString, FString> MetaData;
+	TArray<A3DAsmProductOccurrence*> Childs;
+};
+
+struct FBodyGeometricData
+{
+	unsigned int Uuid = 0;
+	FString Name;
+	//int FaceCount = 0;
+	//int MeshTriangleCount = 0;
+	//double MeshArea = 0.;
+	//FBBox MeshBBox;
+	//FBBox CADBBox;
+	//int MaxFaceTriangleCount = 0;
+	//int MaxCtrlHullSize = 0;
+	//bool bTestBBox = true;
+	//bool bTestArea = true;
+	//bool bTestVertexCount = true;
+	//bool bMeshIsExtract = false;
+
+	A3DRiRepresentationItem* RepresentationItem = nullptr;
+};
+
+namespace TechSoftUtils
+{
+int32 TraverseSource(const A3DEntity* Entity, FEntityData& OutEntityData, bool bIsOccurrence = false, const ECADType FileType = ECADType::OTHER);
+int32 TraverseSpecificMetaData(A3DEModellerType ModellerType, const A3DAsmProductOccurrence* Occurrence, FEntityData& MetaData);
+
+int32 TraverseGraphics(const A3DGraphics* Graphics);
+FMatrixH TraverseCoordinateSystem(const A3DRiCoordinateSystem* CoordinateSystem, double UnitScale);
+FMatrixH TraverseTransformation(const A3DMiscCartesianTransformation* Transformation3d, double UnitScale);
+FMatrixH TraverseGeneralTransformation(const A3DMiscGeneralTransformation* GeneralTransformation, double UnitScale);
+FMatrixH TraverseTransformation3D(const A3DMiscCartesianTransformation* CartesianTransformation, double UnitScale);
+
+FString GetFileName(const FString& InFilePath);
+
+class FUVReparameterization
+{
+private:
+	double Scale[2] = { 1., 1. };
+	double Offset[2] = { 0., 0. };
+	bool bNeedApply = false;
+
+public:
+	static const FUVReparameterization Identity;
+
+	FUVReparameterization(const double InUScale = 1., const double InUOffset = 0., const double InVScale = 1., const double InVOffset = 0.)
+	{
+		SetCoef(InUScale, InUOffset, InVScale, InVOffset);
+		SetNeedApply();
+	}
+
+	void SetScale(const double InUScale, const double InVScale)
+	{
+		Scale[EIso::IsoU] = InUScale;
+		Scale[EIso::IsoV] = InVScale;
+		Offset[EIso::IsoU] = 0.;
+		Offset[EIso::IsoV] = 0.;
+
+		SetNeedApply();
+	}
+
+	void SetCoef(const double InUScale, const double InUOffset, const double InVScale, const double InVOffset)
+	{
+		Scale[EIso::IsoU] = InUScale;
+		Scale[EIso::IsoV] = InVScale;
+		Offset[EIso::IsoU] = InUOffset;
+		Offset[EIso::IsoV] = InVOffset;
+		SetNeedApply();
+	}
+
+	bool NeedApply() const
+	{
+		return bNeedApply;
+	}
+
+	void SetNeedApply()
+	{
+		if (!FMath::IsNearlyEqual(Scale[EIso::IsoU], 1.) || !FMath::IsNearlyEqual(Scale[EIso::IsoV], 1.) || !FMath::IsNearlyEqual(Offset[EIso::IsoU], 0.) || !FMath::IsNearlyEqual(Offset[EIso::IsoV], 0.))
+		{
+			bNeedApply = true;
+		}
+		else
+		{
+			bNeedApply = false;
+		}
+	}
+
+	void ScaleTransform(double InUScale, double InVScale)
+	{
+		Scale[EIso::IsoU] *= InUScale;
+		Scale[EIso::IsoV] *= InVScale;
+		Offset[EIso::IsoU] *= InUScale;
+		Offset[EIso::IsoV] *= InVScale;
+		SetNeedApply();
+	}
+
+	void Apply(FPoint& Point) const
+	{
+		Point.X = Scale[EIso::IsoU] * Point.X + Offset[EIso::IsoU];
+		Point.Y = Scale[EIso::IsoV] * Point.Y + Offset[EIso::IsoV];
+	}
+
+	void AddUVTransform(A3DUVParameterizationData& Transform, bool bSwapUV = false)
+	{
+		int32 UIndex = bSwapUV ? 1 : 0;
+		int32 VIndex = bSwapUV ? 0 : 1;
+
+		Scale[UIndex] = Scale[UIndex] * Transform.m_dUCoeffA;
+		Scale[VIndex] = Scale[VIndex] * Transform.m_dVCoeffA;
+		Offset[UIndex] = Offset[UIndex] * Transform.m_dUCoeffA + Transform.m_dUCoeffB;
+		Offset[VIndex] = Offset[VIndex] * Transform.m_dVCoeffA + Transform.m_dVCoeffB;
+		SetNeedApply();
+	}
+};
+}
+
+
+class FSurfacicBoundary;
+
+class FBody;
+class FCriterion;
+class FCurve;
+class FEntity;
+class FFaceMesh;
+class FMetadataDictionary;
+class FModel;
+class FPoint;
+class FRestrictionCurve;
+class FSession;
+class FShell;
+class FSurface;
+class FTopologicalEdge;
+class FTopologicalLoop;
+
+class FTechSoftBridge
+{
+private:
+
+	FSession& Session;
+	FCADFileReport& Report;
+
+	const double GeometricTolerance;
+	const double SquareGeometricTolerance;
+	const double SquareJoiningVertexTolerance;
+
+	TMap<const A3DEntity*, TSharedPtr<FEntity>> TechSoftToEntity;
+	TMap<const A3DTopoCoEdge*, TSharedPtr<FTopologicalEdge>> A3DEdgeToEdge;
+
+	double UnitScale = 1;
+
+public:
+	FTechSoftBridge(FSession& InSession, FCADFileReport& InReport);
+
+	TSharedRef<FBody> AddBody(A3DRiBrepModel* A3DBrepModel, FEntityData& EntityData, double FileUnit);
+
+private:
+
+	void TraverseBrepData(const A3DTopoBrepData* A3DBrepData, TSharedRef<FBody>& Body);
+
+	void TraverseConnex(const A3DTopoConnex* A3DConnex, TSharedRef<FBody>& Body);
+	void TraverseShell(const A3DTopoShell* A3DShell, TSharedRef<FBody>& Body);
+	void AddFace(const A3DTopoFace* A3DFace, TSharedRef<FShell>& Body);
+
+	TSharedPtr<FTopologicalLoop> AddLoop(const A3DTopoLoop* A3DLoop, const TSharedRef<FSurface>& Surface, const TechSoftUtils::FUVReparameterization& UVReparameterization);
+	TSharedPtr<FTopologicalEdge> AddEdge(const A3DTopoCoEdge* A3DCoedge, const TSharedRef<FSurface>& Surface, const TechSoftUtils::FUVReparameterization& UVReparameterization, EOrientation& OutOrientation);
+
+	TSharedPtr<FSurface> AddSurface(const A3DSurfBase* A3DSurface, TechSoftUtils::FUVReparameterization& OutUVReparameterization);
+	TSharedPtr<FSurface> AddConeSurface(const A3DSurfCone* A3DSurface, TechSoftUtils::FUVReparameterization& OutUVReparameterization);
+	TSharedPtr<FSurface> AddCylinderSurface(const A3DSurfCylinder* A3DSurface, TechSoftUtils::FUVReparameterization& OutUVReparameterization);
+	TSharedPtr<FSurface> AddLinearTransfoSurface(const A3DSurfBase* A3DSurface, TechSoftUtils::FUVReparameterization& OutUVReparameterization);
+	TSharedPtr<FSurface> AddNurbsSurface(const A3DSurfNurbs* A3DSurface, TechSoftUtils::FUVReparameterization& OutUVReparameterization);
+	TSharedPtr<FSurface> AddOffsetSurface(const A3DSurfOffset* A3DSurface, TechSoftUtils::FUVReparameterization& OutUVReparameterization);
+	TSharedPtr<FSurface> AddPlaneSurface(const A3DSurfPlane* A3DSurface, TechSoftUtils::FUVReparameterization& OutUVReparameterization);
+	TSharedPtr<FSurface> AddRevolutionSurface(const A3DSurfRevolution* A3DSurface, TechSoftUtils::FUVReparameterization& OutUVReparameterization);
+	TSharedPtr<FSurface> AddRuledSurface(const A3DSurfRuled* A3DSurface, TechSoftUtils::FUVReparameterization& OutUVReparameterization);
+	TSharedPtr<FSurface> AddSphereSurface(const A3DSurfSphere* A3DSurface, TechSoftUtils::FUVReparameterization& OutUVReparameterization);
+	TSharedPtr<FSurface> AddTorusSurface(const A3DSurfTorus* A3DSurface, TechSoftUtils::FUVReparameterization& OutUVReparameterization);
+
+	TSharedPtr<FSurface> AddBlend01Surface(const A3DSurfBlend01* A3DSurface, TechSoftUtils::FUVReparameterization& OutUVReparameterization);
+	TSharedPtr<FSurface> AddBlend02Surface(const A3DSurfBlend02* A3DSurface, TechSoftUtils::FUVReparameterization& OutUVReparameterization);
+	TSharedPtr<FSurface> AddBlend03Surface(const A3DSurfBlend03* A3DSurface, TechSoftUtils::FUVReparameterization& OutUVReparameterization);
+	TSharedPtr<FSurface> AddCylindricalSurface(const A3DSurfCylindrical* A3DSurface, TechSoftUtils::FUVReparameterization& OutUVReparameterization);
+	TSharedPtr<FSurface> AddPipeSurface(const A3DSurfPipe* A3DSurface, TechSoftUtils::FUVReparameterization& OutUVReparameterization);
+	TSharedPtr<FSurface> AddExtrusionSurface(const A3DSurfExtrusion* A3DSurface, TechSoftUtils::FUVReparameterization& OutUVReparameterization);
+	TSharedPtr<FSurface> AddSurfaceFromCurves(const A3DSurfFromCurves* A3DSurface, TechSoftUtils::FUVReparameterization& OutUVReparameterization);
+	TSharedPtr<FSurface> AddTransformSurface(const A3DSurfTransform* A3DSurface, TechSoftUtils::FUVReparameterization& OutUVReparameterization);
+
+	TSharedPtr<FSurface> AddSurfaceAsNurbs(const A3DSurfBase* A3DSurface, TechSoftUtils::FUVReparameterization& OutUVReparameterization);
+	TSharedPtr<FSurface> AddSurfaceNurbs(const A3DSurfNurbsData& A3DNurbsData, TechSoftUtils::FUVReparameterization& OutUVReparameterization);
+
+	TSharedPtr<FRestrictionCurve> AddRestrictionCurve(const A3DCrvBase* A3DCurve, const TSharedRef<FSurface>& Surface);
+
+	TSharedPtr<FCurve> AddCurve(const A3DCrvBase* A3DCurve, const TechSoftUtils::FUVReparameterization& UVReparameterization = TechSoftUtils::FUVReparameterization::Identity);
+
+	TSharedPtr<FCurve> AddCurveCircle(const A3DCrvCircle* A3DCurve, const TechSoftUtils::FUVReparameterization& UVReparameterization);
+	TSharedPtr<FCurve> AddCurveComposite(const A3DCrvComposite* A3DCurve, const TechSoftUtils::FUVReparameterization& UVReparameterization);
+	TSharedPtr<FCurve> AddCurveEllipse(const A3DCrvEllipse* A3DCurve, const TechSoftUtils::FUVReparameterization& UVReparameterization);
+	TSharedPtr<FCurve> AddCurveHelix(const A3DCrvHelix* A3DCurve, const TechSoftUtils::FUVReparameterization& UVReparameterization);
+	TSharedPtr<FCurve> AddCurveHyperbola(const A3DCrvHyperbola* A3DCurve, const TechSoftUtils::FUVReparameterization& UVReparameterization);
+	TSharedPtr<FCurve> AddCurveLine(const A3DCrvLine* A3DCurve, const TechSoftUtils::FUVReparameterization& UVReparameterization);
+	TSharedPtr<FCurve> AddCurveNurbs(const A3DCrvNurbs* A3DCurve, const TechSoftUtils::FUVReparameterization& UVReparameterization);
+	TSharedPtr<FCurve> AddCurveParabola(const A3DCrvParabola* A3DCurve, const TechSoftUtils::FUVReparameterization& UVReparameterization);
+	TSharedPtr<FCurve> AddCurvePolyLine(const A3DCrvPolyLine* A3DCurve, const TechSoftUtils::FUVReparameterization& UVReparameterization);
+
+	TSharedPtr<FCurve> AddCurveBlend02Boundary(const A3DCrvBlend02Boundary* A3DCurve, const TechSoftUtils::FUVReparameterization& UVReparameterization);
+	TSharedPtr<FCurve> AddCurveEquation(const A3DCrvEquation* A3DCurve, const TechSoftUtils::FUVReparameterization& UVReparameterization);
+	TSharedPtr<FCurve> AddCurveIntersection(const A3DCrvIntersection* A3DCurve, const TechSoftUtils::FUVReparameterization& UVReparameterization);
+	TSharedPtr<FCurve> AddCurveOffset(const A3DCrvOffset* A3DCurve, const TechSoftUtils::FUVReparameterization& UVReparameterization);
+	TSharedPtr<FCurve> AddCurveOnSurf(const A3DCrvOnSurf* A3DCurve, const TechSoftUtils::FUVReparameterization& UVReparameterization);
+	TSharedPtr<FCurve> AddCurveTransform(const A3DCrvTransform* A3DCurve, const TechSoftUtils::FUVReparameterization& UVReparameterization);
+
+	TSharedPtr<FCurve> AddCurveAsNurbs(const A3DCrvBase* A3DCurve, const TechSoftUtils::FUVReparameterization& UVReparameterization);
+
+	void AddMetadata(FEntityData& EntityData, TSharedRef<FMetadataDictionary> Entity);
+};
+}
+
+
+
+#endif // USE_KERNEL_IO_SDK
diff --git a/Engine/Plugins/Enterprise/DatasmithCADImporter/Source/CADInterfaces/Private/TechSoftFileParser.cpp b/Engine/Plugins/Enterprise/DatasmithCADImporter/Source/CADInterfaces/Private/TechSoftFileParser.cpp
--- a/Engine/Plugins/Enterprise/DatasmithCADImporter/Source/CADInterfaces/Private/TechSoftFileParser.cpp
+++ b/Engine/Plugins/Enterprise/DatasmithCADImporter/Source/CADInterfaces/Private/TechSoftFileParser.cpp
@@ -5,10 +5,12 @@
 #ifdef USE_TECHSOFT_SDK
 
 #include "TechSoftInterface.h"
-#include "TUniqueTechSoftObj.h"
+
+#include "Templates/UnrealTemplate.h"
 
 namespace CADLibrary
 {
+
 namespace TechSoftFileParserImpl
 {
 // Functions to clean metadata
@@ -42,7 +44,7 @@ FString CleanCatiaInstanceSdkName(const FString& Name)
 	int32 Index;
 	if (Name.FindChar(TEXT('('), Index))
 	{
-		FString NewName = Name.Right(Index + 1);
+		FString NewName = Name.RightChop(Index + 1);
 		if (Name.FindLastChar(TEXT(')'), Index))
 		{
 			NewName = NewName.Left(Index);
@@ -68,7 +70,7 @@ FString CleanSwInstanceSdkName(const FString& Name)
 	int32 Position;
 	if (Name.FindLastChar(TEXT('-'), Position))
 	{
-		FString NewName = Name.Left(Position) + TEXT("<") + Name.Right(Position + 1) + TEXT(">");
+		FString NewName = Name.Left(Position) + TEXT("<") + Name.RightChop(Position + 1) + TEXT(">");
 		return NewName;
 	}
 	return Name;
@@ -90,7 +92,7 @@ FString CleanCatiaReferenceName(const FString& Name)
 	int32 Position;
 	if (Name.FindLastChar(TEXT('.'), Position))
 	{
-		FString Indice = Name.Right(Position + 1);
+		FString Indice = Name.RightChop(Position + 1);
 		if (Indice.IsNumeric())
 		{
 			FString NewName = Name.Left(Position);
@@ -105,7 +107,7 @@ FString CleanNameByRemoving_prt(const FString& Name)
 	int32 Position;
 	if (Name.FindLastChar(TEXT('.'), Position))
 	{
-		FString Extension = Name.Right(Position);
+		FString Extension = Name.RightChop(Position);
 		if (Extension.Equals(TEXT("prt"), ESearchCase::IgnoreCase))
 		{
 			FString NewName = Name.Left(Position);
@@ -242,14 +244,14 @@ void SetIOOption(A3DImport& Importer)
 	Importer.m_sLoadData.m_sGeneral.m_bReadActiveFilter = A3D_FALSE;
 	Importer.m_sLoadData.m_sGeneral.m_eReadingMode2D3D = kA3DRead_3D;
 
-	Importer.m_sLoadData.m_sGeneral.m_eReadGeomTessMode = kA3DReadGeomOnly;
+	Importer.m_sLoadData.m_sGeneral.m_eReadGeomTessMode = kA3DReadGeomAndTess/*kA3DReadGeomOnly*/;
 	Importer.m_sLoadData.m_sGeneral.m_eDefaultUnit;
 	Importer.m_sLoadData.m_sGeneral.m_bReadFeature = A3D_FALSE;
 
 	Importer.m_sLoadData.m_sGeneral.m_bReadConstraints = A3D_FALSE;
 	Importer.m_sLoadData.m_sGeneral.m_iNbMultiProcess = 1;
 
-	Importer.m_sLoadData.m_sIncremental.m_bLoadNoDependencies = true;
+	Importer.m_sLoadData.m_sIncremental.m_bLoadNoDependencies = CADLibrary::FImportParameters::bGEnableCADCache;
 	Importer.m_sLoadData.m_sIncremental.m_bLoadStructureOnly = false;
 }
 
@@ -278,19 +280,17 @@ void UpdateIOOptionAccordingToFormat(const CADLibrary::ECADFormat Format, A3DImp
 		break;
 	};
 }
-}
+
+} // ns TechSoftFileParserImpl
 
 FTechSoftFileParser::FTechSoftFileParser(FCADFileData& InCADData, const FString& EnginePluginsPath)
 	: CADFileData(InCADData)
-	, TechSoftInterface(GetTechSoftInterface())
+	, TechSoftInterface(TechSoftUtils::GetTechSoftInterface())
 {
 }
 
 ECADParsingResult FTechSoftFileParser::Process()
 {
-	// Process the file
-	ECADParsingResult Result = ECADParsingResult::ProcessOk;
-
 	const FFileDescriptor& File = CADFileData.GetCADFileDescription();
 
 	if (File.GetPathOfFileToLoad().IsEmpty())
@@ -332,31 +332,50 @@ ECADParsingResult FTechSoftFileParser::Process()
 
 	ReadMaterialsAndColors();
 
-	return TraverseModel(TechSoftInterface.GetModelFile());
+	ECADParsingResult Result = TraverseModel(TechSoftInterface.GetModelFile());
+
+	TechSoftInterface.UnloadModel();
+
+	return Result;
 }
 
 void FTechSoftFileParser::ReserveCADFileData()
 {
+	// TODO: Could be more accurate
 	CountUnderModel(TechSoftInterface.GetModelFile());
 
 	CADFileData.ReserveBodyMeshes(ComponentCount[EComponentType::Body]);
 
 	FArchiveSceneGraph& SceneGraphArchive = CADFileData.GetSceneGraphArchive();
-	SceneGraphArchive.Bodies.Reserve(ComponentCount[EComponentType::Body]);
-	SceneGraphArchive.Components.Reserve(ComponentCount[EComponentType::Reference]);
-	SceneGraphArchive.UnloadedComponents.Reserve(ComponentCount[EComponentType::Reference]);
-	SceneGraphArchive.ExternalReferences.Reserve(ComponentCount[EComponentType::Reference]);
-	SceneGraphArchive.Instances.Reserve(ComponentCount[EComponentType::Occurrence]);
-
-	SceneGraphArchive.CADIdToBodyIndex.Reserve(ComponentCount[EComponentType::Body]);
-	SceneGraphArchive.CADIdToComponentIndex.Reserve(ComponentCount[EComponentType::Reference]);
-	SceneGraphArchive.CADIdToUnloadedComponentIndex.Reserve(ComponentCount[EComponentType::Reference]);
-	SceneGraphArchive.CADIdToInstanceIndex.Reserve(ComponentCount[EComponentType::Occurrence]);
+	SceneGraphArchive.Reserve(ComponentCount[EComponentType::Occurrence], ComponentCount[EComponentType::Reference], ComponentCount[EComponentType::Body]);
 
 	uint32 MaterialNum = CountMaterial() + CountColor();
 	SceneGraphArchive.MaterialHIdToMaterial.Reserve(MaterialNum);
 }
 
+void FTechSoftFileParser::CountUnderModel(const A3DAsmModelFile* AsmModel)
+{
+	TUniqueTSObj<A3DAsmModelFileData> ModelFileData(AsmModel);
+	if (!ModelFileData.IsValid())
+	{
+		return;
+	}
+
+	ComponentCount[EComponentType::Occurrence] ++;
+
+	for (uint32 Index = 0; Index < ModelFileData->m_uiPOccurrencesSize; ++Index)
+	{
+		if (IsConfigurationSet(ModelFileData->m_ppPOccurrences[Index]))
+		{
+			CountUnderConfigurationSet(ModelFileData->m_ppPOccurrences[Index]);
+		}
+		else
+		{
+			CountUnderOccurrence(ModelFileData->m_ppPOccurrences[Index]);
+		}
+	}
+}
+
 ECADParsingResult FTechSoftFileParser::TraverseModel(const A3DAsmModelFile* ModelFile)
 {
 	TUniqueTSObj<A3DAsmModelFileData> ModelFileData(ModelFile);
@@ -369,14 +388,14 @@ ECADParsingResult FTechSoftFileParser::TraverseModel(const A3DAsmModelFile* Mode
 	FileUnit = ModelFileData->m_dUnit;
 
 	FEntityMetaData MetaData;
-	TraverseMetaData(ModelFile, MetaData);
-	TraverseSpecificMetaData(ModelFile, MetaData);
+	ExtractMetaData(ModelFile, MetaData);
+	ExtractSpecificMetaData(ModelFile, MetaData);
 
 	for (uint32 Index = 0; Index < ModelFileData->m_uiPOccurrencesSize; ++Index)
 	{
 		if (IsConfigurationSet(ModelFileData->m_ppPOccurrences[Index]))
 		{
-			//TraverseConfigurationSet(ModelFileData->m_ppPOccurrences[Index], MetaData);
+			TraverseConfigurationSet(ModelFileData->m_ppPOccurrences[Index]);
 		}
 		else
 		{
@@ -387,17 +406,85 @@ ECADParsingResult FTechSoftFileParser::TraverseModel(const A3DAsmModelFile* Mode
 	return ECADParsingResult::ProcessOk;
 }
 
+void FTechSoftFileParser::TraverseConfigurationSet(const A3DAsmProductOccurrence* ConfigurationSetPtr)
+{
+	TUniqueTSObj<A3DAsmProductOccurrenceData> ConfigurationSetData(ConfigurationSetPtr);
+	if (!ConfigurationSetData.IsValid())
+	{
+		return;
+	}
+
+	FEntityMetaData MetaData;
+	ExtractMetaData(ConfigurationSetPtr, MetaData);
+	ExtractSpecificMetaData(ConfigurationSetPtr, MetaData);
+
+	TUniqueTSObj<A3DAsmProductOccurrenceData> ConfigurationData;
+	for (unsigned int Index = 0; Index < ConfigurationSetData->m_uiPOccurrencesSize; ++Index)
+	{
+		ConfigurationData.FillFrom(ConfigurationSetData->m_ppPOccurrences[Index]);
+		if (!ConfigurationData.IsValid())
+		{
+			return;
+		}
+
+		// Traverse default configuration
+		if (ConfigurationData->m_uiProductFlags & (A3D_PRODUCT_FLAG_DEFAULT| A3D_PRODUCT_FLAG_CONFIG))
+		{
+			TraverseReference(ConfigurationSetData->m_ppPOccurrences[Index]);
+			return;
+		}
+	}
+
+	// no default configuration, traverse the first
+	if (ConfigurationSetData->m_uiPOccurrencesSize)
+	{
+// #ueent_techsoft: Should we take the first occurence or the first configuration???
+		TraverseReference(ConfigurationSetData->m_ppPOccurrences[0]);
+	}
+}
+
+void FTechSoftFileParser::CountUnderConfigurationSet(const A3DAsmProductOccurrence* ConfigurationSetPtr)
+{
+	TUniqueTSObj<A3DAsmProductOccurrenceData> ConfigurationSetData(ConfigurationSetPtr);
+	if (!ConfigurationSetData.IsValid())
+	{
+		return;
+	}
+
+	TUniqueTSObj<A3DAsmProductOccurrenceData> ConfigurationData;
+	for (unsigned int Index = 0; Index < ConfigurationSetData->m_uiPOccurrencesSize; ++Index)
+	{
+		ConfigurationData.FillFrom(ConfigurationSetData->m_ppPOccurrences[Index]);
+		if (!ConfigurationData.IsValid())
+		{
+			return;
+		}
+
+		if (ConfigurationData->m_uiProductFlags & (A3D_PRODUCT_FLAG_DEFAULT | A3D_PRODUCT_FLAG_CONFIG))
+		{
+			CountUnderOccurrence(ConfigurationSetData->m_ppPOccurrences[Index]);
+			return;
+		}
+	}
+
+	// no default configuration, traverse the first
+	if (ConfigurationSetData->m_uiPOccurrencesSize)
+	{
+		CountUnderOccurrence(ConfigurationSetData->m_ppPOccurrences[0]);
+	}
+}
+
 void FTechSoftFileParser::TraverseReference(const A3DAsmProductOccurrence* ReferencePtr)
 {
 	FEntityMetaData MetaData;
-	TraverseMetaData(ReferencePtr, MetaData);
+	ExtractMetaData(ReferencePtr, MetaData);
 
 	if (MetaData.bRemoved || !MetaData.bShow)
 	{
 		return;
 	}
 
-	TraverseSpecificMetaData(ReferencePtr, MetaData);
+	ExtractSpecificMetaData(ReferencePtr, MetaData);
 	BuildReferenceName(MetaData.MetaData);
 
 	TraverseMaterialProperties(ReferencePtr);
@@ -423,8 +510,15 @@ void FTechSoftFileParser::TraverseReference(const A3DAsmProductOccurrence* Refer
 	{
 		TraversePartDefinition(ReferenceData->m_pPart, Component);
 	}
-}
 
+#ifndef NEW_CODE
+	// Is this really necessary ????
+	if (ReferenceData->m_pPrototype && !ReferenceData->m_uiPOccurrencesSize && !ReferenceData->m_pPart)
+	{
+		TraversePrototype(ReferenceData->m_pPrototype, Component);
+	}
+#endif
+}
 
 FArchiveInstance& FTechSoftFileParser::AddInstance(FEntityMetaData& InstanceMetaData)
 {
@@ -499,17 +593,19 @@ FArchiveBody& FTechSoftFileParser::AddBody(FEntityMetaData& BodyMetaData)
 	return Body;
 }
 
+#pragma optimize ("", off)
+#ifndef NEW_CODE
 FCadId FTechSoftFileParser::TraverseOccurrence(const A3DAsmProductOccurrence* OccurrencePtr)
 {
 	FEntityMetaData InstanceMetaData;
-	TraverseMetaData(OccurrencePtr, InstanceMetaData);
+	ExtractMetaData(OccurrencePtr, InstanceMetaData);
 
 	if (InstanceMetaData.bRemoved || !InstanceMetaData.bShow)
 	{
 		return 0;
 	}
 
-	TraverseSpecificMetaData(OccurrencePtr, InstanceMetaData);
+	ExtractSpecificMetaData(OccurrencePtr, InstanceMetaData);
 	BuildInstanceName(InstanceMetaData.MetaData);
 
 	TraverseMaterialProperties(OccurrencePtr);
@@ -530,8 +626,11 @@ FCadId FTechSoftFileParser::TraverseOccurrence(const A3DAsmProductOccurrence* Oc
 	if (OccurrenceData->m_pPrototype)
 	{
 		FMatrix PrototypeMatrix;
-		TraversePrototype(OccurrenceData->m_pPrototype, PrototypeMetaData, PrototypeMatrix);
-		Instance.TransformMatrix *= PrototypeMatrix;
+		ProcessPrototype(OccurrenceData->m_pPrototype, PrototypeMetaData, PrototypeMatrix);
+		if (!OccurrenceData->m_pLocation)
+		{
+			Instance.TransformMatrix = PrototypeMatrix;
+		}
 	}
 
 	if (PrototypeMetaData.bUnloaded)
@@ -552,10 +651,167 @@ FCadId FTechSoftFileParser::TraverseOccurrence(const A3DAsmProductOccurrence* Oc
 		{
 			TraversePartDefinition(OccurrenceData->m_pPart, Component);
 		}
+
+		if (OccurrenceData->m_pPrototype && !OccurrenceData->m_uiPOccurrencesSize && !OccurrenceData->m_pPart)
+		{
+			TraversePrototype(OccurrenceData->m_pPrototype, Component);
+		}
+	}
+
+	return Instance.ObjectId;
+}
+#else
+FCadId FTechSoftFileParser::TraverseOccurrence(const A3DAsmProductOccurrence* OccurrencePtr)
+{
+	TUniqueTSObj<A3DAsmProductOccurrenceData> OccurrenceData(OccurrencePtr);
+	if (!OccurrenceData.IsValid())
+	{
+		return 0;
+	}
+
+	bool bContinueTraverse = OccurrenceData->m_pPrototype
+		|| OccurrenceData->m_pExternalData
+		|| OccurrenceData->m_pPart
+		|| OccurrenceData->m_uiPOccurrencesSize > 0;
+	if (!bContinueTraverse)
+	{
+		return 0;
+	}
+
+	FEntityMetaData InstanceMetaData;
+	ExtractMetaData(OccurrencePtr, InstanceMetaData);
+
+	if (InstanceMetaData.bRemoved || !InstanceMetaData.bShow)
+	{
+		return 0;
+	}
+
+	ExtractSpecificMetaData(OccurrencePtr, InstanceMetaData);
+	BuildInstanceName(InstanceMetaData.MetaData);
+
+	TraverseMaterialProperties(OccurrencePtr);
+	TraverseLayer(OccurrencePtr);
+
+	FArchiveInstance& Instance = AddInstance(InstanceMetaData);
+
+	if (OccurrenceData->m_eProductLoadStatus == kA3DProductLoadStatusNotLoaded)
+	{
+		FArchiveUnloadedComponent& UnloadedComponent = AddUnloadedComponent(InstanceMetaData, Instance);
+		return Instance.ObjectId;
+	}
+	
+	A3DAsmProductOccurrence* PrototypePtr = OccurrenceData->m_pPrototype;
+	A3DMiscTransformation* LocationPtr = OccurrenceData->m_pLocation;
+
+	while (!LocationPtr && PrototypePtr)
+	{
+		TUniqueTSObj<A3DAsmProductOccurrenceData> PrototypeOccurrenceData(PrototypePtr);
+		LocationPtr = PrototypeOccurrenceData->m_pLocation;
+		PrototypePtr = PrototypeOccurrenceData->m_pPrototype;
+	}
+
+	Instance.TransformMatrix = TraverseTransformation(LocationPtr);
+
+	A3DAsmPartDefinition* PartDefinition = OccurrenceData->m_pPart;
+	PrototypePtr = OccurrenceData->m_pPrototype;
+
+	while (!PartDefinition && PrototypePtr)
+	{
+		TUniqueTSObj<A3DAsmProductOccurrenceData> PrototypeOccurrenceData(PrototypePtr);
+		PartDefinition = PrototypeOccurrenceData->m_pPart;
+		PrototypePtr = PrototypeOccurrenceData->m_pPrototype;
+	}
+
+	FArchiveComponent& Component = AddComponent(InstanceMetaData, Instance);
+
+	TraversePartDefinition(PartDefinition, Component);
+
+	uint32 ChildrenCount = OccurrenceData->m_uiPOccurrencesSize;
+	A3DAsmProductOccurrence** Children = OccurrenceData->m_ppPOccurrences;
+	PrototypePtr = OccurrenceData->m_pPrototype;
+
+	while (ChildrenCount == 0 && PrototypePtr)
+	{
+		TUniqueTSObj<A3DAsmProductOccurrenceData> PrototypeOccurrenceData(PrototypePtr);
+		ChildrenCount = PrototypeOccurrenceData->m_uiPOccurrencesSize;
+		Children = PrototypeOccurrenceData->m_ppPOccurrences;
+		PrototypePtr = PrototypeOccurrenceData->m_pPrototype;
+	}
+
+	for (uint32 Index = 0; Index < ChildrenCount; ++Index)
+	{
+		int32 ChildrenId = TraverseOccurrence(Children[Index]);
+		Component.Children.Add(ChildrenId);
 	}
 
 	return Instance.ObjectId;
 }
+#endif
+#pragma optimize ("", on)
+
+void FTechSoftFileParser::ProcessOccurrence(TUniqueTSObj<A3DAsmProductOccurrenceData>& OccurrenceData, FArchiveComponent& Component)
+{
+	if (!OccurrenceData.IsValid())
+	{
+		return;
+	}
+
+	if (OccurrenceData->m_pPart)
+	{
+		TraversePartDefinition(OccurrenceData->m_pPart, Component);
+	}
+
+	if (OccurrenceData->m_pExternalData)
+	{
+		int32 ChildrenId = TraverseOccurrence(OccurrenceData->m_pExternalData);
+		Component.Children.Add(ChildrenId);
+	}
+
+	for (uint32 Index = 0; Index < OccurrenceData->m_uiPOccurrencesSize; ++Index)
+	{
+		int32 ChildrenId = TraverseOccurrence(OccurrenceData->m_ppPOccurrences[Index]);
+		Component.Children.Add(ChildrenId);
+	}
+};
+
+void FTechSoftFileParser::CountUnderOccurrence(const A3DAsmProductOccurrence* Occurrence)
+{
+	TUniqueTSObj<A3DAsmProductOccurrenceData> OccurrenceData(Occurrence);
+	if (Occurrence && OccurrenceData.IsValid())
+	{
+		ComponentCount[EComponentType::Occurrence]++;
+		ComponentCount[EComponentType::Reference]++;
+
+		A3DAsmProductOccurrence* PrototypePtr = OccurrenceData->m_pPrototype;
+		A3DAsmPartDefinition* PartDefinition = OccurrenceData->m_pPart;
+
+		while (!PartDefinition && PrototypePtr)
+		{
+			TUniqueTSObj<A3DAsmProductOccurrenceData> PrototypeOccurrenceData(PrototypePtr);
+			PartDefinition = PrototypeOccurrenceData->m_pPart;
+			PrototypePtr = PrototypeOccurrenceData->m_pPrototype;
+		}
+
+		CountUnderPartDefinition(PartDefinition);
+
+		uint32 ChildrenCount = OccurrenceData->m_uiPOccurrencesSize;
+		A3DAsmProductOccurrence** Children = OccurrenceData->m_ppPOccurrences;
+		PrototypePtr = OccurrenceData->m_pPrototype;
+
+		while (ChildrenCount == 0 && PrototypePtr)
+		{
+			TUniqueTSObj<A3DAsmProductOccurrenceData> PrototypeOccurrenceData(PrototypePtr);
+			ChildrenCount = PrototypeOccurrenceData->m_uiPOccurrencesSize;
+			Children = PrototypeOccurrenceData->m_ppPOccurrences;
+			PrototypePtr = PrototypeOccurrenceData->m_pPrototype;
+		}
+
+		for (uint32 Index = 0; Index < ChildrenCount; ++Index)
+		{
+			CountUnderOccurrence(Children[Index]);
+		}
+	}
+}
 
 FFileDescriptor FTechSoftFileParser::GetOccurrenceFileName(const A3DAsmProductOccurrence* OccurrencePtr)
 {
@@ -599,7 +855,56 @@ FFileDescriptor FTechSoftFileParser::GetOccurrenceFileName(const A3DAsmProductOc
 	return FFileDescriptor();
 }
 
-void FTechSoftFileParser::TraversePrototype(const A3DAsmProductOccurrence* InPrototypePtr, FEntityMetaData& OutPrototypeMetaData, FMatrix& OutPrototypeMatrix)
+void FTechSoftFileParser::TraversePrototype(const A3DAsmProductOccurrence* InPrototypePtr, FArchiveComponent& Component)
+{
+	TUniqueTSObj<A3DAsmProductOccurrenceData> SubPrototypeData(InPrototypePtr);
+	if (!SubPrototypeData.IsValid())
+	{
+		return;
+	}
+
+	for (uint32 Index = 0; Index < SubPrototypeData->m_uiPOccurrencesSize; ++Index)
+	{
+		int32 ChildrenId = TraverseOccurrence(SubPrototypeData->m_ppPOccurrences[Index]);
+		Component.Children.Add(ChildrenId);
+	}
+
+	if (SubPrototypeData->m_pPart)
+	{
+		TraversePartDefinition(SubPrototypeData->m_pPart, Component);
+	}
+
+	if (SubPrototypeData->m_pPrototype && !SubPrototypeData->m_uiPOccurrencesSize && !SubPrototypeData->m_pPart)
+	{
+		TraversePrototype(SubPrototypeData->m_pPrototype, Component);
+	}
+}
+
+void FTechSoftFileParser::CountUnderSubPrototype(const A3DAsmProductOccurrence* InPrototypePtr)
+{
+	TUniqueTSObj<A3DAsmProductOccurrenceData> SubPrototypeData(InPrototypePtr);
+	if (!SubPrototypeData.IsValid())
+	{
+		return;
+	}
+
+	for (uint32 Index = 0; Index < SubPrototypeData->m_uiPOccurrencesSize; ++Index)
+	{
+		CountUnderOccurrence(SubPrototypeData->m_ppPOccurrences[Index]);
+	}
+
+	if (SubPrototypeData->m_pPart)
+	{
+		CountUnderPartDefinition(SubPrototypeData->m_pPart);
+	}
+
+	if (SubPrototypeData->m_pPrototype && !SubPrototypeData->m_uiPOccurrencesSize && !SubPrototypeData->m_pPart)
+	{
+		CountUnderSubPrototype(SubPrototypeData->m_pPrototype);
+	}
+}
+
+void FTechSoftFileParser::ProcessPrototype(const A3DAsmProductOccurrence* InPrototypePtr, FEntityMetaData& OutPrototypeMetaData, FMatrix& OutPrototypeMatrix)
 {
 	TUniqueTSObj<A3DAsmProductOccurrenceData> PrototypeData(InPrototypePtr);
 	if (!PrototypeData.IsValid())
@@ -607,8 +912,8 @@ void FTechSoftFileParser::TraversePrototype(const A3DAsmProductOccurrence* InPro
 		return;
 	}
 
-	TraverseMetaData(InPrototypePtr, OutPrototypeMetaData);
-	TraverseSpecificMetaData(InPrototypePtr, OutPrototypeMetaData);
+	ExtractMetaData(InPrototypePtr, OutPrototypeMetaData);
+	ExtractSpecificMetaData(InPrototypePtr, OutPrototypeMetaData);
 
 	TraverseMaterialProperties(InPrototypePtr);
 
@@ -624,17 +929,28 @@ void FTechSoftFileParser::TraversePrototype(const A3DAsmProductOccurrence* InPro
 	BuildReferenceName(OutPrototypeMetaData.MetaData);
 }
 
+void FTechSoftFileParser::CountUnderPrototype(const A3DAsmProductOccurrence* Prototype)
+{
+	TUniqueTSObj<A3DAsmProductOccurrenceData> PrototypeData(Prototype);
+	if (!PrototypeData.IsValid())
+	{
+		return;
+	}
+
+	ComponentCount[EComponentType::Reference] ++;
+}
+
 void FTechSoftFileParser::TraversePartDefinition(const A3DAsmPartDefinition* PartDefinitionPtr, FArchiveComponent& Part)
 {
 	FEntityMetaData PartMetaData;
-	TraverseMetaData(PartDefinitionPtr, PartMetaData);
+	ExtractMetaData(PartDefinitionPtr, PartMetaData);
 
 	if (PartMetaData.bRemoved || !PartMetaData.bShow)
 	{
 		return;
 	}
 
-	TraverseSpecificMetaData(PartDefinitionPtr, PartMetaData);
+	ExtractSpecificMetaData(PartDefinitionPtr, PartMetaData);
 	BuildPartName(PartMetaData.MetaData);
 
 	TraverseMaterialProperties(PartDefinitionPtr);
@@ -651,6 +967,21 @@ void FTechSoftFileParser::TraversePartDefinition(const A3DAsmPartDefinition* Par
 	}
 }
 
+void FTechSoftFileParser::CountUnderPartDefinition(const A3DAsmPartDefinition* PartDefinition)
+{
+	TUniqueTSObj<A3DAsmPartDefinitionData> PartData(PartDefinition);
+	if (PartDefinition && PartData.IsValid())
+	{
+		ComponentCount[EComponentType::Reference] ++;
+		ComponentCount[EComponentType::Occurrence] ++;
+
+		for (unsigned int Index = 0; Index < PartData->m_uiRepItemsSize; ++Index)
+		{
+			CountUnderRepresentationItem(PartData->m_ppRepItems[Index]);
+		}
+	}
+}
+
 FCadId FTechSoftFileParser::TraverseRepresentationItem(A3DRiRepresentationItem* RepresentationItem, FEntityMetaData& PartMetaData)
 {
 	A3DEEntityType Type;
@@ -670,6 +1001,26 @@ FCadId FTechSoftFileParser::TraverseRepresentationItem(A3DRiRepresentationItem*
 	return 0;
 }
 
+void FTechSoftFileParser::CountUnderRepresentationItem(const A3DRiRepresentationItem* RepresentationItem)
+{
+	A3DEEntityType Type;
+	A3DEntityGetType(RepresentationItem, &Type);
+
+	switch (Type)
+	{
+	case kA3DTypeRiSet:
+		CountUnderRepresentationSet(RepresentationItem);
+		break;
+	case kA3DTypeRiBrepModel:
+	case kA3DTypeRiPolyBrepModel:
+		ComponentCount[EComponentType::Body] ++;
+		break;
+
+	default:
+		break;
+	}
+}
+
 FCadId FTechSoftFileParser::TraverseRepresentationSet(const A3DRiSet* RepresentationSetPtr, FEntityMetaData& PartMetaData)
 {
 	TUniqueTSObj<A3DRiSetData> RepresentationSetData(RepresentationSetPtr);
@@ -679,7 +1030,7 @@ FCadId FTechSoftFileParser::TraverseRepresentationSet(const A3DRiSet* Representa
 	}
 
 	FEntityMetaData RepresentationSetMetaData;
-	TraverseMetaData(RepresentationSetPtr, RepresentationSetMetaData);
+	ExtractMetaData(RepresentationSetPtr, RepresentationSetMetaData);
 
 	if (RepresentationSetMetaData.bRemoved || !RepresentationSetMetaData.bShow)
 	{
@@ -691,31 +1042,56 @@ FCadId FTechSoftFileParser::TraverseRepresentationSet(const A3DRiSet* Representa
 	FCadId RepresentationSetId = 0;
 	FArchiveComponent& RepresentationSet = AddOccurence(RepresentationSetMetaData, RepresentationSetId);
 
-	for (A3DUns32 ui = 0; ui < RepresentationSetData->m_uiRepItemsSize; ++ui)
+	for (A3DUns32 Index = 0; Index < RepresentationSetData->m_uiRepItemsSize; ++Index)
 	{
-		int32 ChildId = TraverseRepresentationItem(RepresentationSetData->m_ppRepItems[ui], RepresentationSetMetaData);
+		int32 ChildId = TraverseRepresentationItem(RepresentationSetData->m_ppRepItems[Index], RepresentationSetMetaData);
 		RepresentationSet.Children.Add(ChildId);
 	}
 	return RepresentationSetId;
 }
 
+void FTechSoftFileParser::CountUnderRepresentationSet(const A3DRiSet* RepresentationSet)
+{
+	TUniqueTSObj<A3DRiSetData> RepresentationSetData(RepresentationSet);
+	if (RepresentationSet && RepresentationSetData.IsValid())
+	{
+		ComponentCount[EComponentType::Occurrence] ++;
+		ComponentCount[EComponentType::Reference] ++;
+
+		for (A3DUns32 Index = 0; Index < RepresentationSetData->m_uiRepItemsSize; ++Index)
+		{
+			CountUnderRepresentationItem(RepresentationSetData->m_ppRepItems[Index]);
+		}
+	}
+}
+
 FCadId FTechSoftFileParser::TraverseBRepModel(A3DRiBrepModel* BRepModelPtr, FEntityMetaData& PartMetaData)
 {
-	TUniqueTSObj<A3DRiBrepModelData> BodyData(BRepModelPtr);
-	if (!BodyData.IsValid())
+	if (!BRepModelPtr)
 	{
 		return 0;
 	}
 
-	FEntityMetaData BRepMetaData;
-	TraverseMetaData(BRepModelPtr, BRepMetaData);
-	if (!BRepMetaData.bShow || BRepMetaData.bRemoved)
+	if (FCadId* CadIdPtr = RepresentationItemsCache.Find(BRepModelPtr))
 	{
-		return 0;
+		return *CadIdPtr;
 	}
 
-	TraverseSpecificMetaData(BRepModelPtr, BRepMetaData);
-	TraverseMaterialProperties(BRepModelPtr);
+	TUniqueTSObj<A3DRiBrepModelData> BodyData(BRepModelPtr);
+	if (!BodyData.IsValid())
+	{
+		return 0;
+	}
+
+	FEntityMetaData BRepMetaData;
+	ExtractMetaData(BRepModelPtr, BRepMetaData);
+	if (!BRepMetaData.bShow || BRepMetaData.bRemoved)
+	{
+		return 0;
+	}
+
+	ExtractSpecificMetaData(BRepModelPtr, BRepMetaData);
+	TraverseMaterialProperties(BRepModelPtr);
 
 	FArchiveBody& Body = AddBody(BRepMetaData);
 
@@ -728,6 +1104,9 @@ FCadId FTechSoftFileParser::TraverseBRepModel(A3DRiBrepModel* BRepModelPtr, FEnt
 	{
 		// Mesh with CADKernel
 	}
+
+	RepresentationItemsCache.Add(BRepModelPtr, Body.ObjectId);
+
 	return Body.ObjectId;
 }
 
@@ -752,6 +1131,16 @@ void FTechSoftFileParser::TraverseRepresentationContent(const A3DRiRepresentatio
 
 FCadId FTechSoftFileParser::TraversePolyBRepModel(const A3DRiPolyBrepModel* PolygonalPtr, FEntityMetaData& PartMetaData)
 {
+	if (!PolygonalPtr)
+	{
+		return 0;
+	}
+
+	if (FCadId* CadIdPtr = RepresentationItemsCache.Find(PolygonalPtr))
+	{
+		return *CadIdPtr;
+	}
+
 	TUniqueTSObj<A3DRiPolyBrepModelData> BodyData(PolygonalPtr);
 	if (!BodyData.IsValid())
 	{
@@ -759,22 +1148,24 @@ FCadId FTechSoftFileParser::TraversePolyBRepModel(const A3DRiPolyBrepModel* Poly
 	}
 
 	FEntityMetaData BRepMetaData;
-	TraverseMetaData(PolygonalPtr, BRepMetaData);
+	ExtractMetaData(PolygonalPtr, BRepMetaData);
 	if (!BRepMetaData.bShow || BRepMetaData.bRemoved)
 	{
 		return 0;
 	}
 
-	TraverseSpecificMetaData(PolygonalPtr, BRepMetaData);
+	ExtractSpecificMetaData(PolygonalPtr, BRepMetaData);
 	TraverseMaterialProperties(PolygonalPtr);
 
 	FArchiveBody& Body = AddBody(BRepMetaData);
 	TraverseRepresentationContent(PolygonalPtr, Body);
 
+	RepresentationItemsCache.Add(PolygonalPtr, Body.ObjectId);
+
 	return Body.ObjectId;
 }
 
-void FTechSoftFileParser::TraverseMetaData(const A3DEntity* Entity, FEntityMetaData& OutMetaData)
+void FTechSoftFileParser::ExtractMetaData(const A3DEntity* Entity, FEntityMetaData& OutMetaData)
 {
 	TUniqueTSObj<A3DRootBaseData> MetaData(Entity);
 	if (MetaData.IsValid())
@@ -979,7 +1370,7 @@ void FTechSoftFileParser::BuildBodyName(TMap<FString, FString>& MetaData)
 	MetaData.Add(TEXT("Name"), TEXT("NoName"));
 }
 
-void FTechSoftFileParser::TraverseSpecificMetaData(const A3DAsmProductOccurrence* Occurrence, FEntityMetaData& OutMetaData)
+void FTechSoftFileParser::ExtractSpecificMetaData(const A3DAsmProductOccurrence* Occurrence, FEntityMetaData& OutMetaData)
 {
 	//----------- Export Specific information per CAD format -----------
 	switch (ModellerType)
@@ -1228,36 +1619,38 @@ void FTechSoftFileParser::TraverseMaterialProperties(const A3DEntity* Entity)
 	}
 }
 
-FMatrix FTechSoftFileParser::TraverseTransformation3D(const A3DMiscCartesianTransformation* CartesianTransformation)
+FMatrix FTechSoftFileParser::TraverseTransformation3D(const A3DMiscTransformation* CartesianTransformation)
 {
-	if (CartesianTransformation)
-	{
-		TUniqueTSObj<A3DMiscCartesianTransformationData> CartesianTransformationData(CartesianTransformation);
+	TUniqueTSObj<A3DMiscCartesianTransformationData> CartesianTransformationData(CartesianTransformation);
 
-		if (CartesianTransformationData.IsValid())
-		{
-			FVector Origin(CartesianTransformationData->m_sOrigin.m_dX, CartesianTransformationData->m_sOrigin.m_dY, CartesianTransformationData->m_sOrigin.m_dZ);
-			FVector XVector(CartesianTransformationData->m_sXVector.m_dX, CartesianTransformationData->m_sXVector.m_dY, CartesianTransformationData->m_sXVector.m_dZ);;
-			FVector YVector(CartesianTransformationData->m_sYVector.m_dX, CartesianTransformationData->m_sYVector.m_dY, CartesianTransformationData->m_sYVector.m_dZ);;
+	if (CartesianTransformationData.IsValid())
+	{
+		FVector Origin(CartesianTransformationData->m_sOrigin.m_dX, CartesianTransformationData->m_sOrigin.m_dY, CartesianTransformationData->m_sOrigin.m_dZ);
+		FVector XVector(CartesianTransformationData->m_sXVector.m_dX, CartesianTransformationData->m_sXVector.m_dY, CartesianTransformationData->m_sXVector.m_dZ);;
+		FVector YVector(CartesianTransformationData->m_sYVector.m_dX, CartesianTransformationData->m_sYVector.m_dY, CartesianTransformationData->m_sYVector.m_dZ);;
 
-			FVector ZVector = XVector ^ YVector;
-			FMatrix Matrix(XVector, YVector, ZVector, FVector::Zero());
+		FVector ZVector = XVector ^ YVector;
 
-			FMatrix Scale = FMatrix::Identity;
-			Scale.M[0][0] = CartesianTransformationData->m_sScale.m_dX;
-			Scale.M[1][1] = CartesianTransformationData->m_sScale.m_dY;
-			Scale.M[2][2] = CartesianTransformationData->m_sScale.m_dZ;
-			Matrix = Matrix * Scale;
+		const A3DVector3dData& Scale = CartesianTransformationData->m_sScale;
 
-			Matrix = Matrix.ConcatTranslation(Origin);
+		FMatrix Matrix(XVector * Scale.m_dX, YVector * Scale.m_dY, ZVector * Scale.m_dZ, FVector::Zero());
 
-			return Matrix;
+		if (CartesianTransformationData->m_ucBehaviour & kA3DTransformationMirror)
+		{
+			Matrix.M[2][0] *= -1;
+			Matrix.M[2][1] *= -1;
+			Matrix.M[2][2] *= -1;
 		}
+
+		Matrix.SetOrigin(Origin * FileUnit);
+
+		return Matrix;
 	}
+
 	return FMatrix::Identity;
 }
 
-FMatrix FTechSoftFileParser::TraverseGeneralTransformation(const A3DMiscGeneralTransformation* GeneralTransformation)
+FMatrix FTechSoftFileParser::TraverseGeneralTransformation(const A3DMiscTransformation* GeneralTransformation)
 {
 	TUniqueTSObj<A3DMiscGeneralTransformationData> GeneralTransformationData(GeneralTransformation);
 	if (GeneralTransformationData.IsValid())
@@ -1271,12 +1664,18 @@ FMatrix FTechSoftFileParser::TraverseGeneralTransformation(const A3DMiscGeneralT
 				Matrix.M[Andex][Bndex] = GeneralTransformationData->m_adCoeff[Index];
 			}
 		}
+
+		for (Index = 0; Index < 3; ++Index, ++Index)
+		{
+			Matrix.M[3][Index] *= FileUnit;
+		}
+
 		return Matrix;
 	}
 	return FMatrix::Identity;
 }
 
-FMatrix FTechSoftFileParser::TraverseTransformation(const A3DMiscCartesianTransformation* Transformation3D)
+FMatrix FTechSoftFileParser::TraverseTransformation(const A3DMiscTransformation* Transformation3D)
 {
 	if (Transformation3D == NULL)
 	{
@@ -1354,122 +1753,6 @@ bool FTechSoftFileParser::IsConfigurationSet(const A3DAsmProductOccurrence* Occu
 	return bIsConfiguration;
 }
 
-void FTechSoftFileParser::CountUnderModel(const A3DAsmModelFile* AsmModel)
-{
-	TUniqueTSObj<A3DAsmModelFileData> ModelFileData(AsmModel);
-	if (!ModelFileData.IsValid())
-	{
-		return;
-	}
-
-	ComponentCount[EComponentType::Occurrence] ++;
-
-	for (uint32 Index = 0; Index < ModelFileData->m_uiPOccurrencesSize; ++Index)
-	{
-		//if (IsConfigurationSet(ModelFileData.m_ppPOccurrences[Index]))
-		//{
-		//	TraverseConfigurationSet(ModelFileData.m_ppPOccurrences[Index], MetaData);
-		//}
-		//else
-		{
-			CountUnderOccurrence(ModelFileData->m_ppPOccurrences[Index]);
-		}
-	}
-}
-
-void FTechSoftFileParser::CountUnderOccurrence(const A3DAsmProductOccurrence* Occurrence)
-{
-	TUniqueTSObj<A3DAsmProductOccurrenceData> OccurrenceData(Occurrence);
-	if (!OccurrenceData.IsValid())
-	{
-		return;
-	}
-
-	ComponentCount[EComponentType::Occurrence] ++;
-
-	if (OccurrenceData->m_pPrototype)
-	{
-		CountUnderPrototype(OccurrenceData->m_pPrototype);
-	}
-
-	for (uint32 Index = 0; Index < OccurrenceData->m_uiPOccurrencesSize; ++Index)
-	{
-		CountUnderOccurrence(OccurrenceData->m_ppPOccurrences[Index]);
-	}
-
-	if (OccurrenceData->m_pPart)
-	{
-		CountUnderPartDefinition(OccurrenceData->m_pPart);
-	}
-}
-
-void FTechSoftFileParser::CountUnderPrototype(const A3DAsmProductOccurrence* Prototype)
-{
-	if (PrototypeCounted.Contains(Prototype))
-	{
-		return;
-	}
-	PrototypeCounted.Add(Prototype);
-
-	TUniqueTSObj<A3DAsmProductOccurrenceData> PrototypeData(Prototype);
-	if (!PrototypeData.IsValid())
-	{
-		return;
-	}
-
-	ComponentCount[EComponentType::Reference] ++;
-
-	if (PrototypeData->m_pPrototype)
-	{
-		CountUnderPrototype(PrototypeData->m_pPrototype);
-	}
-}
-
-void FTechSoftFileParser::CountUnderPartDefinition(const A3DAsmPartDefinition* PartDefinition)
-{
-	TUniqueTSObj<A3DAsmPartDefinitionData> PartData(PartDefinition);
-	if (!PartData.IsValid())
-	{
-		return;
-	}
-
-	for (unsigned int Index = 0; Index < PartData->m_uiRepItemsSize; ++Index)
-	{
-		CountUnderRepresentation(PartData->m_ppRepItems[Index]);
-	}
-
-}
-
-void FTechSoftFileParser::CountUnderRepresentation(const A3DRiRepresentationItem* RepresentationItem)
-{
-	A3DEEntityType Type;
-	A3DEntityGetType(RepresentationItem, &Type);
-
-	switch (Type)
-	{
-	case kA3DTypeRiSet:
-		CountUnderRepresentationSet(RepresentationItem);
-		break;
-	case kA3DTypeRiBrepModel:
-	case kA3DTypeRiPolyBrepModel:
-		ComponentCount[EComponentType::Body] ++;
-		break;
-
-	default:
-		break;
-	}
-}
-
-void FTechSoftFileParser::CountUnderRepresentationSet(const A3DRiSet* RepresentationSet)
-{
-	TUniqueTSObj<A3DRiSetData> RepresentationSetData(RepresentationSet);
-	if (!RepresentationSetData.IsValid())
-	{
-		return;
-	}
-	ComponentCount[EComponentType::Body] += RepresentationSetData->m_uiRepItemsSize;
-}
-
 int32 FTechSoftFileParser::CountMaterial()
 {
 	// TODO
@@ -1546,67 +1829,101 @@ void FTechSoftFileParser::TraverseTessellationBase(const A3DTessBase* Tessellati
 	}
 }
 
-void AddFaceTriangleWithUniqueNormal(FTessellationData& Tessellation, const A3DTess3DData& Tessellation3DData, long& InOutStartIndex, unsigned long InTrinangleCount)
+namespace TechSoftFileParserImpl
 {
-	TechSoftFileParserImpl::Reserve(Tessellation, InTrinangleCount, /*bWithTexture*/ false);
 
-	unsigned long Index = InOutStartIndex;
+uint32 CountTriangles(const A3DTessFaceData& FaceTessData)
+{
+	const int32 TessellationFaceDataWithTriangle = 0x2222;
+	const int32 TessellationFaceDataWithFan = 0x4444;
+	const int32 TessellationFaceDataWithStrip = 0x8888;
+	const int32 TessellationFaceDataWithOneNormal = 0xE0E0;
+
+	uint32 UsedEntitiesFlags = FaceTessData.m_usUsedEntitiesFlags;
+
+	uint32 TriangleCount = 0;
+	uint32 FaceSetIndex = 0;
+	if (UsedEntitiesFlags & TessellationFaceDataWithTriangle)
+	{
+		TriangleCount += FaceTessData.m_puiSizesTriangulated[FaceSetIndex];
+		FaceSetIndex++;
+	}
 
+	if (FaceTessData.m_uiSizesTriangulatedSize > FaceSetIndex)
+	{
+		if (UsedEntitiesFlags & TessellationFaceDataWithFan)
+		{
+			uint32 LastFanIndex = 1 + FaceSetIndex + FaceTessData.m_puiSizesTriangulated[FaceSetIndex];
+			FaceSetIndex++;
+			for (; FaceSetIndex < LastFanIndex; FaceSetIndex++)
+			{
+				uint32 FanSize = (FaceTessData.m_puiSizesTriangulated[FaceSetIndex] & kA3DTessFaceDataNormalMask);
+				TriangleCount += (FanSize - 2);
+			}
+		}
+	}
+
+	if (FaceTessData.m_uiSizesTriangulatedSize > FaceSetIndex)
+	{
+		FaceSetIndex++;
+		for (; FaceSetIndex < FaceTessData.m_uiSizesTriangulatedSize; FaceSetIndex++)
+		{
+			uint32 StripSize = (FaceTessData.m_puiSizesTriangulated[FaceSetIndex] & kA3DTessFaceDataNormalMask);
+			TriangleCount += (StripSize - 2);
+		}
+	}
+	return TriangleCount;
+}
+
+void AddFaceTriangleWithUniqueNormal(FTessellationData& Tessellation, const A3DTess3DData& Tessellation3DData, const uint32 InTriangleCount, uint32& InOutLastTriangleIndex, int32& InOutLastVertexIndex)
+{
 	int32 FaceIndex[3] = { 0, 0, 0 };
 	int32 NormalIndex[3] = { 0, 0, 0 };
-	int32 VertexIndex = 0;
-
-	NormalIndex[0] = Tessellation3DData.m_puiTriangulatedIndexes[Index++];
-	NormalIndex[1] = NormalIndex[0];
-	NormalIndex[2] = NormalIndex[0];
 
 	// Get Triangles
-	for (unsigned long TriangleIndex = 0; TriangleIndex < InTrinangleCount; TriangleIndex++)
+	for (uint32 TriangleIndex = 0; TriangleIndex < InTriangleCount; TriangleIndex++)
 	{
-		FaceIndex[0] = Tessellation3DData.m_puiTriangulatedIndexes[Index++];
-		FaceIndex[1] = Tessellation3DData.m_puiTriangulatedIndexes[Index++];
-		FaceIndex[2] = Tessellation3DData.m_puiTriangulatedIndexes[Index++];
+		NormalIndex[0] = Tessellation3DData.m_puiTriangulatedIndexes[InOutLastTriangleIndex++];
+		NormalIndex[1] = NormalIndex[0];
+		NormalIndex[2] = NormalIndex[0];
+
+		FaceIndex[0] = Tessellation3DData.m_puiTriangulatedIndexes[InOutLastTriangleIndex++] / 3;
+		FaceIndex[1] = Tessellation3DData.m_puiTriangulatedIndexes[InOutLastTriangleIndex++] / 3;
+		FaceIndex[2] = Tessellation3DData.m_puiTriangulatedIndexes[InOutLastTriangleIndex++] / 3;
 
-		if (!TechSoftFileParserImpl::AddFace(FaceIndex, Tessellation, VertexIndex))
+		if (!TechSoftFileParserImpl::AddFace(FaceIndex, Tessellation, InOutLastVertexIndex))
 		{
 			continue;
 		}
 
 		TechSoftFileParserImpl::AddNormals(Tessellation3DData.m_pdNormals, NormalIndex, Tessellation.NormalArray);
 	}
-
-	InOutStartIndex = Index;
 }
 
-void AddFaceTriangleWithUniqueNormalAndTexture(FTessellationData& Tessellation, const A3DTess3DData& Tessellation3DData, long& InOutStartIndex, unsigned long InTrinangleCount, long TextureCount)
+void AddFaceTriangleWithUniqueNormalAndTexture(FTessellationData& Tessellation, const A3DTess3DData& Tessellation3DData, const uint32 InTriangleCount, const uint32 TextureCount, uint32& InOutStartIndex, int32& InOutLastVertexIndex)
 {
-	TechSoftFileParserImpl::Reserve(Tessellation, InTrinangleCount, /*bWithTexture*/ true);
-
-	unsigned long Index = InOutStartIndex;
-
 	int32 FaceIndex[3] = { 0, 0, 0 };
 	int32 NormalIndex[3] = { 0, 0, 0 };
 	int32 TextureIndex[3] = { 0, 0, 0 };
-	int32 VertexIndex = 0;
-
-	NormalIndex[0] = Tessellation3DData.m_puiTriangulatedIndexes[Index++];
-	NormalIndex[1] = NormalIndex[0];
-	NormalIndex[2] = NormalIndex[0];
 
 	// Get Triangles
-	for (unsigned long TriangleIndex = 0; TriangleIndex < InTrinangleCount; TriangleIndex++)
-	{
-		TextureIndex[0] = Tessellation3DData.m_puiTriangulatedIndexes[Index];
-		Index += TextureCount;
-		FaceIndex[0] = Tessellation3DData.m_puiTriangulatedIndexes[Index++];
-		TextureIndex[1] = Tessellation3DData.m_puiTriangulatedIndexes[Index];
-		Index += TextureCount;
-		FaceIndex[1] = Tessellation3DData.m_puiTriangulatedIndexes[Index++];
-		TextureIndex[2] = Tessellation3DData.m_puiTriangulatedIndexes[Index];
-		Index += TextureCount;
-		FaceIndex[2] = Tessellation3DData.m_puiTriangulatedIndexes[Index++];
-
-		if (!TechSoftFileParserImpl::AddFace(FaceIndex, Tessellation, VertexIndex))
+	for (uint32 TriangleIndex = 0; TriangleIndex < InTriangleCount; TriangleIndex++)
+	{
+		NormalIndex[0] = Tessellation3DData.m_puiTriangulatedIndexes[InOutStartIndex++];
+		NormalIndex[1] = NormalIndex[0];
+		NormalIndex[2] = NormalIndex[0];
+
+		TextureIndex[0] = Tessellation3DData.m_puiTriangulatedIndexes[InOutStartIndex];
+		InOutStartIndex += TextureCount;
+		FaceIndex[0] = Tessellation3DData.m_puiTriangulatedIndexes[InOutStartIndex++] / 3;
+		TextureIndex[1] = Tessellation3DData.m_puiTriangulatedIndexes[InOutStartIndex];
+		InOutStartIndex += TextureCount;
+		FaceIndex[1] = Tessellation3DData.m_puiTriangulatedIndexes[InOutStartIndex++] / 3;
+		TextureIndex[2] = Tessellation3DData.m_puiTriangulatedIndexes[InOutStartIndex];
+		InOutStartIndex += TextureCount;
+		FaceIndex[2] = Tessellation3DData.m_puiTriangulatedIndexes[InOutStartIndex++] / 3;
+
+		if (!TechSoftFileParserImpl::AddFace(FaceIndex, Tessellation, InOutLastVertexIndex))
 		{
 			continue;
 		}
@@ -1614,156 +1931,371 @@ void AddFaceTriangleWithUniqueNormalAndTexture(FTessellationData& Tessellation,
 		TechSoftFileParserImpl::AddNormals(Tessellation3DData.m_pdNormals, NormalIndex, Tessellation.NormalArray);
 		TechSoftFileParserImpl::AddTextureCoordinates(Tessellation3DData.m_pdTextureCoords, TextureIndex, Tessellation.TexCoordArray);
 	}
+}
+
+void AddFaceTriangle(FTessellationData& Tessellation, const A3DTess3DData& Tessellation3DData, const uint32 InTriangleCount, uint32& InOutStartIndex, int32& InOutLastVertexIndex)
+{
+	int32 FaceIndex[3] = { 0, 0, 0 };
+	int32 NormalIndex[3] = { 0, 0, 0 };
 
-	InOutStartIndex = Index;
+	// Get Triangles
+	for (uint32 TriangleIndex = 0; TriangleIndex < InTriangleCount; TriangleIndex++)
+	{
+		NormalIndex[0] = Tessellation3DData.m_puiTriangulatedIndexes[InOutStartIndex++];
+		FaceIndex[0] = Tessellation3DData.m_puiTriangulatedIndexes[InOutStartIndex++] / 3;
+		NormalIndex[1] = Tessellation3DData.m_puiTriangulatedIndexes[InOutStartIndex++];
+		FaceIndex[1] = Tessellation3DData.m_puiTriangulatedIndexes[InOutStartIndex++] / 3;
+		NormalIndex[2] = Tessellation3DData.m_puiTriangulatedIndexes[InOutStartIndex++];
+		FaceIndex[2] = Tessellation3DData.m_puiTriangulatedIndexes[InOutStartIndex++] / 3;
+
+		if (TechSoftFileParserImpl::AddFace(FaceIndex, Tessellation, InOutLastVertexIndex))
+		{
+			TechSoftFileParserImpl::AddNormals(Tessellation3DData.m_pdNormals, NormalIndex, Tessellation.NormalArray);
+		}
+	}
 }
 
-void AddFaceTriangle(FTessellationData& Tessellation, const A3DTess3DData& Tessellation3DData, long& InOutStartIndex, unsigned long InTrinangleCount)
+void AddFaceTriangleWithTexture(FTessellationData& Tessellation, const A3DTess3DData& Tessellation3DData, const uint32 InTriangleCount, const uint32 InTextureCount, uint32& InOutStartIndex, int32& inOutLastVertexIndex)
 {
-	TechSoftFileParserImpl::Reserve(Tessellation, InTrinangleCount, /*bWithTexture*/ false);
+	int32 FaceIndex[3] = { 0, 0, 0 };
+	int32 NormalIndex[3] = { 0, 0, 0 };
+	int32 TextureIndex[3] = { 0, 0, 0 };
 
-	unsigned long Index = InOutStartIndex;
+	// Get Triangles
+	for (uint64 TriangleIndex = 0; TriangleIndex < InTriangleCount; TriangleIndex++)
+	{
+		NormalIndex[0] = Tessellation3DData.m_puiTriangulatedIndexes[InOutStartIndex++];
+		TextureIndex[0] = Tessellation3DData.m_puiTriangulatedIndexes[InOutStartIndex];
+		InOutStartIndex += InTextureCount;
+		FaceIndex[0] = Tessellation3DData.m_puiTriangulatedIndexes[InOutStartIndex++] / 3;
+		NormalIndex[1] = Tessellation3DData.m_puiTriangulatedIndexes[InOutStartIndex++];
+		TextureIndex[1] = Tessellation3DData.m_puiTriangulatedIndexes[InOutStartIndex];
+		InOutStartIndex += InTextureCount;
+		FaceIndex[1] = Tessellation3DData.m_puiTriangulatedIndexes[InOutStartIndex++] / 3;
+		NormalIndex[2] = Tessellation3DData.m_puiTriangulatedIndexes[InOutStartIndex++];
+		TextureIndex[2] = Tessellation3DData.m_puiTriangulatedIndexes[InOutStartIndex];
+		InOutStartIndex += InTextureCount;
+		FaceIndex[2] = Tessellation3DData.m_puiTriangulatedIndexes[InOutStartIndex++] / 3;
+
+		if (TechSoftFileParserImpl::AddFace(FaceIndex, Tessellation, inOutLastVertexIndex))
+		{
+			TechSoftFileParserImpl::AddNormals(Tessellation3DData.m_pdNormals, NormalIndex, Tessellation.NormalArray);
+			TechSoftFileParserImpl::AddTextureCoordinates(Tessellation3DData.m_pdTextureCoords, TextureIndex, Tessellation.TexCoordArray);
+		}
+	}
+}
 
+void AddFaceTriangleFanWithUniqueNormal(FTessellationData& Tessellation, const A3DTess3DData& Tessellation3DData, const uint32 InTriangleCount, uint32& InOutLastTriangleIndex, int32& LastVertexIndex)
+{
 	int32 FaceIndex[3] = { 0, 0, 0 };
 	int32 NormalIndex[3] = { 0, 0, 0 };
-	int32 VertexIndex = 0;
+
+	NormalIndex[0] = Tessellation3DData.m_puiTriangulatedIndexes[InOutLastTriangleIndex++];
+	NormalIndex[1] = NormalIndex[0];
+	NormalIndex[2] = NormalIndex[0];
+
+	FaceIndex[0] = Tessellation3DData.m_puiTriangulatedIndexes[InOutLastTriangleIndex++] / 3;
+	FaceIndex[1] = Tessellation3DData.m_puiTriangulatedIndexes[InOutLastTriangleIndex++] / 3;
 
 	// Get Triangles
-	for (unsigned long TriangleIndex = 0; TriangleIndex < InTrinangleCount; TriangleIndex++)
+	for (uint32 TriangleIndex = 2; TriangleIndex < InTriangleCount; TriangleIndex++)
 	{
-		NormalIndex[0] = Tessellation3DData.m_puiTriangulatedIndexes[Index++];
-		FaceIndex[0] = Tessellation3DData.m_puiTriangulatedIndexes[Index++];
-		NormalIndex[1] = Tessellation3DData.m_puiTriangulatedIndexes[Index++];
-		FaceIndex[1] = Tessellation3DData.m_puiTriangulatedIndexes[Index++];
-		NormalIndex[2] = Tessellation3DData.m_puiTriangulatedIndexes[Index++];
-		FaceIndex[2] = Tessellation3DData.m_puiTriangulatedIndexes[Index++];
+		FaceIndex[2] = Tessellation3DData.m_puiTriangulatedIndexes[InOutLastTriangleIndex++] / 3;
 
-		if (!TechSoftFileParserImpl::AddFace(FaceIndex, Tessellation, VertexIndex))
+		if (TechSoftFileParserImpl::AddFace(FaceIndex, Tessellation, LastVertexIndex))
 		{
-			continue;
+			TechSoftFileParserImpl::AddNormals(Tessellation3DData.m_pdNormals, NormalIndex, Tessellation.NormalArray);
 		}
 
-		TechSoftFileParserImpl::AddNormals(Tessellation3DData.m_pdNormals, NormalIndex, Tessellation.NormalArray);
+		FaceIndex[1] = FaceIndex[2];
 	}
+}
 
-	InOutStartIndex = Index;
+void AddFaceTriangleFanWithUniqueNormalAndTexture(FTessellationData& Tessellation, const A3DTess3DData& Tessellation3DData, const uint32 InTriangleCount, const uint32 TextureCount, uint32& InOutLastTriangleIndex, int32& InOutLastVertexIndex)
+{
+	int32 FaceIndex[3] = { 0, 0, 0 };
+	int32 NormalIndex[3] = { 0, 0, 0 };
+	int32 TextureIndex[3] = { 0, 0, 0 };
+
+	NormalIndex[0] = Tessellation3DData.m_puiTriangulatedIndexes[InOutLastTriangleIndex++];
+	NormalIndex[1] = NormalIndex[0];
+	NormalIndex[2] = NormalIndex[0];
+
+	TextureIndex[0] = Tessellation3DData.m_puiTriangulatedIndexes[InOutLastTriangleIndex];
+	InOutLastTriangleIndex += TextureCount;
+	FaceIndex[0] = Tessellation3DData.m_puiTriangulatedIndexes[InOutLastTriangleIndex++] / 3;
+
+	TextureIndex[1] = Tessellation3DData.m_puiTriangulatedIndexes[InOutLastTriangleIndex];
+	InOutLastTriangleIndex += TextureCount;
+	FaceIndex[1] = Tessellation3DData.m_puiTriangulatedIndexes[InOutLastTriangleIndex++] / 3;
+
+	for (uint32 TriangleIndex = 2; TriangleIndex < InTriangleCount; TriangleIndex++)
+	{
+		TextureIndex[2] = Tessellation3DData.m_puiTriangulatedIndexes[InOutLastTriangleIndex];
+		InOutLastTriangleIndex += TextureCount;
+		FaceIndex[2] = Tessellation3DData.m_puiTriangulatedIndexes[InOutLastTriangleIndex++] / 3;
+
+		if (TechSoftFileParserImpl::AddFace(FaceIndex, Tessellation, InOutLastVertexIndex))
+		{
+			TechSoftFileParserImpl::AddNormals(Tessellation3DData.m_pdNormals, NormalIndex, Tessellation.NormalArray);
+			TechSoftFileParserImpl::AddTextureCoordinates(Tessellation3DData.m_pdTextureCoords, TextureIndex, Tessellation.TexCoordArray);
+		}
+
+		FaceIndex[1] = FaceIndex[2];
+		TextureIndex[1] = TextureIndex[2];
+	}
 }
 
-void AddFaceTriangleWithTexture(FTessellationData& Tessellation, const A3DTess3DData& Tessellation3DData, long& InOutStartIndex, unsigned long InTrinangleCount, long TextureCount)
+void AddFaceTriangleFan(FTessellationData& Tessellation, const A3DTess3DData& Tessellation3DData, const uint32 InTriangleCount, uint32& InOutLastTriangleIndex, int32& InOutLastVertexIndex)
 {
-	TechSoftFileParserImpl::Reserve(Tessellation, InTrinangleCount, /*bWithTexture*/ true);
+	int32 FaceIndex[3] = { 0, 0, 0 };
+	int32 NormalIndex[3] = { 0, 0, 0 };
+
+	NormalIndex[0] = Tessellation3DData.m_puiTriangulatedIndexes[InOutLastTriangleIndex++];
+	FaceIndex[0] = Tessellation3DData.m_puiTriangulatedIndexes[InOutLastTriangleIndex++] / 3;
+	NormalIndex[1] = Tessellation3DData.m_puiTriangulatedIndexes[InOutLastTriangleIndex++];
+	FaceIndex[1] = Tessellation3DData.m_puiTriangulatedIndexes[InOutLastTriangleIndex++] / 3;
 
-	unsigned long Index = InOutStartIndex;
+	for (unsigned long TriangleIndex = 2; TriangleIndex < InTriangleCount; TriangleIndex++)
+	{
+		NormalIndex[2] = Tessellation3DData.m_puiTriangulatedIndexes[InOutLastTriangleIndex++];
+		FaceIndex[2] = Tessellation3DData.m_puiTriangulatedIndexes[InOutLastTriangleIndex++] / 3;
 
+		if (TechSoftFileParserImpl::AddFace(FaceIndex, Tessellation, InOutLastVertexIndex))
+		{
+			TechSoftFileParserImpl::AddNormals(Tessellation3DData.m_pdNormals, NormalIndex, Tessellation.NormalArray);
+		}
+
+		NormalIndex[1] = NormalIndex[2];
+		FaceIndex[1] = FaceIndex[2];
+	}
+}
+
+void AddFaceTriangleFanWithTexture(FTessellationData& Tessellation, const A3DTess3DData& Tessellation3DData, const uint32 InTriangleCount, const uint32 TextureCount, uint32& InOutLastTriangleIndex, int32& InOutLastVertexIndex)
+{
 	int32 FaceIndex[3] = { 0, 0, 0 };
 	int32 NormalIndex[3] = { 0, 0, 0 };
 	int32 TextureIndex[3] = { 0, 0, 0 };
-	int32 VertexIndex = 0;
 
-	// Get Triangles
-	for (unsigned long TriangleIndex = 0; TriangleIndex < InTrinangleCount; TriangleIndex++)
-	{
-		NormalIndex[0] = Tessellation3DData.m_puiTriangulatedIndexes[Index++];
-		TextureIndex[0] = Tessellation3DData.m_puiTriangulatedIndexes[Index];
-		Index += TextureCount;
-		FaceIndex[0] = Tessellation3DData.m_puiTriangulatedIndexes[Index++];
-		NormalIndex[1] = Tessellation3DData.m_puiTriangulatedIndexes[Index++];
-		TextureIndex[1] = Tessellation3DData.m_puiTriangulatedIndexes[Index];
-		Index += TextureCount;
-		FaceIndex[1] = Tessellation3DData.m_puiTriangulatedIndexes[Index++];
-		NormalIndex[2] = Tessellation3DData.m_puiTriangulatedIndexes[Index++];
-		TextureIndex[2] = Tessellation3DData.m_puiTriangulatedIndexes[Index];
-		Index += TextureCount;
-		FaceIndex[2] = Tessellation3DData.m_puiTriangulatedIndexes[Index++];
-
-		if (!TechSoftFileParserImpl::AddFace(FaceIndex, Tessellation, VertexIndex))
+	NormalIndex[0] = Tessellation3DData.m_puiTriangulatedIndexes[InOutLastTriangleIndex++];
+	TextureIndex[0] = Tessellation3DData.m_puiTriangulatedIndexes[InOutLastTriangleIndex];
+	InOutLastTriangleIndex += TextureCount;
+	FaceIndex[0] = Tessellation3DData.m_puiTriangulatedIndexes[InOutLastTriangleIndex++] / 3;
+
+	NormalIndex[1] = Tessellation3DData.m_puiTriangulatedIndexes[InOutLastTriangleIndex++];
+	TextureIndex[1] = Tessellation3DData.m_puiTriangulatedIndexes[InOutLastTriangleIndex];
+	InOutLastTriangleIndex += TextureCount;
+	FaceIndex[1] = Tessellation3DData.m_puiTriangulatedIndexes[InOutLastTriangleIndex++] / 3;
+
+	for (uint32 TriangleIndex = 2; TriangleIndex < InTriangleCount; TriangleIndex++)
+	{
+		NormalIndex[2] = Tessellation3DData.m_puiTriangulatedIndexes[InOutLastTriangleIndex++];
+		TextureIndex[2] = Tessellation3DData.m_puiTriangulatedIndexes[InOutLastTriangleIndex];
+		InOutLastTriangleIndex += TextureCount;
+		FaceIndex[2] = Tessellation3DData.m_puiTriangulatedIndexes[InOutLastTriangleIndex++] / 3;
+
+		if (TechSoftFileParserImpl::AddFace(FaceIndex, Tessellation, InOutLastVertexIndex))
 		{
-			continue;
+			TechSoftFileParserImpl::AddNormals(Tessellation3DData.m_pdNormals, NormalIndex, Tessellation.NormalArray);
+			TechSoftFileParserImpl::AddTextureCoordinates(Tessellation3DData.m_pdTextureCoords, TextureIndex, Tessellation.TexCoordArray);
 		}
 
-		TechSoftFileParserImpl::AddNormals(Tessellation3DData.m_pdNormals, NormalIndex, Tessellation.NormalArray);
-		TechSoftFileParserImpl::AddTextureCoordinates(Tessellation3DData.m_pdTextureCoords, TextureIndex, Tessellation.TexCoordArray);
+		NormalIndex[1] = NormalIndex[2];
+		TextureIndex[1] = TextureIndex[2];
+		FaceIndex[1] = FaceIndex[2];
 	}
-
-	InOutStartIndex = Index;
 }
 
-// TODO import fan and strip meshes
-#ifdef UNUSED
-double ComputeFaceAreaTriangleFan(double* InCoordinates, unsigned long* InTriangulatedIndexes, long& InOutStartIndex, unsigned long InTrinangleCount, bool bInOneNormal, long TextureCount, int& TriangleCount)
+void AddFaceTriangleStripWithUniqueNormal(FTessellationData& Tessellation, const A3DTess3DData& Tessellation3DData, const uint32 InTriangleCount, uint32& InOutLastTriangleIndex, int32& InOutLastVertexIndex)
 {
-	double Area = 0;
+	int32 FaceIndex[3] = { 0, 0, 0 };
+	int32 NormalIndex[3] = { 0, 0, 0 };
 
-	unsigned long OffSet = TextureCount + (bInOneNormal ? 0 : 1) + 1;
-	unsigned long Index = InOutStartIndex + 2 + TextureCount;
+	NormalIndex[0] = Tessellation3DData.m_puiTriangulatedIndexes[InOutLastTriangleIndex++];
+	NormalIndex[1] = NormalIndex[0];
+	NormalIndex[2] = NormalIndex[0];
 
-	int Index0 = InTriangulatedIndexes[Index];
-	double* Point0 = InCoordinates + (InTriangulatedIndexes[Index]);
-	Index += OffSet;
-	int Index1 = InTriangulatedIndexes[Index];
-	double* Point1 = InCoordinates + (InTriangulatedIndexes[Index]);
-	Index += OffSet;
+	FaceIndex[0] = Tessellation3DData.m_puiTriangulatedIndexes[InOutLastTriangleIndex++] / 3;
+	FaceIndex[1] = Tessellation3DData.m_puiTriangulatedIndexes[InOutLastTriangleIndex++] / 3;
 
-	for (unsigned long TriangleIndex = 2; TriangleIndex < InTrinangleCount; TriangleIndex++)
+	for (unsigned long TriangleIndex = 0; TriangleIndex < InTriangleCount; TriangleIndex++)
 	{
-		double* Point2 = InCoordinates + (InTriangulatedIndexes[Index]);
-		int Index2 = InTriangulatedIndexes[Index];
-		if (Index0 == Index1 || Index0 == Index2 || Index1 == Index2)
+		FaceIndex[2] = Tessellation3DData.m_puiTriangulatedIndexes[InOutLastTriangleIndex++] / 3;
+
+		if (TechSoftFileParserImpl::AddFace(FaceIndex, Tessellation, InOutLastVertexIndex))
 		{
-			TriangleCount--;
+			TechSoftFileParserImpl::AddNormals(Tessellation3DData.m_pdNormals, NormalIndex, Tessellation.NormalArray);
 		}
-		else
+
+		TriangleIndex++;
+		if (TriangleIndex == InTriangleCount)
 		{
-			Area += FaceArea(Point0, Point1, Point2);
+			break;
 		}
-		Index += OffSet;
-		Point1 = Point2;
-		Index1 = Index2;
-	}
-	InOutStartIndex = Index - OffSet;
 
-	return Area;
+		Swap(FaceIndex[1], FaceIndex[2]);
+
+		NormalIndex[0] = Tessellation3DData.m_puiTriangulatedIndexes[InOutLastTriangleIndex++];
+		FaceIndex[0] = Tessellation3DData.m_puiTriangulatedIndexes[InOutLastTriangleIndex++] / 3;
+
+		if (TechSoftFileParserImpl::AddFace(FaceIndex, Tessellation, InOutLastVertexIndex))
+		{
+			TechSoftFileParserImpl::AddNormals(Tessellation3DData.m_pdNormals, NormalIndex, Tessellation.NormalArray);
+		}
+
+		Swap(FaceIndex[0], FaceIndex[1]);
+	}
 }
 
-double ComputeFaceAreaTriangleStrip(double* InCoordinates, unsigned long* InTriangulatedIndexes, long& InOutStartIndex, unsigned long InTrinangleCount, bool bInOneNormal, long TextureCount, int& TriangleCount)
+void AddFaceTriangleStripWithUniqueNormalAndTexture(FTessellationData& Tessellation, const A3DTess3DData& Tessellation3DData, const uint32 InTriangleCount, const uint32 TextureCount, uint32& InOutLastTriangleIndex, int32& InOutLastVertexIndex)
 {
-	double Area = 0;
+	int32 FaceIndex[3] = { 0, 0, 0 };
+	int32 NormalIndex[3] = { 0, 0, 0 };
+	int32 TextureIndex[3] = { 0, 0, 0 };
 
-	unsigned long OffSet = TextureCount + (bInOneNormal ? 0 : 1) + 1;
-	unsigned long Index = InOutStartIndex + 2 + TextureCount;
+	NormalIndex[0] = Tessellation3DData.m_puiTriangulatedIndexes[InOutLastTriangleIndex++];
+	NormalIndex[1] = NormalIndex[0];
+	NormalIndex[2] = NormalIndex[0];
+
+	TextureIndex[0] = Tessellation3DData.m_puiTriangulatedIndexes[InOutLastTriangleIndex];
+	InOutLastTriangleIndex += TextureCount;
+	FaceIndex[0] = Tessellation3DData.m_puiTriangulatedIndexes[InOutLastTriangleIndex++];
 
-	int Index0 = InTriangulatedIndexes[Index];
-	double* Point0 = InCoordinates + (InTriangulatedIndexes[Index]);
-	Index += OffSet;
-	int Index1 = InTriangulatedIndexes[Index];
-	double* Point1 = InCoordinates + (InTriangulatedIndexes[Index]);
-	Index += OffSet;
+	TextureIndex[1] = Tessellation3DData.m_puiTriangulatedIndexes[InOutLastTriangleIndex];
+	InOutLastTriangleIndex += TextureCount;
+	FaceIndex[1] = Tessellation3DData.m_puiTriangulatedIndexes[InOutLastTriangleIndex++];
 
-	for (unsigned long TriangleIndex = 2; TriangleIndex < InTrinangleCount; TriangleIndex++)
+	for (unsigned long TriangleIndex = 0; TriangleIndex < InTriangleCount; TriangleIndex++)
 	{
-		double* Point2 = InCoordinates + (InTriangulatedIndexes[Index]);
-		int Index2 = InTriangulatedIndexes[Index];
-		if (Index0 == Index1 || Index0 == Index2 || Index1 == Index2)
+		TextureIndex[2] = Tessellation3DData.m_puiTriangulatedIndexes[InOutLastTriangleIndex];
+		InOutLastTriangleIndex += TextureCount;
+		FaceIndex[2] = Tessellation3DData.m_puiTriangulatedIndexes[InOutLastTriangleIndex++];
+
+		if (TechSoftFileParserImpl::AddFace(FaceIndex, Tessellation, InOutLastVertexIndex))
 		{
-			TriangleCount--;
+			TechSoftFileParserImpl::AddNormals(Tessellation3DData.m_pdNormals, NormalIndex, Tessellation.NormalArray);
+			TechSoftFileParserImpl::AddTextureCoordinates(Tessellation3DData.m_pdTextureCoords, TextureIndex, Tessellation.TexCoordArray);
 		}
-		else
+
+		TriangleIndex++;
+		if (TriangleIndex == InTriangleCount)
+		{
+			break;
+		}
+
+		Swap(FaceIndex[1], FaceIndex[2]);
+		Swap(TextureIndex[1], TextureIndex[2]);
+
+		FaceIndex[0] = Tessellation3DData.m_puiTriangulatedIndexes[InOutLastTriangleIndex++] / 3;
+
+		if (TechSoftFileParserImpl::AddFace(FaceIndex, Tessellation, InOutLastVertexIndex))
 		{
-			Area += FaceArea(Point0, Point1, Point2);
+			TechSoftFileParserImpl::AddNormals(Tessellation3DData.m_pdNormals, NormalIndex, Tessellation.NormalArray);
 		}
 
-		Index += OffSet;
-		Point0 = Point1;
-		Point1 = Point2;
-		Index0 = Index1;
-		Index1 = Index2;
+		Swap(FaceIndex[0], FaceIndex[1]);
+		Swap(TextureIndex[0], TextureIndex[1]);
 	}
-	InOutStartIndex = Index - OffSet;
+}
+
+void AddFaceTriangleStrip(FTessellationData& Tessellation, const A3DTess3DData& Tessellation3DData, const uint32 InTriangleCount, uint32& InOutLastTriangleIndex, int32& LastVertexIndex)
+{
+	int32 FaceIndex[3] = { 0, 0, 0 };
+	int32 NormalIndex[3] = { 0, 0, 0 };
 
-	return Area;
+	NormalIndex[0] = Tessellation3DData.m_puiTriangulatedIndexes[InOutLastTriangleIndex++];
+	FaceIndex[0] = Tessellation3DData.m_puiTriangulatedIndexes[InOutLastTriangleIndex++] / 3;
+	NormalIndex[1] = Tessellation3DData.m_puiTriangulatedIndexes[InOutLastTriangleIndex++];
+	FaceIndex[1] = Tessellation3DData.m_puiTriangulatedIndexes[InOutLastTriangleIndex++] / 3;
+
+	for (unsigned long TriangleIndex = 2; TriangleIndex < InTriangleCount; TriangleIndex++)
+	{
+		NormalIndex[2] = Tessellation3DData.m_puiTriangulatedIndexes[InOutLastTriangleIndex++];
+		FaceIndex[2] = Tessellation3DData.m_puiTriangulatedIndexes[InOutLastTriangleIndex++] / 3;
+
+		if (TechSoftFileParserImpl::AddFace(FaceIndex, Tessellation, LastVertexIndex))
+		{
+			TechSoftFileParserImpl::AddNormals(Tessellation3DData.m_pdNormals, NormalIndex, Tessellation.NormalArray);
+		}
+
+		TriangleIndex++;
+		if (TriangleIndex == InTriangleCount)
+		{
+			break;
+		}
+
+		Swap(FaceIndex[1], FaceIndex[2]);
+		Swap(NormalIndex[1], NormalIndex[2]);
+
+		NormalIndex[0] = Tessellation3DData.m_puiTriangulatedIndexes[InOutLastTriangleIndex++];
+		FaceIndex[0] = Tessellation3DData.m_puiTriangulatedIndexes[InOutLastTriangleIndex++] / 3;
+
+		if (TechSoftFileParserImpl::AddFace(FaceIndex, Tessellation, LastVertexIndex))
+		{
+			TechSoftFileParserImpl::AddNormals(Tessellation3DData.m_pdNormals, NormalIndex, Tessellation.NormalArray);
+		}
+
+		Swap(FaceIndex[0], FaceIndex[1]);
+		Swap(NormalIndex[0], NormalIndex[1]);
+	}
 }
-#endif
 
+void AddFaceTriangleStripWithTexture(FTessellationData& Tessellation, const A3DTess3DData& Tessellation3DData, const uint32 InTriangleCount, const uint32 TextureCount, uint32& InOutLastTriangleIndex, int32& InOutLastVertexIndex)
+{
+	int32 FaceIndex[3] = { 0, 0, 0 };
+	int32 NormalIndex[3] = { 0, 0, 0 };
+	int32 TextureIndex[3] = { 0, 0, 0 };
+
+	NormalIndex[0] = Tessellation3DData.m_puiTriangulatedIndexes[InOutLastTriangleIndex++];
+	TextureIndex[0] = Tessellation3DData.m_puiTriangulatedIndexes[InOutLastTriangleIndex];
+	InOutLastTriangleIndex += TextureCount;
+	FaceIndex[0] = Tessellation3DData.m_puiTriangulatedIndexes[InOutLastTriangleIndex++];
+	NormalIndex[1] = Tessellation3DData.m_puiTriangulatedIndexes[InOutLastTriangleIndex++];
+	TextureIndex[1] = Tessellation3DData.m_puiTriangulatedIndexes[InOutLastTriangleIndex];
+	InOutLastTriangleIndex += TextureCount;
+	FaceIndex[1] = Tessellation3DData.m_puiTriangulatedIndexes[InOutLastTriangleIndex++];
+
+	for (unsigned long TriangleIndex = 0; TriangleIndex < InTriangleCount; TriangleIndex++)
+	{
+		NormalIndex[2] = Tessellation3DData.m_puiTriangulatedIndexes[InOutLastTriangleIndex++];
+		TextureIndex[2] = Tessellation3DData.m_puiTriangulatedIndexes[InOutLastTriangleIndex];
+		InOutLastTriangleIndex += TextureCount;
+		FaceIndex[2] = Tessellation3DData.m_puiTriangulatedIndexes[InOutLastTriangleIndex++];
+
+		if (TechSoftFileParserImpl::AddFace(FaceIndex, Tessellation, InOutLastVertexIndex))
+		{
+			TechSoftFileParserImpl::AddNormals(Tessellation3DData.m_pdNormals, NormalIndex, Tessellation.NormalArray);
+			TechSoftFileParserImpl::AddTextureCoordinates(Tessellation3DData.m_pdTextureCoords, TextureIndex, Tessellation.TexCoordArray);
+		}
+
+		TriangleIndex++;
+		if (TriangleIndex == InTriangleCount)
+		{
+			break;
+		}
+
+		Swap(FaceIndex[1], FaceIndex[2]);
+		Swap(NormalIndex[1], NormalIndex[2]);
+		Swap(TextureIndex[1], TextureIndex[2]);
+
+		NormalIndex[0] = Tessellation3DData.m_puiTriangulatedIndexes[InOutLastTriangleIndex++];
+		FaceIndex[0] = Tessellation3DData.m_puiTriangulatedIndexes[InOutLastTriangleIndex++] / 3;
+
+		if (TechSoftFileParserImpl::AddFace(FaceIndex, Tessellation, InOutLastVertexIndex))
+		{
+			TechSoftFileParserImpl::AddNormals(Tessellation3DData.m_pdNormals, NormalIndex, Tessellation.NormalArray);
+		}
+
+		Swap(FaceIndex[0], FaceIndex[1]);
+		Swap(NormalIndex[0], NormalIndex[1]);
+		Swap(TextureIndex[0], TextureIndex[1]);
+	}
+}
+
+} //ns FTechSoftFileParserImpl
+
+
+#pragma optimize ("",off)
 void FTechSoftFileParser::TraverseTessellation3D(const A3DTess3D* TessellationPtr, FArchiveBody& Body)
 {
 	FBodyMesh& BodyMesh = CADFileData.AddBodyMesh(Body.ObjectId, Body);
@@ -1826,11 +2358,15 @@ void FTechSoftFileParser::TraverseTessellation3D(const A3DTess3D* TessellationPt
 				}
 			}
 
-			unsigned int UsedEntitiesFlags = FaceTessData.m_usUsedEntitiesFlags;
-			long StartTriangulated = FaceTessData.m_uiStartTriangulated;
+			uint32 TriangleCount = TechSoftFileParserImpl::CountTriangles(FaceTessData);
+			TechSoftFileParserImpl::Reserve(Tessellation, TriangleCount, /*bWithTexture*/ FaceTessData.m_uiTextureCoordIndexesSize > 0);
+
+			uint32 UsedEntitiesFlags = FaceTessData.m_usUsedEntitiesFlags;
+			uint32 LastTrianguleIndex = FaceTessData.m_uiStartTriangulated;
+
+			uint32 FaceSetIndex = 0;
+			int32 LastVertexIndex = 0;
 
-			unsigned int FaceSetIndex = 0;
-			// Triangles
 			if (UsedEntitiesFlags & TessellationFaceDataWithTriangle)
 			{
 				bool bTessellationWithOneNormal = (UsedEntitiesFlags & TessellationFaceDataWithOneNormal);
@@ -1838,45 +2374,107 @@ void FTechSoftFileParser::TraverseTessellation3D(const A3DTess3D* TessellationPt
 				{
 					if (FaceTessData.m_uiTextureCoordIndexesSize)
 					{
-						AddFaceTriangleWithUniqueNormalAndTexture(Tessellation, *Tessellation3DData, StartTriangulated, FaceTessData.m_puiSizesTriangulated[0], FaceTessData.m_uiTextureCoordIndexesSize);
+						TechSoftFileParserImpl::AddFaceTriangleWithUniqueNormalAndTexture(Tessellation, *Tessellation3DData, FaceTessData.m_puiSizesTriangulated[0], FaceTessData.m_uiTextureCoordIndexesSize, LastTrianguleIndex, LastVertexIndex);
 					}
 					else
 					{
-						AddFaceTriangleWithUniqueNormal(Tessellation, *Tessellation3DData, StartTriangulated, FaceTessData.m_puiSizesTriangulated[0]);
+						TechSoftFileParserImpl::AddFaceTriangleWithUniqueNormal(Tessellation, *Tessellation3DData, FaceTessData.m_puiSizesTriangulated[0], LastTrianguleIndex, LastVertexIndex);
 					}
 				}
 				else
 				{
 					if (FaceTessData.m_uiTextureCoordIndexesSize)
 					{
-						AddFaceTriangleWithTexture(Tessellation, *Tessellation3DData, StartTriangulated, FaceTessData.m_puiSizesTriangulated[0], FaceTessData.m_uiTextureCoordIndexesSize);
+						TechSoftFileParserImpl::AddFaceTriangleWithTexture(Tessellation, *Tessellation3DData, FaceTessData.m_puiSizesTriangulated[0], FaceTessData.m_uiTextureCoordIndexesSize, LastTrianguleIndex, LastVertexIndex);
 					}
 					else
 					{
-						AddFaceTriangle(Tessellation, *Tessellation3DData, StartTriangulated, FaceTessData.m_puiSizesTriangulated[0]);
+						TechSoftFileParserImpl::AddFaceTriangle(Tessellation, *Tessellation3DData, FaceTessData.m_puiSizesTriangulated[0], LastTrianguleIndex, LastVertexIndex);
 					}
 				}
 				FaceSetIndex++;
 			}
 
-			// Fans TODO
 			if (FaceTessData.m_uiSizesTriangulatedSize > FaceSetIndex)
 			{
-				if (UsedEntitiesFlags & TessellationFaceDataWithFan)
+				//if (UsedEntitiesFlags & TessellationFaceDataWithFan)
+				//{
+				//	unsigned int LastFanIndex = 1 + FaceSetIndex + FaceTessData.m_puiSizesTriangulated[FaceSetIndex];
+				//	FaceSetIndex++;
+				//	for (; FaceSetIndex < LastFanIndex; FaceSetIndex++)
+				//	{
+				//		bool bTessellationWithOneNormal = (UsedEntitiesFlags & TessellationFaceDataWithOneNormal) && (FaceTessData.m_puiSizesTriangulated[FaceSetIndex] & kA3DTessFaceDataNormalSingle);
+				//		A3DUns32 FanSize = (FaceTessData.m_puiSizesTriangulated[FaceSetIndex] & kA3DTessFaceDataNormalMask);
+				//		if (bTessellationWithOneNormal)
+				//		{
+				//			if (FaceTessData.m_uiTextureCoordIndexesSize)
+				//			{
+				//				TechSoftFileParserImpl::AddFaceTriangleFanWithUniqueNormalAndTexture(Tessellation, *Tessellation3DData, FaceTessData.m_puiSizesTriangulated[FaceSetIndex], FaceTessData.m_uiTextureCoordIndexesSize, LastTrianguleIndex, LastVertexIndex);
+				//			}
+				//			else
+				//			{
+				//				TechSoftFileParserImpl::AddFaceTriangleFanWithUniqueNormal(Tessellation, *Tessellation3DData, FaceTessData.m_puiSizesTriangulated[FaceSetIndex], LastTrianguleIndex, LastVertexIndex);
+				//			}
+				//		}
+				//		else
+				//		{
+				//			if (FaceTessData.m_uiTextureCoordIndexesSize)
+				//			{
+				//				TechSoftFileParserImpl::AddFaceTriangleFanWithTexture(Tessellation, *Tessellation3DData, FaceTessData.m_puiSizesTriangulated[FaceSetIndex], FaceTessData.m_uiTextureCoordIndexesSize, LastTrianguleIndex, LastVertexIndex);
+				//			}
+				//			else
+				//			{
+				//				TechSoftFileParserImpl::AddFaceTriangleFan(Tessellation, *Tessellation3DData, FaceTessData.m_puiSizesTriangulated[FaceSetIndex], LastTrianguleIndex, LastVertexIndex);
+				//			}
+				//		}
+				//	}
+				//}
+
+				if (UsedEntitiesFlags & kA3DTessFaceDataTriangleFan)
 				{
-					unsigned int LastFanIndex = 1 + FaceSetIndex + FaceTessData.m_puiSizesTriangulated[FaceSetIndex];
-					FaceSetIndex++;
-					for (; FaceSetIndex < LastFanIndex; FaceSetIndex++)
+					uint32 FanCount = FaceTessData.m_puiSizesTriangulated[FaceSetIndex++];
+					for (uint32 FanIndex = 0; FanIndex < FanCount; ++FanIndex)
 					{
-						bool bTessellationWithOneNormal = (UsedEntitiesFlags & TessellationFaceDataWithOneNormal) && (FaceTessData.m_puiSizesTriangulated[FaceSetIndex] & kA3DTessFaceDataNormalSingle);
-						A3DUns32 FanSize = (FaceTessData.m_puiSizesTriangulated[FaceSetIndex] & kA3DTessFaceDataNormalMask);
-						//TriangleCount += (FanSize - 2);
-						//BodyMeshMetaData.MeshArea += ComputeFaceAreaTriangleFan(TessellationBaseData.m_pdCoords, Tessellation3DData->m_puiTriangulatedIndexes, StartTriangulated, FanSize, bTessellationWithOneNormal, FaceTessData.m_uiTextureCoordIndexesSize, TriangleCount);
+						uint32 VertexCount = FaceTessData.m_puiSizesTriangulated[FaceSetIndex++];
+						TechSoftFileParserImpl::AddFaceTriangleFan(Tessellation, *Tessellation3DData, VertexCount, LastTrianguleIndex, LastVertexIndex);
+					}
+				}
+
+				if (UsedEntitiesFlags & kA3DTessFaceDataTriangleFanOneNormal)
+				{
+					uint32 FanCount = FaceTessData.m_puiSizesTriangulated[FaceSetIndex++] & kA3DTessFaceDataNormalMask;
+					for (uint32 FanIndex = 0; FanIndex < FanCount; ++FanIndex)
+					{
+						ensure((FaceTessData.m_puiSizesTriangulated[FaceSetIndex] & kA3DTessFaceDataNormalSingle) != 0);
+
+						uint32 VertexCount = FaceTessData.m_puiSizesTriangulated[FaceSetIndex++] & kA3DTessFaceDataNormalMask;
+						TechSoftFileParserImpl::AddFaceTriangleFanWithUniqueNormal(Tessellation, *Tessellation3DData, VertexCount, LastTrianguleIndex, LastVertexIndex);
+					}
+				}
+
+				if (UsedEntitiesFlags & kA3DTessFaceDataTriangleFanTextured)
+				{
+					uint32 FanCount = FaceTessData.m_puiSizesTriangulated[FaceSetIndex++];
+					for (uint32 FanIndex = 0; FanIndex < FanCount; ++FanIndex)
+					{
+						uint32 VertexCount = FaceTessData.m_puiSizesTriangulated[FaceSetIndex++];
+						TechSoftFileParserImpl::AddFaceTriangleFanWithTexture(Tessellation, *Tessellation3DData, VertexCount, FaceTessData.m_uiTextureCoordIndexesSize, LastTrianguleIndex, LastVertexIndex);
+					}
+				}
+
+				if (UsedEntitiesFlags & kA3DTessFaceDataTriangleFanOneNormalTextured)
+				{
+					uint32 FanCount = FaceTessData.m_puiSizesTriangulated[FaceSetIndex++] & kA3DTessFaceDataNormalMask;
+					for (uint32 FanIndex = 0; FanIndex < FanCount; ++FanIndex)
+					{
+						ensure((FaceTessData.m_puiSizesTriangulated[FaceSetIndex] & kA3DTessFaceDataNormalSingle) != 0);
+
+						uint32 VertexCount = FaceTessData.m_puiSizesTriangulated[FaceSetIndex++] & kA3DTessFaceDataNormalMask;
+						TechSoftFileParserImpl::AddFaceTriangleFanWithUniqueNormalAndTexture(Tessellation, *Tessellation3DData, VertexCount, FaceTessData.m_uiTextureCoordIndexesSize, LastTrianguleIndex, LastVertexIndex);
 					}
 				}
 			}
 
-			// Strip TODO
 			if (FaceTessData.m_uiSizesTriangulatedSize > FaceSetIndex)
 			{
 				FaceSetIndex++;
@@ -1884,8 +2482,28 @@ void FTechSoftFileParser::TraverseTessellation3D(const A3DTess3D* TessellationPt
 				{
 					bool bTessellationWithOneNormal = (UsedEntitiesFlags & TessellationFaceDataWithOneNormal) && (FaceTessData.m_puiSizesTriangulated[FaceSetIndex] & kA3DTessFaceDataNormalSingle);
 					A3DUns32 StripSize = (FaceTessData.m_puiSizesTriangulated[FaceSetIndex] & kA3DTessFaceDataNormalMask);
-					//TriangleCount += (StripSize - 2);
-					//BodyMeshMetaData.MeshArea += ComputeFaceAreaTriangleStrip(TessellationBaseData.m_pdCoords, Tessellation3DData->m_puiTriangulatedIndexes, StartTriangulated, StripSize, bTessellationWithOneNormal, FaceTessData.m_uiTextureCoordIndexesSize, TriangleCount);
+					if (bTessellationWithOneNormal)
+					{
+						if (FaceTessData.m_uiTextureCoordIndexesSize)
+						{
+							TechSoftFileParserImpl::AddFaceTriangleStripWithUniqueNormalAndTexture(Tessellation, *Tessellation3DData, FaceTessData.m_puiSizesTriangulated[FaceSetIndex], FaceTessData.m_uiTextureCoordIndexesSize, LastTrianguleIndex, LastVertexIndex);
+						}
+						else
+						{
+							TechSoftFileParserImpl::AddFaceTriangleStripWithUniqueNormal(Tessellation, *Tessellation3DData, FaceTessData.m_puiSizesTriangulated[FaceSetIndex], LastTrianguleIndex, LastVertexIndex);
+						}
+					}
+					else
+					{
+						if (FaceTessData.m_uiTextureCoordIndexesSize)
+						{
+							TechSoftFileParserImpl::AddFaceTriangleStripWithTexture(Tessellation, *Tessellation3DData, FaceTessData.m_puiSizesTriangulated[FaceSetIndex], FaceTessData.m_uiTextureCoordIndexesSize, LastTrianguleIndex, LastVertexIndex);
+						}
+						else
+						{
+							TechSoftFileParserImpl::AddFaceTriangleStrip(Tessellation, *Tessellation3DData, FaceTessData.m_puiSizesTriangulated[FaceSetIndex], LastTrianguleIndex, LastVertexIndex);
+						}
+					}
 				}
 			}
 
@@ -1895,7 +2513,8 @@ void FTechSoftFileParser::TraverseTessellation3D(const A3DTess3D* TessellationPt
 	Body.ColorFaceSet = BodyMesh.ColorSet;
 	Body.MaterialFaceSet = BodyMesh.MaterialSet;
 }
+#pragma optimize ("",on)
 
-} // ns
+} // ns CADLibrary
 
-#endif
\ No newline at end of file
+#endif  
\ No newline at end of file
diff --git a/Engine/Plugins/Enterprise/DatasmithCADImporter/Source/CADInterfaces/Private/TechSoftFileParser.h b/Engine/Plugins/Enterprise/DatasmithCADImporter/Source/CADInterfaces/Private/TechSoftFileParser.h
--- a/Engine/Plugins/Enterprise/DatasmithCADImporter/Source/CADInterfaces/Private/TechSoftFileParser.h
+++ b/Engine/Plugins/Enterprise/DatasmithCADImporter/Source/CADInterfaces/Private/TechSoftFileParser.h
@@ -9,6 +9,10 @@
 #include "CADOptions.h"
 #include "CADSceneGraph.h"
 
+#include "TUniqueTechSoftObj.h"
+
+#define NEW_CODE
+
 typedef void A3DAsmModelFile;
 typedef void A3DAsmPartDefinition;
 typedef void A3DAsmProductOccurrence;
@@ -142,7 +146,7 @@ public:
 	virtual ECADParsingResult Process() override
 	{
 		return ECADParsingResult::ProcessFailed;
-}
+	}
 #else
 	virtual ECADParsingResult Process() override;
 
@@ -152,10 +156,12 @@ private:
 	// Needed to reserve CADFileData
 	// Start with CountUnderModel
 	void CountUnderModel(const A3DAsmModelFile* AsmModel);
+	void CountUnderConfigurationSet(const A3DAsmProductOccurrence* Occurrence);
 	void CountUnderOccurrence(const A3DAsmProductOccurrence* Occurrence);
 	void CountUnderPrototype(const A3DAsmProductOccurrence* Prototype);
+	void CountUnderSubPrototype(const A3DAsmProductOccurrence* Prototype);
 	void CountUnderPartDefinition(const A3DAsmPartDefinition* PartDefinition);
-	void CountUnderRepresentation(const A3DRiRepresentationItem* RepresentationItem);
+	void CountUnderRepresentationItem(const A3DRiRepresentationItem* RepresentationItem);
 	void CountUnderRepresentationSet(const A3DRiSet* RepresentationSet);
 
 	void ReserveCADFileData();
@@ -169,11 +175,11 @@ private:
 	ECADParsingResult TraverseModel(const A3DAsmModelFile* AsmModel);
 	void TraverseReference(const A3DAsmProductOccurrence* Reference);
 	bool IsConfigurationSet(const A3DAsmProductOccurrence* Occurrence);
-	// TODO for SW //void TraverseReferenceFromConfiguration(const A3DAsmProductOccurrence* Reference, FEntityData& ParentMetaData);
-	// TODO for SW //void TraverseConfigurationSet(const A3DAsmProductOccurrence* ConfigurationSet, FEntityData& ParentMetaData);
-	// TODO for SW //void TraverseConfiguration(const A3DAsmProductOccurrence* OccurrConfigurationence, FEntityData& ParentMetaData);
+	void TraverseConfigurationSet(const A3DAsmProductOccurrence* ConfigurationSet);
 	FCadId TraverseOccurrence(const A3DAsmProductOccurrence* Occurrence);
-	void TraversePrototype(const A3DAsmProductOccurrence* InPrototype, FEntityMetaData& OutMetaData, FMatrix& OutPrototypeMatrix);
+	void TraversePrototype(const A3DAsmProductOccurrence* Prototype, FArchiveComponent& Component);
+	void ProcessPrototype(const A3DAsmProductOccurrence* InPrototype, FEntityMetaData& OutMetaData, FMatrix& OutPrototypeMatrix);
+	void ProcessOccurrence(TUniqueTSObj<A3DAsmProductOccurrenceData>& Occurrence, FArchiveComponent& Component);
 	void TraversePartDefinition(const A3DAsmPartDefinition* PartDefinition, FArchiveComponent& Component);
 	FCadId TraverseRepresentationSet(const A3DRiSet* pSet, FEntityMetaData& PartMetaData);
 	FCadId TraverseRepresentationItem(A3DRiRepresentationItem* RepresentationItem, FEntityMetaData& PartMetaData);
@@ -187,8 +193,8 @@ private:
 	void TraverseTessellation3D(const A3DTess3D* Tessellation, FArchiveBody& Body);
 
 	// MetaData
-	void TraverseMetaData(const A3DEntity* Entity, FEntityMetaData& OutMetaData);
-	void TraverseSpecificMetaData(const A3DAsmProductOccurrence* Occurrence, FEntityMetaData& OutMetaData);
+	void ExtractMetaData(const A3DEntity* Entity, FEntityMetaData& OutMetaData);
+	void ExtractSpecificMetaData(const A3DAsmProductOccurrence* Occurrence, FEntityMetaData& OutMetaData);
 
 	void BuildInstanceName(TMap<FString, FString>& MetaData);
 	void BuildReferenceName(TMap<FString, FString>& MetaData);
@@ -205,9 +211,9 @@ private:
 
 	// Transform
 	FMatrix TraverseCoordinateSystem(const A3DRiCoordinateSystem* CoordinateSystem);
-	FMatrix TraverseTransformation(const A3DMiscCartesianTransformation* Transformation3d);
-	FMatrix TraverseGeneralTransformation(const A3DMiscGeneralTransformation* GeneralTransformation);
-	FMatrix TraverseTransformation3D(const A3DMiscCartesianTransformation* CartesianTransformation);
+	FMatrix TraverseTransformation(const A3DMiscTransformation* Transformation3d);
+	FMatrix TraverseGeneralTransformation(const A3DMiscTransformation* GeneralTransformation);
+	FMatrix TraverseTransformation3D(const A3DMiscTransformation* CartesianTransformation);
 
 	FFileDescriptor GetOccurrenceFileName(const A3DAsmProductOccurrence* OccurrencePtr);
 
@@ -233,6 +239,9 @@ private:
 	EModellerType ModellerType;
 	double FileUnit = 1;
 	FCadId LastEntityId = 1;
+
+	// 
+	TMap<const A3DRiRepresentationItem*, FCadId> RepresentationItemsCache;
 };
 
 namespace TechSoftFileParserImpl
@@ -240,7 +249,7 @@ namespace TechSoftFileParserImpl
 // Methodes used by TraverseTessellation3D
 
 typedef double A3DDouble;
-inline bool AddFace(int32 FaceIndex[3], CADLibrary::FTessellationData& Tessellation, int32& VertexIndex)
+inline bool AddFace(int32 FaceIndex[3], CADLibrary::FTessellationData& Tessellation, int32& InOutVertexIndex)
 {
 	if (FaceIndex[0] == FaceIndex[1] || FaceIndex[0] == FaceIndex[2] || FaceIndex[1] == FaceIndex[2])
 	{
@@ -249,8 +258,7 @@ inline bool AddFace(int32 FaceIndex[3], CADLibrary::FTessellationData& Tessellat
 
 	for (int32 Index = 0; Index < 3; ++Index)
 	{
-		FaceIndex[Index] /= 3;
-		Tessellation.VertexIndices.Add(VertexIndex++);
+		Tessellation.VertexIndices.Add(InOutVertexIndex++);
 	}
 	Tessellation.PositionIndices.Append(FaceIndex, 3);
 	return true;
diff --git a/Engine/Plugins/Enterprise/DatasmithCADImporter/Source/CADInterfaces/Private/TechSoftInterface.cpp b/Engine/Plugins/Enterprise/DatasmithCADImporter/Source/CADInterfaces/Private/TechSoftInterface.cpp
--- a/Engine/Plugins/Enterprise/DatasmithCADImporter/Source/CADInterfaces/Private/TechSoftInterface.cpp
+++ b/Engine/Plugins/Enterprise/DatasmithCADImporter/Source/CADInterfaces/Private/TechSoftInterface.cpp
@@ -32,19 +32,28 @@ bool FTechSoftInterface::InitializeKernel(const TCHAR* InEnginePluginsPath)
 	}
 
 	FString EnginePluginsPath(InEnginePluginsPath);
+#ifndef CADKERNEL_DEV
 	if (EnginePluginsPath.IsEmpty())
 	{
 		EnginePluginsPath = FPaths::EnginePluginsDir();
 	}
+#endif
 
+#ifdef CADKERNEL_DEV
+	FString TechSoftDllPath = EnginePluginsPath;
+#else
 	FString TechSoftDllPath = FPaths::Combine(EnginePluginsPath, TEXT("Enterprise/DatasmithCADImporter"), TEXT("Binaries"), FPlatformProcess::GetBinariesSubdirectory(), "TechSoft");
+#endif
+
 	TechSoftDllPath = FPaths::ConvertRelativePathToFull(TechSoftDllPath);
 	ExchangeLoader = MakeUnique<A3DSDKHOOPSExchangeLoader>(*TechSoftDllPath);
 
 	const A3DStatus IRet = ExchangeLoader->m_eSDKStatus;
 	if (IRet != A3D_SUCCESS)
 	{
+#ifndef CADKERNEL_DEV
 		UE_LOG(LogCADInterfaces, Warning, TEXT("Failed to load required library in %s. Plug-in will not be functional."), *TechSoftDllPath);
+#endif
 	}
 	else
 	{
@@ -63,12 +72,22 @@ A3DStatus FTechSoftInterface::Import(const A3DImport& Import)
 	return ExchangeLoader->Import(Import);
 }
 
+A3DStatus FTechSoftInterface::UnloadModel()
+{
+	A3DStatus Ret = A3DStatus::A3D_ERROR;
+	if (ExchangeLoader)
+	{
+		Ret = A3DAsmModelFileDelete(ExchangeLoader->m_psModelFile);
+	}
+	ExchangeLoader->m_psModelFile = nullptr;
+	return Ret;
+}
+
 A3DAsmModelFile* FTechSoftInterface::GetModelFile()
 {
 	return ExchangeLoader->m_psModelFile;
 }
 
-
 void TUniqueTSObj<A3DAsmModelFileData>::InitializeData()
 {
 	A3D_INITIALIZE_DATA(A3DAsmModelFileData, Data);
@@ -99,6 +118,56 @@ void TUniqueTSObj<A3DAsmProductOccurrenceDataUg>::InitializeData()
 	A3D_INITIALIZE_DATA(A3DAsmProductOccurrenceDataUg, Data);
 }
 
+void TUniqueTSObj<A3DBoundingBoxData>::InitializeData()
+{
+	A3D_INITIALIZE_DATA(A3DBoundingBoxData, Data);
+}
+
+void TUniqueTSObj<A3DCrvCircleData>::InitializeData()
+{
+	A3D_INITIALIZE_DATA(A3DCrvCircleData, Data);
+}
+
+void TUniqueTSObj<A3DCrvCompositeData>::InitializeData()
+{
+	A3D_INITIALIZE_DATA(A3DCrvCompositeData, Data);
+}
+
+void TUniqueTSObj<A3DCrvEllipseData>::InitializeData()
+{
+	A3D_INITIALIZE_DATA(A3DCrvEllipseData, Data);
+}
+
+void TUniqueTSObj<A3DCrvHelixData>::InitializeData()
+{
+	A3D_INITIALIZE_DATA(A3DCrvHelixData, Data);
+}
+
+void TUniqueTSObj<A3DCrvHyperbolaData>::InitializeData()
+{
+	A3D_INITIALIZE_DATA(A3DCrvHyperbolaData, Data);
+}
+
+void TUniqueTSObj<A3DCrvLineData>::InitializeData()
+{
+	A3D_INITIALIZE_DATA(A3DCrvLineData, Data);
+}
+
+void TUniqueTSObj<A3DCrvNurbsData>::InitializeData()
+{
+	A3D_INITIALIZE_DATA(A3DCrvNurbsData, Data);
+}
+
+void TUniqueTSObj<A3DCrvParabolaData>::InitializeData()
+{
+	A3D_INITIALIZE_DATA(A3DCrvParabolaData, Data);
+}
+
+void TUniqueTSObj<A3DCrvPolyLineData>::InitializeData()
+{
+	A3D_INITIALIZE_DATA(A3DCrvPolyLineData, Data);
+}
+
 void TUniqueTSObj<A3DGraphicsData>::InitializeData()
 {
 	A3D_INITIALIZE_DATA(A3DGraphicsData, Data);
@@ -119,6 +188,11 @@ void TUniqueTSObjFromIndex<A3DGraphStyleData>::InitializeData()
 	A3D_INITIALIZE_DATA(A3DGraphStyleData, Data);
 }
 
+void TUniqueTSObj<A3DIntervalData>::InitializeData()
+{
+	A3D_INITIALIZE_DATA(A3DIntervalData, Data);
+}
+
 void TUniqueTSObj<A3DMiscAttributeData>::InitializeData()
 {
 	A3D_INITIALIZE_DATA(A3DMiscAttributeData, Data);
@@ -129,6 +203,11 @@ void TUniqueTSObj<A3DMiscCartesianTransformationData>::InitializeData()
 	A3D_INITIALIZE_DATA(A3DMiscCartesianTransformationData, Data);
 }
 
+void TUniqueTSObj<A3DMiscEntityReferenceData>::InitializeData()
+{
+	A3D_INITIALIZE_DATA(A3DMiscEntityReferenceData, Data);
+}
+
 void TUniqueTSObj<A3DMiscGeneralTransformationData>::InitializeData()
 {
 	A3D_INITIALIZE_DATA(A3DMiscGeneralTransformationData, Data);
@@ -139,6 +218,21 @@ void TUniqueTSObj<A3DMiscMaterialPropertiesData>::InitializeData()
 	A3D_INITIALIZE_DATA(A3DMiscMaterialPropertiesData, Data);
 }
 
+void TUniqueTSObj<A3DMiscReferenceOnCsysItemData>::InitializeData()
+{
+	A3D_INITIALIZE_DATA(A3DMiscReferenceOnCsysItemData, Data);
+}
+
+void TUniqueTSObj<A3DMiscReferenceOnTopologyData>::InitializeData()
+{
+	A3D_INITIALIZE_DATA(A3DMiscReferenceOnTopologyData, Data);
+}
+
+void TUniqueTSObj<A3DMiscReferenceOnTessData>::InitializeData()
+{
+	A3D_INITIALIZE_DATA(A3DMiscReferenceOnTessData, Data);
+}
+
 void TUniqueTSObj<A3DRiBrepModelData>::InitializeData()
 {
 	A3D_INITIALIZE_DATA(A3DRiBrepModelData, Data);
@@ -179,6 +273,81 @@ void TUniqueTSObj<A3DRootBaseWithGraphicsData>::InitializeData()
 	A3D_INITIALIZE_DATA(A3DRootBaseWithGraphicsData, Data);
 }
 
+void TUniqueTSObj<A3DSurfBlend01Data>::InitializeData()
+{
+	A3D_INITIALIZE_DATA(A3DSurfBlend01Data, Data);
+}
+
+void TUniqueTSObj<A3DSurfBlend02Data>::InitializeData()
+{
+	A3D_INITIALIZE_DATA(A3DSurfBlend02Data, Data);
+}
+
+void TUniqueTSObj<A3DSurfBlend03Data>::InitializeData()
+{
+	A3D_INITIALIZE_DATA(A3DSurfBlend03Data, Data);
+}
+
+void TUniqueTSObj<A3DSurfConeData>::InitializeData()
+{
+	A3D_INITIALIZE_DATA(A3DSurfConeData, Data);
+}
+
+void TUniqueTSObj<A3DSurfCylinderData>::InitializeData()
+{
+	A3D_INITIALIZE_DATA(A3DSurfCylinderData, Data);
+}
+
+void TUniqueTSObj<A3DSurfCylindricalData>::InitializeData()
+{
+	A3D_INITIALIZE_DATA(A3DSurfCylindricalData, Data);
+}
+
+void TUniqueTSObj<A3DSurfExtrusionData>::InitializeData()
+{
+	A3D_INITIALIZE_DATA(A3DSurfExtrusionData, Data);
+}
+
+void TUniqueTSObj<A3DSurfFromCurvesData>::InitializeData()
+{
+	A3D_INITIALIZE_DATA(A3DSurfFromCurvesData, Data);
+}
+
+void TUniqueTSObj<A3DSurfNurbsData>::InitializeData()
+{
+	A3D_INITIALIZE_DATA(A3DSurfNurbsData, Data);
+}
+
+void TUniqueTSObj<A3DSurfPipeData>::InitializeData()
+{
+	A3D_INITIALIZE_DATA(A3DSurfPipeData, Data);
+}
+
+void TUniqueTSObj<A3DSurfPlaneData>::InitializeData()
+{
+	A3D_INITIALIZE_DATA(A3DSurfPlaneData, Data);
+}
+
+void TUniqueTSObj<A3DSurfRevolutionData>::InitializeData()
+{
+	A3D_INITIALIZE_DATA(A3DSurfRevolutionData, Data);
+}
+
+void TUniqueTSObj<A3DSurfRuledData>::InitializeData()
+{
+	A3D_INITIALIZE_DATA(A3DSurfRuledData, Data);
+}
+
+void TUniqueTSObj<A3DSurfSphereData>::InitializeData()
+{
+	A3D_INITIALIZE_DATA(A3DSurfSphereData, Data);
+}
+
+void TUniqueTSObj<A3DSurfTorusData>::InitializeData()
+{
+	A3D_INITIALIZE_DATA(A3DSurfTorusData, Data);
+}
+
 void TUniqueTSObj<A3DTess3DData>::InitializeData()
 {
 	A3D_INITIALIZE_DATA(A3DTess3DData, Data);
@@ -189,11 +358,31 @@ void TUniqueTSObj<A3DTessBaseData>::InitializeData()
 	A3D_INITIALIZE_DATA(A3DTessBaseData, Data);
 }
 
+void TUniqueTSObj<A3DTopoBodyData>::InitializeData()
+{
+	A3D_INITIALIZE_DATA(A3DTopoBodyData, Data);
+}
+
 void TUniqueTSObj<A3DTopoBrepDataData>::InitializeData()
 {
 	A3D_INITIALIZE_DATA(A3DTopoBrepDataData, Data);
 }
 
+void TUniqueTSObj<A3DTopoCoEdgeData>::InitializeData()
+{
+	A3D_INITIALIZE_DATA(A3DTopoCoEdgeData, Data);
+}
+
+void TUniqueTSObj<A3DTopoConnexData>::InitializeData()
+{
+	A3D_INITIALIZE_DATA(A3DTopoConnexData, Data);
+}
+
+void TUniqueTSObj<A3DTopoContextData>::InitializeData()
+{
+	A3D_INITIALIZE_DATA(A3DTopoContextData, Data);
+}
+
 void TUniqueTSObj<A3DTopoEdgeData>::InitializeData()
 {
 	A3D_INITIALIZE_DATA(A3DTopoEdgeData, Data);
@@ -204,6 +393,16 @@ void TUniqueTSObj<A3DTopoFaceData>::InitializeData()
 	A3D_INITIALIZE_DATA(A3DTopoFaceData, Data);
 }
 
+void TUniqueTSObj<A3DTopoLoopData>::InitializeData()
+{
+	A3D_INITIALIZE_DATA(A3DTopoLoopData, Data);
+}
+
+void TUniqueTSObj<A3DTopoShellData>::InitializeData()
+{
+	A3D_INITIALIZE_DATA(A3DTopoShellData, Data);
+}
+
 void TUniqueTSObj<A3DTopoUniqueVertexData>::InitializeData()
 {
 	A3D_INITIALIZE_DATA(A3DTopoUniqueVertexData, Data);
@@ -244,6 +443,62 @@ A3DStatus TUniqueTSObj<A3DAsmProductOccurrenceDataUg>::GetData(const A3DAsmProdu
 	return A3DAsmProductOccurrenceGetUg(InAsmProductOccurrencePtr, &Data);
 }
 
+A3DStatus TUniqueTSObj<A3DBoundingBoxData>::GetData(const A3DEntity* InGraphicsPtr)
+{
+	return A3DMiscGetBoundingBox(InGraphicsPtr, &Data);
+}
+
+A3DStatus TUniqueTSObj<A3DCrvCircleData>::GetData(const A3DEntity* InCurvePtr)
+{
+	return A3DCrvCircleGet(InCurvePtr, &Data);
+}
+
+A3DStatus TUniqueTSObj<A3DCrvCompositeData>::GetData(const A3DEntity* InCurvePtr)
+{
+	return A3DCrvCompositeGet(InCurvePtr, &Data);
+}
+
+A3DStatus TUniqueTSObj<A3DCrvEllipseData>::GetData(const A3DEntity* InCurvePtr)
+{
+	return A3DCrvEllipseGet(InCurvePtr, &Data);
+}
+
+A3DStatus TUniqueTSObj<A3DCrvHelixData>::GetData(const A3DEntity* InCurvePtr)
+{
+	return A3DCrvHelixGet(InCurvePtr, &Data);
+}
+
+A3DStatus TUniqueTSObj<A3DCrvHyperbolaData>::GetData(const A3DEntity* InCurvePtr)
+{
+	return A3DCrvHyperbolaGet(InCurvePtr, &Data);
+}
+
+A3DStatus TUniqueTSObj<A3DCrvLineData>::GetData(const A3DEntity* InCurvePtr)
+{
+	return A3DCrvLineGet(InCurvePtr, &Data);
+}
+
+A3DStatus TUniqueTSObj<A3DCrvNurbsData>::GetData(const A3DEntity* InCurvePtr)
+{
+	return A3DCrvNurbsGet(InCurvePtr, &Data);
+}
+
+A3DStatus TUniqueTSObj<A3DCrvParabolaData>::GetData(const A3DEntity* InCurvePtr)
+{
+	return A3DCrvParabolaGet(InCurvePtr, &Data);
+}
+
+A3DStatus TUniqueTSObj<A3DCrvPolyLineData>::GetData(const A3DEntity* InCurvePtr)
+{
+	return A3DCrvPolyLineGet(InCurvePtr, &Data);
+}
+
+A3DStatus TUniqueTSObj<A3DIntervalData>::GetData(const A3DEntity* InCurvePtr)
+{
+	return A3DCrvGetInterval(InCurvePtr, &Data);
+}
+
+
 A3DStatus TUniqueTSObj<A3DGraphicsData>::GetData(const A3DGraphics* InGraphicsPtr)
 {
 	return A3DGraphicsGet(InGraphicsPtr, &Data);
@@ -274,6 +529,11 @@ A3DStatus TUniqueTSObj<A3DMiscCartesianTransformationData>::GetData(const A3DMis
 	return A3DMiscCartesianTransformationGet(InMiscCartesianTransformationPtr, &Data);
 }
 
+A3DStatus TUniqueTSObj<A3DMiscEntityReferenceData>::GetData(const A3DMiscEntityReference* InMiscEntityReferencePtr)
+{
+	return A3DMiscEntityReferenceGet(InMiscEntityReferencePtr, &Data);
+}
+
 A3DStatus TUniqueTSObj<A3DMiscGeneralTransformationData>::GetData(const A3DMiscGeneralTransformation* InMiscGeneralTransformationPtr)
 {
 	return A3DMiscGeneralTransformationGet(InMiscGeneralTransformationPtr, &Data);
@@ -284,6 +544,21 @@ A3DStatus TUniqueTSObj<A3DMiscMaterialPropertiesData>::GetData(const A3DEntity*
 	return A3DMiscGetMaterialProperties(InMiscMaterialPropertiesPtr, &Data);
 }
 
+A3DStatus TUniqueTSObj<A3DMiscReferenceOnTopologyData>::GetData(const A3DEntity* InMiscReferenceOnTopologyPtr)
+{
+	return A3DMiscReferenceOnTopologyGet(InMiscReferenceOnTopologyPtr, &Data);
+}
+
+A3DStatus TUniqueTSObj<A3DMiscReferenceOnCsysItemData>::GetData(const A3DEntity* InMiscReferenceOnCsysItem)
+{
+	return A3DMiscReferenceOnCsysItemGet(InMiscReferenceOnCsysItem, &Data);
+}
+
+A3DStatus TUniqueTSObj<A3DMiscReferenceOnTessData>::GetData(const A3DEntity* InMiscReferenceOnTess)
+{
+	return A3DMiscReferenceOnTessGet(InMiscReferenceOnTess, &Data);
+}
+
 A3DStatus TUniqueTSObj<A3DRiBrepModelData>::GetData(const A3DRiBrepModel* InRiBrepModelPtr)
 {
 	return A3DRiBrepModelGet(InRiBrepModelPtr, &Data);
@@ -324,6 +599,81 @@ A3DStatus TUniqueTSObj<A3DRootBaseWithGraphicsData>::GetData(const A3DRootBaseWi
 	return A3DRootBaseWithGraphicsGet(InRootBaseWithGraphicsPtr, &Data);
 }
 
+A3DStatus TUniqueTSObj<A3DSurfBlend01Data>::GetData(const A3DEntity* InSurfPtr)
+{
+	return A3DSurfBlend01Get(InSurfPtr, &Data);
+}
+
+A3DStatus TUniqueTSObj<A3DSurfBlend02Data>::GetData(const A3DEntity* InSurfacePtr)
+{
+	return A3DSurfBlend02Get(InSurfacePtr, &Data);
+}
+
+A3DStatus TUniqueTSObj<A3DSurfBlend03Data>::GetData(const A3DEntity* InSurfacePtr)
+{
+	return A3DSurfBlend03Get(InSurfacePtr, &Data);
+}
+
+A3DStatus TUniqueTSObj<A3DSurfConeData>::GetData(const A3DEntity* InSurfacePtr)
+{
+	return A3DSurfConeGet(InSurfacePtr, &Data);
+}
+
+A3DStatus TUniqueTSObj<A3DSurfCylinderData>::GetData(const A3DEntity* InSurfacePtr)
+{
+	return A3DSurfCylinderGet(InSurfacePtr, &Data);
+}
+
+A3DStatus TUniqueTSObj<A3DSurfCylindricalData>::GetData(const A3DEntity* InSurfacePtr)
+{
+	return A3DSurfCylindricalGet(InSurfacePtr, &Data);
+}
+
+A3DStatus TUniqueTSObj<A3DSurfExtrusionData>::GetData(const A3DEntity* InSurfacePtr)
+{
+	return A3DSurfExtrusionGet(InSurfacePtr, &Data);
+}
+
+A3DStatus TUniqueTSObj<A3DSurfFromCurvesData>::GetData(const A3DEntity* InSurfacePtr)
+{
+	return A3DSurfFromCurvesGet(InSurfacePtr, &Data);
+}
+
+A3DStatus TUniqueTSObj<A3DSurfNurbsData>::GetData(const A3DEntity* InSurfacePtr)
+{
+	return A3DSurfNurbsGet(InSurfacePtr, &Data);
+}
+
+A3DStatus TUniqueTSObj<A3DSurfPipeData>::GetData(const A3DEntity* InSurfacePtr)
+{
+	return A3DSurfPipeGet(InSurfacePtr, &Data);
+}
+
+A3DStatus TUniqueTSObj<A3DSurfPlaneData>::GetData(const A3DEntity* InSurfacePtr)
+{
+	return A3DSurfPlaneGet(InSurfacePtr, &Data);
+}
+
+A3DStatus TUniqueTSObj<A3DSurfRevolutionData>::GetData(const A3DEntity* InSurfacePtr)
+{
+	return A3DSurfRevolutionGet(InSurfacePtr, &Data);
+}
+
+A3DStatus TUniqueTSObj<A3DSurfRuledData>::GetData(const A3DEntity* InSurfacePtr)
+{
+	return A3DSurfRuledGet(InSurfacePtr, &Data);
+}
+
+A3DStatus TUniqueTSObj<A3DSurfSphereData>::GetData(const A3DEntity* InSurfacePtr)
+{
+	return A3DSurfSphereGet(InSurfacePtr, &Data);
+}
+
+A3DStatus TUniqueTSObj<A3DSurfTorusData>::GetData(const A3DEntity* InSurfacePtr)
+{
+	return A3DSurfTorusGet(InSurfacePtr, &Data);
+}
+
 A3DStatus TUniqueTSObj<A3DTess3DData>::GetData(const A3DTess3D* InTess3DPtr)
 {
 	return A3DTess3DGet(InTess3DPtr, &Data);
@@ -334,11 +684,31 @@ A3DStatus TUniqueTSObj<A3DTessBaseData>::GetData(const A3DTessBase* InTessBasePt
 	return A3DTessBaseGet(InTessBasePtr, &Data);
 }
 
+A3DStatus TUniqueTSObj<A3DTopoBodyData>::GetData(const A3DEntity* InBodyPtr)
+{
+	return A3DTopoBodyGet(InBodyPtr, &Data);
+}
+
 A3DStatus TUniqueTSObj<A3DTopoBrepDataData>::GetData(const A3DTopoBrepData* InTopoBrepDataPtr)
 {
 	return A3DTopoBrepDataGet(InTopoBrepDataPtr, &Data);
 }
 
+A3DStatus TUniqueTSObj<A3DTopoCoEdgeData>::GetData(const A3DEntity* InGraphicsPtr)
+{
+	return A3DTopoCoEdgeGet(InGraphicsPtr, &Data);
+}
+
+A3DStatus TUniqueTSObj<A3DTopoConnexData>::GetData(const A3DEntity* InGraphicsPtr)
+{
+	return A3DTopoConnexGet(InGraphicsPtr, &Data);
+}
+
+A3DStatus TUniqueTSObj<A3DTopoContextData>::GetData(const A3DEntity* InGraphicsPtr)
+{
+	return A3DTopoContextGet(InGraphicsPtr, &Data);
+}
+
 A3DStatus TUniqueTSObj<A3DTopoEdgeData>::GetData(const A3DTopoEdge* InTopoEdgePtr)
 {
 	return A3DTopoEdgeGet(InTopoEdgePtr, &Data);
@@ -349,6 +719,16 @@ A3DStatus TUniqueTSObj<A3DTopoFaceData>::GetData(const A3DTopoFace* InTopoFacePt
 	return A3DTopoFaceGet(InTopoFacePtr, &Data);
 }
 
+A3DStatus TUniqueTSObj<A3DTopoLoopData>::GetData(const A3DEntity* InTopoLoopPtr)
+{
+	return A3DTopoLoopGet(InTopoLoopPtr, &Data);
+}
+
+A3DStatus TUniqueTSObj<A3DTopoShellData>::GetData(const A3DEntity* InTopoPtr)
+{
+	return A3DTopoShellGet(InTopoPtr, &Data);
+}
+
 A3DStatus TUniqueTSObj<A3DTopoUniqueVertexData>::GetData(const A3DTopoUniqueVertex* InTopoUniqueVertexPtr)
 {
 	return A3DTopoUniqueVertexGet(InTopoUniqueVertexPtr, &Data);
@@ -364,7 +744,51 @@ uint32 TUniqueTSObjFromIndex<A3DGraphMaterialData>::DefaultValue = A3D_DEFAULT_M
 uint32 TUniqueTSObjFromIndex<A3DGraphPictureData>::DefaultValue = A3D_DEFAULT_MATERIAL_INDEX;
 uint32 TUniqueTSObjFromIndex<A3DGraphRgbColorData>::DefaultValue = A3D_DEFAULT_COLOR_INDEX;
 uint32 TUniqueTSObjFromIndex<A3DGraphStyleData>::DefaultValue = A3D_DEFAULT_STYLE_INDEX;
+#endif
 
+namespace TechSoftUtils
+{
+
+FTechSoftInterface& GetTechSoftInterface()
+{
+	static FTechSoftInterface TechSoftInterface;
+	return TechSoftInterface;
+}
+
+bool TECHSOFT_InitializeKernel(const TCHAR* InEnginePluginsPath)
+{
+	return GetTechSoftInterface().InitializeKernel(InEnginePluginsPath);
+}
+
+#ifdef USE_TECHSOFT_SDK
+
+A3DStatus GetSurfaceAsNurbs(const A3DSurfBase* SurfacePtr, A3DSurfNurbsData* DataPtr, A3DDouble Tolerance, A3DBool bUseSameParameterization)
+{
+	return A3DSurfBaseGetAsNurbs(SurfacePtr, Tolerance, bUseSameParameterization, DataPtr);
+}
+
+A3DStatus GetCurveAsNurbs(const A3DCrvBase* CurvePtr, A3DCrvNurbsData* DataPtr, A3DDouble Tolerance, A3DBool bUseSameParameterization)
+{
+	return A3DCrvBaseGetAsNurbs(CurvePtr, Tolerance, bUseSameParameterization, DataPtr);
+}
+
+A3DStatus GetEntityType(const A3DEntity* EntityPtr, A3DEEntityType* EntityTypePtr)
+{
+	return A3DEntityGetType(EntityPtr, EntityTypePtr);
+}
+
+bool IsEntityBaseWithGraphicsType(const A3DEntity* EntityPtr)
+{
+	return (bool)A3DEntityIsBaseWithGraphicsType(EntityPtr);
+}
+
+bool IsEntityBaseType(const A3DEntity* EntityPtr)
+{
+	return (bool)A3DEntityIsBaseType(EntityPtr);
+}
 #endif
 
 }
+
+
+}
diff --git a/Engine/Plugins/Enterprise/DatasmithCADImporter/Source/CADInterfaces/Private/TechSoftParserUtils.cpp b/Engine/Plugins/Enterprise/DatasmithCADImporter/Source/CADInterfaces/Private/TechSoftParserUtils.cpp
new file mode 100644
--- /dev/null
+++ b/Engine/Plugins/Enterprise/DatasmithCADImporter/Source/CADInterfaces/Private/TechSoftParserUtils.cpp
@@ -0,0 +1,645 @@
+// Copyright Epic Games, Inc. All Rights Reserved.
+
+#include "TechSoftParserUtils.h"
+
+#ifdef USE_TECHSOFT_SDK
+
+namespace TechSoftParserUtils
+{
+	bool FTechSoftMeshConverter::UpdateBodyMesh()
+	{
+		using namespace CADLibrary;
+
+		if (!Tessellation3DData.IsValid() || !TessellationBaseData.IsValid() || TessellationBaseData->m_uiCoordSize == 0)
+		{
+			return false;
+		}
+
+		// Vertex's positions are global to the BodYMesh. Normals and UVs are not.
+		{
+			const int32 VertexCount = TessellationBaseData->m_uiCoordSize / 3;
+			BodyMesh.VertexArray.Reserve(VertexCount);
+
+			double* Coordinates = TessellationBaseData->m_pdCoords;
+			for (unsigned int Index = 0; Index < TessellationBaseData->m_uiCoordSize; ++Index)
+			{
+				Coordinates[Index] *= ScaleFactor;
+			}
+
+			for (unsigned int Index = 0; Index < TessellationBaseData->m_uiCoordSize; Index += 3)
+			{
+				BodyMesh.VertexArray.Emplace(Coordinates[Index], Coordinates[Index + 1], Coordinates[Index + 2]);
+			}
+		}
+
+		BodyMesh.Faces.Reserve(BodyMesh.Faces.Num() + Tessellation3DData->m_uiFaceTessSize);
+		for (unsigned int Index = 0; Index < Tessellation3DData->m_uiFaceTessSize; ++Index)
+		{
+			const A3DTessFaceData& FaceTessData = Tessellation3DData->m_psFaceTessData[Index];
+			if (FaceTessData.m_uiSizesTriangulatedSize > 0)
+			{
+				ConvertFace(FaceTessData, BodyMesh.Faces.Emplace_GetRef());
+			}
+		}
+		
+		return true;
+	}
+
+	void FTechSoftMeshConverter::ConvertFace(const A3DTessFaceData& FaceTessData, CADLibrary::FTessellationData& Tessellation)
+	{
+		//if (FaceTessData.m_uiStyleIndexesSize == 1)
+		//{
+		//	A3DUns32 StyleIndex = FaceTessData.m_puiStyleIndexes[0];
+		//	FCADUUID ColorName;
+		//	FCADUUID MaterialName;
+		//	TraverseGraphStyleData(StyleIndex, ColorName, MaterialName);
+		//	if (ColorName)
+		//	{
+		//		Tessellation.ColorName = ColorName;
+		//		BodyMesh.ColorSet.Add(ColorName);
+		//	}
+
+		//	if (MaterialName)
+		//	{
+		//		Tessellation.MaterialName = MaterialName;
+		//		BodyMesh.MaterialSet.Add(MaterialName);
+		//	}
+		//}
+
+		uint32 TriangleCount = CountTriangles(FaceTessData);
+		Reserve(Tessellation, TriangleCount, /*bWithTexture*/ FaceTessData.m_uiTextureCoordIndexesSize > 0);
+
+		uint32 UsedEntitiesFlags = FaceTessData.m_usUsedEntitiesFlags;
+		const A3DUns32* SizesTriangulated = FaceTessData.m_puiSizesTriangulated;
+
+		uint32 FaceSetIndex = 0;
+
+		LastTriangleIndex = FaceTessData.m_uiStartTriangulated;
+		LastVertexIndex = 0;
+
+		if (UsedEntitiesFlags & kA3DTessFaceDataTriangle)
+		{
+			AddFaceTriangle(Tessellation, SizesTriangulated[FaceSetIndex++]);
+		}
+
+		if (UsedEntitiesFlags & kA3DTessFaceDataTriangleFan)
+		{
+			ensure(FaceTessData.m_uiSizesTriangulatedSize > FaceSetIndex);
+			uint32 FanCount = SizesTriangulated[FaceSetIndex++];
+			for (uint32 Index = 0; Index < FanCount; ++Index)
+			{
+				ensure(FaceTessData.m_uiSizesTriangulatedSize > FaceSetIndex);
+				AddFaceTriangleFan(Tessellation, SizesTriangulated[FaceSetIndex++]);
+			}
+		}
+
+		if (UsedEntitiesFlags & kA3DTessFaceDataTriangleStripe)
+		{
+			ensure(FaceTessData.m_uiSizesTriangulatedSize > FaceSetIndex);
+			uint32 StripeCount = SizesTriangulated[FaceSetIndex++];
+			for (uint32 Index = 0; Index < StripeCount; ++Index)
+			{
+				ensure(FaceTessData.m_uiSizesTriangulatedSize > FaceSetIndex);
+				AddFaceTriangleStripe(Tessellation, SizesTriangulated[FaceSetIndex++]);
+			}
+		}
+
+		if (UsedEntitiesFlags & kA3DTessFaceDataTriangleOneNormal)
+		{
+			ensure(FaceTessData.m_uiSizesTriangulatedSize > FaceSetIndex);
+			AddFaceTriangleWithOneNormal(Tessellation, SizesTriangulated[FaceSetIndex++]);
+		}
+
+		if (UsedEntitiesFlags & kA3DTessFaceDataTriangleFanOneNormal)
+		{
+			ensure(FaceTessData.m_uiSizesTriangulatedSize > FaceSetIndex);
+			uint32 FanCount = SizesTriangulated[FaceSetIndex++];
+			for (uint32 Index = 0; Index < FanCount; ++Index)
+			{
+				ensure(FaceTessData.m_uiSizesTriangulatedSize > FaceSetIndex);
+				AddFaceTriangleFanWithOneNormal(Tessellation, SizesTriangulated[FaceSetIndex++]);
+			}
+		}
+
+		if (UsedEntitiesFlags & kA3DTessFaceDataTriangleStripeOneNormal)
+		{
+			ensure(FaceTessData.m_uiSizesTriangulatedSize > FaceSetIndex);
+			uint32 StripeCount = SizesTriangulated[FaceSetIndex++];
+			for (uint32 Index = 0; Index < StripeCount; ++Index)
+			{
+				ensure(FaceTessData.m_uiSizesTriangulatedSize > FaceSetIndex);
+				AddFaceTriangleStripeWithOneNormal(Tessellation, SizesTriangulated[FaceSetIndex++]);
+			}
+		}
+
+		if (UsedEntitiesFlags & kA3DTessFaceDataTriangleTextured)
+		{
+			ensure(FaceTessData.m_uiSizesTriangulatedSize > FaceSetIndex);
+			AddFaceTriangleWithTexture(Tessellation, SizesTriangulated[FaceSetIndex++], FaceTessData.m_uiTextureCoordIndexesSize);
+		}
+
+		if (UsedEntitiesFlags & kA3DTessFaceDataTriangleFanTextured)
+		{
+			ensure(FaceTessData.m_uiSizesTriangulatedSize > FaceSetIndex);
+			uint32 FanCount = SizesTriangulated[FaceSetIndex++];
+			for (uint32 Index = 0; Index < FanCount; ++Index)
+			{
+				ensure(FaceTessData.m_uiSizesTriangulatedSize > FaceSetIndex);
+				AddFaceTriangleFanWithTexture(Tessellation, SizesTriangulated[FaceSetIndex++], FaceTessData.m_uiTextureCoordIndexesSize);
+			}
+		}
+
+		if (UsedEntitiesFlags & kA3DTessFaceDataTriangleStripeTextured)
+		{
+			ensure(FaceTessData.m_uiSizesTriangulatedSize > FaceSetIndex);
+			uint32 StripeCount = SizesTriangulated[FaceSetIndex++];
+			for (uint32 Index = 0; Index < StripeCount; ++Index)
+			{
+				ensure(FaceTessData.m_uiSizesTriangulatedSize > FaceSetIndex);
+				AddFaceTriangleStripeWithTexture(Tessellation, SizesTriangulated[FaceSetIndex++], FaceTessData.m_uiTextureCoordIndexesSize);
+			}
+		}
+
+		if (UsedEntitiesFlags & kA3DTessFaceDataTriangleOneNormalTextured)
+		{
+			ensure(FaceTessData.m_uiSizesTriangulatedSize > FaceSetIndex);
+			AddFaceTriangleWithOneNormalAndTexture(Tessellation, SizesTriangulated[FaceSetIndex++], FaceTessData.m_uiTextureCoordIndexesSize);
+		}
+
+		if (UsedEntitiesFlags & kA3DTessFaceDataTriangleFanOneNormalTextured)
+		{
+			ensure(FaceTessData.m_uiSizesTriangulatedSize > FaceSetIndex);
+			uint32 FanCount = SizesTriangulated[FaceSetIndex++];
+			for (uint32 Index = 0; Index < FanCount; ++Index)
+			{
+				ensure(FaceTessData.m_uiSizesTriangulatedSize > FaceSetIndex);
+				AddFaceTriangleFanWithOneNormalAndTexture(Tessellation, SizesTriangulated[FaceSetIndex++], FaceTessData.m_uiTextureCoordIndexesSize);
+			}
+		}
+
+		if (UsedEntitiesFlags & kA3DTessFaceDataTriangleStripeOneNormalTextured)
+		{
+			ensure(FaceTessData.m_uiSizesTriangulatedSize > FaceSetIndex);
+			uint32 StripeCount = SizesTriangulated[FaceSetIndex++];
+			for (uint32 Index = 0; Index < StripeCount; ++Index)
+			{
+				ensure(FaceTessData.m_uiSizesTriangulatedSize > FaceSetIndex);
+				AddFaceTriangleStripeWithOneNormalAndTexture(Tessellation, SizesTriangulated[FaceSetIndex++], FaceTessData.m_uiTextureCoordIndexesSize);
+			}
+		}
+	}
+
+	void FTechSoftMeshConverter::AddFaceTriangle(CADLibrary::FTessellationData& Tessellation, A3DUns32 InTriangleCount)
+	{
+		int32 FaceIndex[3] = { 0, 0, 0 };
+		int32 NormalIndex[3] = { 0, 0, 0 };
+
+		// Get Triangles
+		for (uint32 TriangleIndex = 0; TriangleIndex < InTriangleCount; TriangleIndex++)
+		{
+			NormalIndex[0] = Tessellation3DData->m_puiTriangulatedIndexes[LastVertexIndex++];
+			FaceIndex[0] = Tessellation3DData->m_puiTriangulatedIndexes[LastVertexIndex++] / 3;
+			NormalIndex[1] = Tessellation3DData->m_puiTriangulatedIndexes[LastVertexIndex++];
+			FaceIndex[1] = Tessellation3DData->m_puiTriangulatedIndexes[LastVertexIndex++] / 3;
+			NormalIndex[2] = Tessellation3DData->m_puiTriangulatedIndexes[LastVertexIndex++];
+			FaceIndex[2] = Tessellation3DData->m_puiTriangulatedIndexes[LastVertexIndex++] / 3;
+
+			if (AddFace(FaceIndex, Tessellation, LastVertexIndex))
+			{
+				AddNormals(Tessellation3DData->m_pdNormals, NormalIndex, Tessellation.NormalArray);
+			}
+		}
+	}
+
+	void FTechSoftMeshConverter::AddFaceTriangleFan(CADLibrary::FTessellationData& Tessellation, A3DUns32 InTriangleCount)
+	{
+		int32 FaceIndex[3] = { 0, 0, 0 };
+		int32 NormalIndex[3] = { 0, 0, 0 };
+
+		NormalIndex[0] = Tessellation3DData->m_puiTriangulatedIndexes[LastTriangleIndex++];
+		FaceIndex[0] = Tessellation3DData->m_puiTriangulatedIndexes[LastTriangleIndex++] / 3;
+		NormalIndex[1] = Tessellation3DData->m_puiTriangulatedIndexes[LastTriangleIndex++];
+		FaceIndex[1] = Tessellation3DData->m_puiTriangulatedIndexes[LastTriangleIndex++] / 3;
+
+		for (unsigned long TriangleIndex = 2; TriangleIndex < InTriangleCount; TriangleIndex++)
+		{
+			NormalIndex[2] = Tessellation3DData->m_puiTriangulatedIndexes[LastTriangleIndex++];
+			FaceIndex[2] = Tessellation3DData->m_puiTriangulatedIndexes[LastTriangleIndex++] / 3;
+
+			if (AddFace(FaceIndex, Tessellation, LastVertexIndex))
+			{
+				AddNormals(Tessellation3DData->m_pdNormals, NormalIndex, Tessellation.NormalArray);
+			}
+
+			NormalIndex[1] = NormalIndex[2];
+			FaceIndex[1] = FaceIndex[2];
+		}
+	}
+
+	void FTechSoftMeshConverter::AddFaceTriangleStripe(CADLibrary::FTessellationData& Tessellation, A3DUns32 InTriangleCount)
+	{
+		int32 FaceIndex[3] = { 0, 0, 0 };
+		int32 NormalIndex[3] = { 0, 0, 0 };
+
+		NormalIndex[0] = Tessellation3DData->m_puiTriangulatedIndexes[LastTriangleIndex++];
+		FaceIndex[0] = Tessellation3DData->m_puiTriangulatedIndexes[LastTriangleIndex++] / 3;
+		NormalIndex[1] = Tessellation3DData->m_puiTriangulatedIndexes[LastTriangleIndex++];
+		FaceIndex[1] = Tessellation3DData->m_puiTriangulatedIndexes[LastTriangleIndex++] / 3;
+
+		for (unsigned long TriangleIndex = 2; TriangleIndex < InTriangleCount; TriangleIndex++)
+		{
+			NormalIndex[2] = Tessellation3DData->m_puiTriangulatedIndexes[LastTriangleIndex++];
+			FaceIndex[2] = Tessellation3DData->m_puiTriangulatedIndexes[LastTriangleIndex++] / 3;
+
+			if (AddFace(FaceIndex, Tessellation, LastVertexIndex))
+			{
+				AddNormals(Tessellation3DData->m_pdNormals, NormalIndex, Tessellation.NormalArray);
+			}
+
+			TriangleIndex++;
+			if (TriangleIndex == InTriangleCount)
+			{
+				break;
+			}
+
+			Swap(FaceIndex[1], FaceIndex[2]);
+			Swap(NormalIndex[1], NormalIndex[2]);
+
+			NormalIndex[0] = Tessellation3DData->m_puiTriangulatedIndexes[LastTriangleIndex++];
+			FaceIndex[0] = Tessellation3DData->m_puiTriangulatedIndexes[LastTriangleIndex++] / 3;
+
+			if (AddFace(FaceIndex, Tessellation, LastVertexIndex))
+			{
+				AddNormals(Tessellation3DData->m_pdNormals, NormalIndex, Tessellation.NormalArray);
+			}
+
+			Swap(FaceIndex[0], FaceIndex[1]);
+			Swap(NormalIndex[0], NormalIndex[1]);
+		}
+	}
+
+	void FTechSoftMeshConverter::AddFaceTriangleWithOneNormal(CADLibrary::FTessellationData& Tessellation, A3DUns32 InTriangleCount)
+	{
+		int32 FaceIndex[3] = { 0, 0, 0 };
+		int32 NormalIndex[3] = { 0, 0, 0 };
+
+		// Get Triangles
+		for (uint32 TriangleIndex = 0; TriangleIndex < InTriangleCount; TriangleIndex++)
+		{
+			NormalIndex[0] = Tessellation3DData->m_puiTriangulatedIndexes[LastTriangleIndex++];
+			NormalIndex[1] = NormalIndex[0];
+			NormalIndex[2] = NormalIndex[0];
+
+			FaceIndex[0] = Tessellation3DData->m_puiTriangulatedIndexes[LastTriangleIndex++] / 3;
+			FaceIndex[1] = Tessellation3DData->m_puiTriangulatedIndexes[LastTriangleIndex++] / 3;
+			FaceIndex[2] = Tessellation3DData->m_puiTriangulatedIndexes[LastTriangleIndex++] / 3;
+
+			if (!AddFace(FaceIndex, Tessellation, LastVertexIndex))
+			{
+				continue;
+			}
+
+			AddNormals(Tessellation3DData->m_pdNormals, NormalIndex, Tessellation.NormalArray);
+		}
+	}
+
+	void FTechSoftMeshConverter::AddFaceTriangleFanWithOneNormal(CADLibrary::FTessellationData& Tessellation, A3DUns32 InTriangleCount)
+	{
+		int32 FaceIndex[3] = { 0, 0, 0 };
+		int32 NormalIndex[3] = { 0, 0, 0 };
+
+		NormalIndex[0] = Tessellation3DData->m_puiTriangulatedIndexes[LastTriangleIndex++];
+		NormalIndex[1] = NormalIndex[0];
+		NormalIndex[2] = NormalIndex[0];
+
+		FaceIndex[0] = Tessellation3DData->m_puiTriangulatedIndexes[LastTriangleIndex++];
+		FaceIndex[1] = Tessellation3DData->m_puiTriangulatedIndexes[LastTriangleIndex++];
+
+		// Get Triangles
+		for (uint32 TriangleIndex = 0; TriangleIndex < InTriangleCount; TriangleIndex++)
+		{
+			FaceIndex[2] = Tessellation3DData->m_puiTriangulatedIndexes[LastTriangleIndex++];
+
+			if (AddFace(FaceIndex, Tessellation, LastVertexIndex))
+			{
+				AddNormals(Tessellation3DData->m_pdNormals, NormalIndex, Tessellation.NormalArray);
+			}
+		}
+	}
+
+	void FTechSoftMeshConverter::AddFaceTriangleStripeWithOneNormal(CADLibrary::FTessellationData& Tessellation, A3DUns32 InTriangleCount)
+	{
+		int32 FaceIndex[3] = { 0, 0, 0 };
+		int32 NormalIndex[3] = { 0, 0, 0 };
+
+		NormalIndex[0] = Tessellation3DData->m_puiTriangulatedIndexes[LastTriangleIndex++];
+		NormalIndex[1] = NormalIndex[0];
+		NormalIndex[2] = NormalIndex[0];
+
+		FaceIndex[0] = Tessellation3DData->m_puiTriangulatedIndexes[LastTriangleIndex++] / 3;
+		FaceIndex[1] = Tessellation3DData->m_puiTriangulatedIndexes[LastTriangleIndex++] / 3;
+
+		for (unsigned long TriangleIndex = 0; TriangleIndex < InTriangleCount; TriangleIndex++)
+		{
+			FaceIndex[2] = Tessellation3DData->m_puiTriangulatedIndexes[LastTriangleIndex++] / 3;
+
+			if (AddFace(FaceIndex, Tessellation, LastVertexIndex))
+			{
+				AddNormals(Tessellation3DData->m_pdNormals, NormalIndex, Tessellation.NormalArray);
+			}
+
+			TriangleIndex++;
+			if (TriangleIndex == InTriangleCount)
+			{
+				break;
+			}
+
+			Swap(FaceIndex[1], FaceIndex[2]);
+
+			NormalIndex[0] = Tessellation3DData->m_puiTriangulatedIndexes[LastTriangleIndex++];
+			FaceIndex[0] = Tessellation3DData->m_puiTriangulatedIndexes[LastTriangleIndex++] / 3;
+
+			if (AddFace(FaceIndex, Tessellation, LastVertexIndex))
+			{
+				AddNormals(Tessellation3DData->m_pdNormals, NormalIndex, Tessellation.NormalArray);
+			}
+
+			Swap(FaceIndex[0], FaceIndex[1]);
+		}
+	}
+
+	void FTechSoftMeshConverter::AddFaceTriangleWithTexture(CADLibrary::FTessellationData& Tessellation, A3DUns32 InTriangleCount, A3DUns32 InTextureCount)
+	{
+		int32 FaceIndex[3] = { 0, 0, 0 };
+		int32 NormalIndex[3] = { 0, 0, 0 };
+		int32 TextureIndex[3] = { 0, 0, 0 };
+
+		// Get Triangles
+		for (uint64 TriangleIndex = 0; TriangleIndex < InTriangleCount; TriangleIndex++)
+		{
+			NormalIndex[0] = Tessellation3DData->m_puiTriangulatedIndexes[LastVertexIndex++];
+			TextureIndex[0] = Tessellation3DData->m_puiTriangulatedIndexes[LastVertexIndex];
+			LastVertexIndex += InTextureCount;
+			FaceIndex[0] = Tessellation3DData->m_puiTriangulatedIndexes[LastVertexIndex++] / 3;
+			NormalIndex[1] = Tessellation3DData->m_puiTriangulatedIndexes[LastVertexIndex++];
+			TextureIndex[1] = Tessellation3DData->m_puiTriangulatedIndexes[LastVertexIndex];
+			LastVertexIndex += InTextureCount;
+			FaceIndex[1] = Tessellation3DData->m_puiTriangulatedIndexes[LastVertexIndex++] / 3;
+			NormalIndex[2] = Tessellation3DData->m_puiTriangulatedIndexes[LastVertexIndex++];
+			TextureIndex[2] = Tessellation3DData->m_puiTriangulatedIndexes[LastVertexIndex];
+			LastVertexIndex += InTextureCount;
+			FaceIndex[2] = Tessellation3DData->m_puiTriangulatedIndexes[LastVertexIndex++] / 3;
+
+			if (AddFace(FaceIndex, Tessellation, LastVertexIndex))
+			{
+				AddNormals(Tessellation3DData->m_pdNormals, NormalIndex, Tessellation.NormalArray);
+				AddTextureCoordinates(Tessellation3DData->m_pdTextureCoords, TextureIndex, Tessellation.TexCoordArray);
+			}
+		}
+	}
+
+	void FTechSoftMeshConverter::AddFaceTriangleFanWithTexture(CADLibrary::FTessellationData& Tessellation, A3DUns32 InTriangleCount, A3DUns32 TextureCount)
+	{
+		int32 FaceIndex[3] = { 0, 0, 0 };
+		int32 NormalIndex[3] = { 0, 0, 0 };
+		int32 TextureIndex[3] = { 0, 0, 0 };
+
+		NormalIndex[0] = Tessellation3DData->m_puiTriangulatedIndexes[LastTriangleIndex++];
+		TextureIndex[0] = Tessellation3DData->m_puiTriangulatedIndexes[LastTriangleIndex];
+		LastTriangleIndex += TextureCount;
+		FaceIndex[0] = Tessellation3DData->m_puiTriangulatedIndexes[LastTriangleIndex++];
+
+		NormalIndex[1] = Tessellation3DData->m_puiTriangulatedIndexes[LastTriangleIndex++];
+		TextureIndex[1] = Tessellation3DData->m_puiTriangulatedIndexes[LastTriangleIndex];
+		LastTriangleIndex += TextureCount;
+		FaceIndex[1] = Tessellation3DData->m_puiTriangulatedIndexes[LastTriangleIndex++];
+
+		for (uint32 TriangleIndex = 2; TriangleIndex < InTriangleCount; TriangleIndex++)
+		{
+			NormalIndex[2] = Tessellation3DData->m_puiTriangulatedIndexes[LastTriangleIndex++];
+			TextureIndex[2] = Tessellation3DData->m_puiTriangulatedIndexes[LastTriangleIndex];
+			LastTriangleIndex += TextureCount;
+			FaceIndex[2] = Tessellation3DData->m_puiTriangulatedIndexes[LastTriangleIndex++];
+
+			if (AddFace(FaceIndex, Tessellation, LastVertexIndex))
+			{
+				AddNormals(Tessellation3DData->m_pdNormals, NormalIndex, Tessellation.NormalArray);
+				AddTextureCoordinates(Tessellation3DData->m_pdTextureCoords, TextureIndex, Tessellation.TexCoordArray);
+			}
+		}
+	}
+
+	void FTechSoftMeshConverter::AddFaceTriangleStripeWithTexture(CADLibrary::FTessellationData& Tessellation, A3DUns32 InTriangleCount, A3DUns32 TextureCount)
+	{
+		int32 FaceIndex[3] = { 0, 0, 0 };
+		int32 NormalIndex[3] = { 0, 0, 0 };
+		int32 TextureIndex[3] = { 0, 0, 0 };
+
+		NormalIndex[0] = Tessellation3DData->m_puiTriangulatedIndexes[LastTriangleIndex++];
+		TextureIndex[0] = Tessellation3DData->m_puiTriangulatedIndexes[LastTriangleIndex];
+		LastTriangleIndex += TextureCount;
+		FaceIndex[0] = Tessellation3DData->m_puiTriangulatedIndexes[LastTriangleIndex++];
+		NormalIndex[1] = Tessellation3DData->m_puiTriangulatedIndexes[LastTriangleIndex++];
+		TextureIndex[1] = Tessellation3DData->m_puiTriangulatedIndexes[LastTriangleIndex];
+		LastTriangleIndex += TextureCount;
+		FaceIndex[1] = Tessellation3DData->m_puiTriangulatedIndexes[LastTriangleIndex++];
+
+		for (unsigned long TriangleIndex = 0; TriangleIndex < InTriangleCount; TriangleIndex++)
+		{
+			NormalIndex[2] = Tessellation3DData->m_puiTriangulatedIndexes[LastTriangleIndex++];
+			TextureIndex[2] = Tessellation3DData->m_puiTriangulatedIndexes[LastTriangleIndex];
+			LastTriangleIndex += TextureCount;
+			FaceIndex[2] = Tessellation3DData->m_puiTriangulatedIndexes[LastTriangleIndex++];
+
+			if (AddFace(FaceIndex, Tessellation, LastVertexIndex))
+			{
+				AddNormals(Tessellation3DData->m_pdNormals, NormalIndex, Tessellation.NormalArray);
+				AddTextureCoordinates(Tessellation3DData->m_pdTextureCoords, TextureIndex, Tessellation.TexCoordArray);
+			}
+
+			TriangleIndex++;
+			if (TriangleIndex == InTriangleCount)
+			{
+				break;
+			}
+
+			Swap(FaceIndex[1], FaceIndex[2]);
+			Swap(NormalIndex[1], NormalIndex[2]);
+			Swap(TextureIndex[1], TextureIndex[2]);
+
+			NormalIndex[0] = Tessellation3DData->m_puiTriangulatedIndexes[LastTriangleIndex++];
+			FaceIndex[0] = Tessellation3DData->m_puiTriangulatedIndexes[LastTriangleIndex++] / 3;
+
+			if (AddFace(FaceIndex, Tessellation, LastVertexIndex))
+			{
+				AddNormals(Tessellation3DData->m_pdNormals, NormalIndex, Tessellation.NormalArray);
+			}
+
+			Swap(FaceIndex[0], FaceIndex[1]);
+			Swap(NormalIndex[0], NormalIndex[1]);
+			Swap(TextureIndex[0], TextureIndex[1]);
+		}
+	}
+
+	void FTechSoftMeshConverter::AddFaceTriangleWithOneNormalAndTexture(CADLibrary::FTessellationData& Tessellation, A3DUns32 InTriangleCount, A3DUns32 TextureCount)
+	{
+		int32 FaceIndex[3] = { 0, 0, 0 };
+		int32 NormalIndex[3] = { 0, 0, 0 };
+		int32 TextureIndex[3] = { 0, 0, 0 };
+
+		// Get Triangles
+		for (uint32 TriangleIndex = 0; TriangleIndex < InTriangleCount; TriangleIndex++)
+		{
+			NormalIndex[0] = Tessellation3DData->m_puiTriangulatedIndexes[LastVertexIndex++];
+			NormalIndex[1] = NormalIndex[0];
+			NormalIndex[2] = NormalIndex[0];
+
+			TextureIndex[0] = Tessellation3DData->m_puiTriangulatedIndexes[LastVertexIndex];
+			LastVertexIndex += TextureCount;
+			FaceIndex[0] = Tessellation3DData->m_puiTriangulatedIndexes[LastVertexIndex++] / 3;
+			TextureIndex[1] = Tessellation3DData->m_puiTriangulatedIndexes[LastVertexIndex];
+			LastVertexIndex += TextureCount;
+			FaceIndex[1] = Tessellation3DData->m_puiTriangulatedIndexes[LastVertexIndex++] / 3;
+			TextureIndex[2] = Tessellation3DData->m_puiTriangulatedIndexes[LastVertexIndex];
+			LastVertexIndex += TextureCount;
+			FaceIndex[2] = Tessellation3DData->m_puiTriangulatedIndexes[LastVertexIndex++] / 3;
+
+			if (!AddFace(FaceIndex, Tessellation, LastVertexIndex))
+			{
+				continue;
+			}
+
+			AddNormals(Tessellation3DData->m_pdNormals, NormalIndex, Tessellation.NormalArray);
+			AddTextureCoordinates(Tessellation3DData->m_pdTextureCoords, TextureIndex, Tessellation.TexCoordArray);
+		}
+	}
+
+	void FTechSoftMeshConverter::AddFaceTriangleFanWithOneNormalAndTexture(CADLibrary::FTessellationData& Tessellation, A3DUns32 InTriangleCount, A3DUns32 TextureCount)
+	{
+		int32 FaceIndex[3] = { 0, 0, 0 };
+		int32 NormalIndex[3] = { 0, 0, 0 };
+		int32 TextureIndex[3] = { 0, 0, 0 };
+
+		NormalIndex[0] = Tessellation3DData->m_puiTriangulatedIndexes[LastTriangleIndex++];
+		NormalIndex[1] = NormalIndex[0];
+		NormalIndex[2] = NormalIndex[0];
+
+		TextureIndex[0] = Tessellation3DData->m_puiTriangulatedIndexes[LastTriangleIndex];
+		LastTriangleIndex += TextureCount;
+		FaceIndex[0] = Tessellation3DData->m_puiTriangulatedIndexes[LastTriangleIndex++];
+
+		TextureIndex[1] = Tessellation3DData->m_puiTriangulatedIndexes[LastTriangleIndex];
+		LastTriangleIndex += TextureCount;
+		FaceIndex[1] = Tessellation3DData->m_puiTriangulatedIndexes[LastTriangleIndex++];
+
+		for (uint32 TriangleIndex = 2; TriangleIndex < InTriangleCount; TriangleIndex++)
+		{
+			TextureIndex[2] = Tessellation3DData->m_puiTriangulatedIndexes[LastTriangleIndex];
+			LastTriangleIndex += TextureCount;
+			FaceIndex[2] = Tessellation3DData->m_puiTriangulatedIndexes[LastTriangleIndex++];
+
+			if (AddFace(FaceIndex, Tessellation, LastVertexIndex))
+			{
+				AddNormals(Tessellation3DData->m_pdNormals, NormalIndex, Tessellation.NormalArray);
+				AddTextureCoordinates(Tessellation3DData->m_pdTextureCoords, TextureIndex, Tessellation.TexCoordArray);
+			}
+		}
+	}
+
+	void FTechSoftMeshConverter::AddFaceTriangleStripeWithOneNormalAndTexture(CADLibrary::FTessellationData& Tessellation, A3DUns32 InTriangleCount, A3DUns32 TextureCount)
+	{
+		int32 FaceIndex[3] = { 0, 0, 0 };
+		int32 NormalIndex[3] = { 0, 0, 0 };
+		int32 TextureIndex[3] = { 0, 0, 0 };
+
+		NormalIndex[0] = Tessellation3DData->m_puiTriangulatedIndexes[LastTriangleIndex++];
+		NormalIndex[1] = NormalIndex[0];
+		NormalIndex[2] = NormalIndex[0];
+
+		TextureIndex[0] = Tessellation3DData->m_puiTriangulatedIndexes[LastTriangleIndex];
+		LastTriangleIndex += TextureCount;
+		FaceIndex[0] = Tessellation3DData->m_puiTriangulatedIndexes[LastTriangleIndex++];
+
+		TextureIndex[1] = Tessellation3DData->m_puiTriangulatedIndexes[LastTriangleIndex];
+		LastTriangleIndex += TextureCount;
+		FaceIndex[1] = Tessellation3DData->m_puiTriangulatedIndexes[LastTriangleIndex++];
+
+		for (unsigned long TriangleIndex = 0; TriangleIndex < InTriangleCount; TriangleIndex++)
+		{
+			TextureIndex[2] = Tessellation3DData->m_puiTriangulatedIndexes[LastTriangleIndex];
+			LastTriangleIndex += TextureCount;
+			FaceIndex[2] = Tessellation3DData->m_puiTriangulatedIndexes[LastTriangleIndex++];
+
+			if (AddFace(FaceIndex, Tessellation, LastVertexIndex))
+			{
+				AddNormals(Tessellation3DData->m_pdNormals, NormalIndex, Tessellation.NormalArray);
+				AddTextureCoordinates(Tessellation3DData->m_pdTextureCoords, TextureIndex, Tessellation.TexCoordArray);
+			}
+
+			TriangleIndex++;
+			if (TriangleIndex == InTriangleCount)
+			{
+				break;
+			}
+
+			Swap(FaceIndex[1], FaceIndex[2]);
+			Swap(TextureIndex[1], TextureIndex[2]);
+
+			FaceIndex[0] = Tessellation3DData->m_puiTriangulatedIndexes[LastTriangleIndex++] / 3;
+
+			if (AddFace(FaceIndex, Tessellation, LastVertexIndex))
+			{
+				AddNormals(Tessellation3DData->m_pdNormals, NormalIndex, Tessellation.NormalArray);
+			}
+
+			Swap(FaceIndex[0], FaceIndex[1]);
+			Swap(TextureIndex[0], TextureIndex[1]);
+		}
+	}
+
+	uint32 CountTriangles(const A3DTessFaceData& FaceTessData)
+	{
+		// To revisit
+		//const int32 TessellationFaceDataWithTriangle = 0x2222;
+		//const int32 TessellationFaceDataWithFan = 0x4444;
+		//const int32 TessellationFaceDataWithStrip = 0x8888;
+		//const int32 TessellationFaceDataWithOneNormal = 0xE0E0;
+
+		uint32 UsedEntitiesFlags = FaceTessData.m_usUsedEntitiesFlags;
+		const A3DUns32* SizesTriangulated = FaceTessData.m_puiSizesTriangulated;
+
+		uint32 TriangleCount = 0;
+		uint32 FaceSetIndex = 0;
+		if (UsedEntitiesFlags & kA3DTessFaceDataTriangle)
+		{
+			ensure(FaceTessData.m_uiSizesTriangulatedSize > FaceSetIndex);
+			TriangleCount += SizesTriangulated[FaceSetIndex++];
+		}
+
+		if (UsedEntitiesFlags & kA3DTessFaceDataTriangleFan)
+		{
+			ensure(FaceTessData.m_uiSizesTriangulatedSize > FaceSetIndex);
+			uint32 LastFanIndex = 1 + FaceSetIndex + SizesTriangulated[FaceSetIndex++];
+			for (; FaceSetIndex < LastFanIndex; FaceSetIndex++)
+			{
+				uint32 FanSize = (FaceTessData.m_puiSizesTriangulated[FaceSetIndex] & kA3DTessFaceDataNormalMask);
+				TriangleCount += (FanSize - 2);
+			}
+		}
+
+		if (FaceTessData.m_uiSizesTriangulatedSize > FaceSetIndex)
+		{
+			FaceSetIndex++;
+			for (; FaceSetIndex < FaceTessData.m_uiSizesTriangulatedSize; FaceSetIndex++)
+			{
+				uint32 StripSize = (FaceTessData.m_puiSizesTriangulated[FaceSetIndex] & kA3DTessFaceDataNormalMask);
+				TriangleCount += (StripSize - 2);
+			}
+		}
+
+		return TriangleCount;
+	}
+
+} // ns TechSoftParserUtils
+
+#endif // USE_TECHSOFT_SDK
\ No newline at end of file
diff --git a/Engine/Plugins/Enterprise/DatasmithCADImporter/Source/CADInterfaces/Private/TechSoftParserUtils.h b/Engine/Plugins/Enterprise/DatasmithCADImporter/Source/CADInterfaces/Private/TechSoftParserUtils.h
new file mode 100644
--- /dev/null
+++ b/Engine/Plugins/Enterprise/DatasmithCADImporter/Source/CADInterfaces/Private/TechSoftParserUtils.h
@@ -0,0 +1,139 @@
+// Copyright Epic Games, Inc. All Rights Reserved.
+#pragma once
+
+#include "CoreMinimal.h"
+
+#include "TechSoftInterface.h"
+#include "TUniqueTechSoftObj.h"
+
+#include "CADData.h"
+#include "CADFileData.h"
+#include "CADFileParser.h"
+#include "CADOptions.h"
+#include "CADSceneGraph.h"
+
+typedef void A3DAsmModelFile;
+typedef void A3DAsmPartDefinition;
+typedef void A3DAsmProductOccurrence;
+typedef void A3DEntity;
+typedef void A3DGraphics;
+typedef void A3DMiscAttribute;
+typedef void A3DMiscCartesianTransformation;
+typedef void A3DMiscCartesianTransformation;
+typedef void A3DMiscGeneralTransformation;
+typedef void A3DRiBrepModel;
+typedef void A3DRiCoordinateSystem;
+typedef void A3DRiPolyBrepModel;
+typedef void A3DRiRepresentationItem;
+typedef void A3DRiSet;
+typedef void A3DTess3D;
+typedef void A3DTessBase;
+typedef void A3DTopoBrepData;
+
+#ifdef USE_TECHSOFT_SDK
+
+namespace TechSoftParserUtils
+{
+	typedef double A3DDouble;
+
+	inline bool AddFace(int32 FaceIndex[3], CADLibrary::FTessellationData& Tessellation, int32& InOutVertexIndex)
+	{
+		if (FaceIndex[0] == FaceIndex[1] || FaceIndex[0] == FaceIndex[2] || FaceIndex[1] == FaceIndex[2])
+		{
+			return false;
+		}
+
+		for (int32 Index = 0; Index < 3; ++Index)
+		{
+			Tessellation.VertexIndices.Add(InOutVertexIndex++);
+		}
+		Tessellation.PositionIndices.Append(FaceIndex, 3);
+		return true;
+	};
+
+	inline void AddNormals(const A3DDouble* Normals, const int32 Indices[3], TArray<FVector>& TessellationNormals)
+	{
+		for (int32 Index = 0; Index < 3; ++Index)
+		{
+			int32 NormalIndex = Indices[Index];
+			TessellationNormals.Emplace(Normals[NormalIndex], Normals[NormalIndex + 1], Normals[NormalIndex + 2]);
+		}
+	};
+
+	inline void AddTextureCoordinates(const A3DDouble* TextureCoords, const int32 Indices[3], TArray<FVector2D>& TessellationTextures)
+	{
+		for (int32 Index = 0; Index < 3; ++Index)
+		{
+			int32 TextureIndex = Indices[Index];
+			TessellationTextures.Emplace(TextureCoords[TextureIndex], TextureCoords[TextureIndex + 1]);
+		}
+	};
+
+	inline void Reserve(CADLibrary::FTessellationData& Tessellation, int32 InTrinangleCount, bool bWithTexture)
+	{
+		Tessellation.PositionIndices.Reserve(3 * InTrinangleCount);
+		Tessellation.VertexIndices.Reserve(3 * InTrinangleCount);
+		Tessellation.NormalArray.Reserve(3 * InTrinangleCount);
+		if (bWithTexture)
+		{
+			Tessellation.TexCoordArray.Reserve(3 * InTrinangleCount);
+		}
+	};
+
+	uint32 CountTriangles(const A3DTessFaceData& FaceTessData);
+
+	class FTechSoftMeshConverter
+	{
+	public:
+		FTechSoftMeshConverter(const A3DTess3D* TessellationPtr, double ScaleFactor, CADLibrary::FBodyMesh& InBodyMesh)
+			: BodyMesh(InBodyMesh)
+			, TessellationBaseData(TessellationPtr)
+			, Tessellation3DData(TessellationPtr)
+		{
+		}
+
+		bool UpdateBodyMesh();
+
+	private:
+		void ConvertFace(const A3DTessFaceData& FaceTessData, CADLibrary::FTessellationData& Tessellation);
+
+		void AddFaceTriangle(CADLibrary::FTessellationData& Tessellation, A3DUns32 InTriangleCount);
+
+		void AddFaceTriangleFan(CADLibrary::FTessellationData& Tessellation, A3DUns32 InTriangleCount);
+
+		void AddFaceTriangleStripe(CADLibrary::FTessellationData& Tessellation, A3DUns32 InTriangleCount);
+
+		void AddFaceTriangleWithOneNormal(CADLibrary::FTessellationData& Tessellation, A3DUns32 InTriangleCount);
+
+		void AddFaceTriangleFanWithOneNormal(CADLibrary::FTessellationData& Tessellation, A3DUns32 InTriangleCount);
+
+		void AddFaceTriangleStripeWithOneNormal(CADLibrary::FTessellationData& Tessellation, A3DUns32 InTriangleCount);
+
+		void AddFaceTriangleWithTexture(CADLibrary::FTessellationData& Tessellation, A3DUns32 InTriangleCount, A3DUns32 InTextureCount);
+
+		void AddFaceTriangleFanWithTexture(CADLibrary::FTessellationData& Tessellation, A3DUns32 InTriangleCount, A3DUns32 InTextureCount);
+
+		void AddFaceTriangleStripeWithTexture(CADLibrary::FTessellationData& Tessellation, A3DUns32 InTriangleCount, A3DUns32 InTextureCount);
+
+		void AddFaceTriangleWithOneNormalAndTexture(CADLibrary::FTessellationData& Tessellation, A3DUns32 InTriangleCount, A3DUns32 InTextureCount);
+
+		void AddFaceTriangleFanWithOneNormalAndTexture(CADLibrary::FTessellationData& Tessellation, A3DUns32 InTriangleCount, A3DUns32 InTextureCount);
+
+		void AddFaceTriangleStripeWithOneNormalAndTexture(CADLibrary::FTessellationData& Tessellation, A3DUns32 InTriangleCount, A3DUns32 InTextureCount);
+
+	private:
+		CADLibrary::FBodyMesh& BodyMesh;
+
+		CADLibrary::TUniqueTSObj<A3DTessBaseData> TessellationBaseData;
+
+		CADLibrary::TUniqueTSObj<A3DTess3DData> Tessellation3DData;
+
+		double ScaleFactor;
+
+		uint32 LastTriangleIndex;
+
+		int32 LastVertexIndex;
+	};
+} // ns TechSoftParserUtils
+
+#endif
diff --git a/Engine/Plugins/Enterprise/DatasmithCADImporter/Source/CADInterfaces/Public/CADFileData.h b/Engine/Plugins/Enterprise/DatasmithCADImporter/Source/CADInterfaces/Public/CADFileData.h
--- a/Engine/Plugins/Enterprise/DatasmithCADImporter/Source/CADInterfaces/Public/CADFileData.h
+++ b/Engine/Plugins/Enterprise/DatasmithCADImporter/Source/CADInterfaces/Public/CADFileData.h
@@ -127,6 +127,8 @@ public:
 
 	int32 AddComponent(FCadId ComponentId)
 	{
+		ensure(SceneGraphArchive.Components.Num() < SceneGraphArchive.Components.Max());
+
 		int32 Index = SceneGraphArchive.Components.Emplace(ComponentId);
 		SceneGraphArchive.CADIdToComponentIndex.Add(ComponentId, Index);
 		return Index;
@@ -145,6 +147,8 @@ public:
 
 	int32 AddInstance(FCadId InstanceId)
 	{
+		ensure(SceneGraphArchive.Instances.Num() < SceneGraphArchive.Instances.Max());
+
 		int32 Index = SceneGraphArchive.Instances.Emplace(InstanceId);
 		SceneGraphArchive.CADIdToInstanceIndex.Add(InstanceId, Index);
 		return Index;
@@ -155,6 +159,13 @@ public:
 		return SceneGraphArchive.Instances[Index];
 	}
 
+	int32 GetInstanceIndexFromId(FCadId InstanceId)
+	{
+		int32* IndexPtr = SceneGraphArchive.CADIdToInstanceIndex.Find(InstanceId);
+
+		return IndexPtr ? *IndexPtr : INDEX_NONE;
+	}
+
 
 	bool HasBodyOfId(FCadId BodyId) const
 	{
@@ -163,6 +174,8 @@ public:
 
 	int32 AddBody(FCadId BodyId)
 	{
+		ensure(SceneGraphArchive.Bodies.Num() < SceneGraphArchive.Bodies.Max());
+
 		int32 Index = SceneGraphArchive.Bodies.Emplace(BodyId);
 		SceneGraphArchive.CADIdToBodyIndex.Add(BodyId, Index);
 		return Index;
@@ -186,6 +199,8 @@ public:
 
 	int32 AddUnloadedComponent(FCadId ComponentId)
 	{
+		ensure(SceneGraphArchive.UnloadedComponents.Num() < SceneGraphArchive.UnloadedComponents.Max());
+
 		int32 Index = SceneGraphArchive.UnloadedComponents.Emplace(ComponentId);
 		SceneGraphArchive.CADIdToUnloadedComponentIndex.Add(ComponentId, Index);
 		return Index;
diff --git a/Engine/Plugins/Enterprise/DatasmithCADImporter/Source/CADInterfaces/Public/CADSceneGraph.h b/Engine/Plugins/Enterprise/DatasmithCADImporter/Source/CADInterfaces/Public/CADSceneGraph.h
--- a/Engine/Plugins/Enterprise/DatasmithCADImporter/Source/CADInterfaces/Public/CADSceneGraph.h
+++ b/Engine/Plugins/Enterprise/DatasmithCADImporter/Source/CADInterfaces/Public/CADSceneGraph.h
@@ -143,6 +143,20 @@ public:
 	TMap<FCadId, int32> CADIdToComponentIndex;
 	TMap<FCadId, int32> CADIdToUnloadedComponentIndex;
 	TMap<FCadId, int32> CADIdToInstanceIndex;
+
+	void Reserve(int32 InstanceNum, int32 ComponentNum, int32 BodyNum)
+	{
+		Instances.Reserve(InstanceNum);
+		Components.Reserve(ComponentNum);
+		UnloadedComponents.Reserve(ComponentNum);
+		ExternalReferences.Reserve(ComponentNum);
+		Bodies.Reserve(BodyNum);
+
+		CADIdToInstanceIndex.Reserve(InstanceNum);
+		CADIdToComponentIndex.Reserve(ComponentNum);
+		CADIdToUnloadedComponentIndex.Reserve(ComponentNum);
+		CADIdToBodyIndex.Reserve(BodyNum);
+	}
 };
 
 
diff --git a/Engine/Plugins/Enterprise/DatasmithCADImporter/Source/CADInterfaces/Public/TechSoftInterface.h b/Engine/Plugins/Enterprise/DatasmithCADImporter/Source/CADInterfaces/Public/TechSoftInterface.h
--- a/Engine/Plugins/Enterprise/DatasmithCADImporter/Source/CADInterfaces/Public/TechSoftInterface.h
+++ b/Engine/Plugins/Enterprise/DatasmithCADImporter/Source/CADInterfaces/Public/TechSoftInterface.h
@@ -1,7 +1,11 @@
 // Copyright Epic Games, Inc. All Rights Reserved.
 #pragma once
 
+#ifdef CADKERNEL_DEV
+#include "CADKernel/Core/Types.h"
+#else
 #include "CoreMinimal.h"
+#endif
 
 #if PLATFORM_WINDOWS
 #include "Windows/AllowWindowsPlatformTypes.h"
@@ -19,7 +23,6 @@ THIRD_PARTY_INCLUDES_END
 
 namespace CADLibrary
 {
-CADINTERFACES_API bool TECHSOFT_InitializeKernel(const TCHAR* = TEXT(""));
 
 CADINTERFACES_API class FTechSoftInterface
 {
@@ -44,6 +47,7 @@ public:
 #ifdef USE_TECHSOFT_SDK
 	A3DStatus Import(const A3DImport& Import);
 	A3DAsmModelFile* GetModelFile();
+	A3DStatus UnloadModel();
 #endif
 
 private:
@@ -56,7 +60,22 @@ private:
 #endif
 };
 
+namespace TechSoftUtils
+{
+
+CADINTERFACES_API bool TECHSOFT_InitializeKernel(const TCHAR* = TEXT(""));
 CADINTERFACES_API FTechSoftInterface& GetTechSoftInterface();
 
+#ifdef USE_TECHSOFT_SDK
+CADINTERFACES_API A3DStatus GetSurfaceAsNurbs(const A3DSurfBase* SurfacePtr, A3DSurfNurbsData* DataPtr, A3DDouble Tolerance, A3DBool bUseSameParameterization);
+CADINTERFACES_API A3DStatus GetCurveAsNurbs(const A3DCrvBase* A3DCurve, A3DCrvNurbsData* DataPtr, A3DDouble Tolerance, A3DBool bUseSameParameterization);
+
+CADINTERFACES_API A3DStatus GetEntityType(const A3DEntity* pEntity, A3DEEntityType* peEntityType);
+CADINTERFACES_API bool IsEntityBaseWithGraphicsType(const A3DEntity* pEntity);
+CADINTERFACES_API bool IsEntityBaseType(const A3DEntity* EntityPtr);
+#endif
+
+} // NS TechSoftUtils
+
 }
 
diff --git a/Engine/Plugins/Enterprise/DatasmithCADImporter/Source/DatasmithCADTranslator/Private/DatasmithCADTranslator.cpp b/Engine/Plugins/Enterprise/DatasmithCADImporter/Source/DatasmithCADTranslator/Private/DatasmithCADTranslator.cpp
--- a/Engine/Plugins/Enterprise/DatasmithCADImporter/Source/DatasmithCADTranslator/Private/DatasmithCADTranslator.cpp
+++ b/Engine/Plugins/Enterprise/DatasmithCADImporter/Source/DatasmithCADTranslator/Private/DatasmithCADTranslator.cpp
@@ -76,7 +76,15 @@ void FDatasmithCADTranslator::Initialize(FDatasmithTranslatorCapabilities& OutCa
 	OutCapabilities.SupportedFileFormats.Add(FFileFormatInfo{ TEXT("dwg"), TEXT("AutoCAD, Model files") });
 	OutCapabilities.SupportedFileFormats.Add(FFileFormatInfo{ TEXT("dgn"), TEXT("MicroStation files") });
 
-	OutCapabilities.SupportedFileFormats.Add(FFileFormatInfo{ TEXT("ct"), TEXT("Kernel_IO files") });
+	if (CADLibrary::GCADLibrary.Equals(TEXT("TechSoft")))
+	{
+		OutCapabilities.SupportedFileFormats.Add(FFileFormatInfo{ TEXT("ifc"), TEXT("IFC files") });
+		OutCapabilities.SupportedFileFormats.Add(FFileFormatInfo{ TEXT("hsf"), TEXT("HOOPS stream files") });
+	}
+	else
+	{
+		OutCapabilities.SupportedFileFormats.Add(FFileFormatInfo{ TEXT("ct"), TEXT("Kernel_IO files") });
+	}
 }
 
 bool FDatasmithCADTranslator::IsSourceSupported(const FDatasmithSceneSource& Source)
