TechSoft Wip
- Improve Product Strucutre Import
- Import SolidWorks (Wip)
- Import External Ref

diff --git a/Engine/Plugins/Enterprise/DatasmithCADImporter/Source/CADInterfaces/Private/TechSoftFileParser.cpp b/Engine/Plugins/Enterprise/DatasmithCADImporter/Source/CADInterfaces/Private/TechSoftFileParser.cpp
--- a/Engine/Plugins/Enterprise/DatasmithCADImporter/Source/CADInterfaces/Private/TechSoftFileParser.cpp
+++ b/Engine/Plugins/Enterprise/DatasmithCADImporter/Source/CADInterfaces/Private/TechSoftFileParser.cpp
@@ -45,7 +45,7 @@ FString CleanCatiaInstanceSdkName(const FString& Name)
 	if (Name.FindChar(TEXT('('), Index))
 	{
 		FString NewName = Name.RightChop(Index + 1);
-		if (Name.FindLastChar(TEXT(')'), Index))
+		if (NewName.FindLastChar(TEXT(')'), Index))
 		{
 			NewName = NewName.Left(Index);
 		}
@@ -140,33 +140,6 @@ bool ReplaceOrAddNameValue(TMap<FString, FString>& MetaData, TCHAR* Key)
 }
 
 // Functions used in traverse model process
-bool IsUnloadedPrototype(const A3DAsmProductOccurrenceData& PrototypeData)
-{
-	bool bIsUnloaded = true;
-
-	if (PrototypeData.m_uiPOccurrencesSize > 0)
-	{
-		bIsUnloaded = false;
-	}
-	else if (PrototypeData.m_uiEntityReferenceSize > 0)
-	{
-		bIsUnloaded = false;
-	}
-	else if (PrototypeData.m_pPart)
-	{
-		bIsUnloaded = false;
-	}
-	else if (PrototypeData.m_pPrototype)
-	{
-		TUniqueTSObj<A3DAsmProductOccurrenceData> SubPrototypeData(PrototypeData.m_pPrototype);
-		if (SubPrototypeData.IsValid())
-		{
-			bIsUnloaded = TechSoftFileParserImpl::IsUnloadedPrototype(*SubPrototypeData);
-		}
-	}
-
-	return bIsUnloaded;
-}
 
 void TraverseAttribute(const A3DMiscAttributeData& AttributeData, TMap<FString, FString>& OutMetaData)
 {
@@ -418,28 +391,59 @@ void FTechSoftFileParser::TraverseConfigurationSet(const A3DAsmProductOccurrence
 	ExtractMetaData(ConfigurationSetPtr, MetaData);
 	ExtractSpecificMetaData(ConfigurationSetPtr, MetaData);
 
+	const FString& ConfigurationToLoad = CADFileData.GetCADFileDescription().GetConfiguration();
+
 	TUniqueTSObj<A3DAsmProductOccurrenceData> ConfigurationData;
 	for (unsigned int Index = 0; Index < ConfigurationSetData->m_uiPOccurrencesSize; ++Index)
 	{
 		ConfigurationData.FillFrom(ConfigurationSetData->m_ppPOccurrences[Index]);
 		if (!ConfigurationData.IsValid())
 		{
-			return;
+			continue;
 		}
 
-		// Traverse default configuration
-		if (ConfigurationData->m_uiProductFlags & (A3D_PRODUCT_FLAG_DEFAULT| A3D_PRODUCT_FLAG_CONFIG))
+		if (ConfigurationData->m_uiProductFlags & A3D_PRODUCT_FLAG_CONFIG)
 		{
-			TraverseReference(ConfigurationSetData->m_ppPOccurrences[Index]);
-			return;
+			bool bIsConfigurationToLoad = false;
+			if (!ConfigurationToLoad.IsEmpty())
+			{
+				FEntityMetaData ConfigurationMetaData;
+				ExtractMetaData(ConfigurationSetData->m_ppPOccurrences[Index], ConfigurationMetaData);
+				const FString* ConfigurationName = ConfigurationMetaData.MetaData.Find(TEXT("SDKName"));
+				if (ConfigurationName)
+				{
+					bIsConfigurationToLoad = (ConfigurationName->Equals(ConfigurationToLoad));
+				}
+			}
+			else
+			{
+				bIsConfigurationToLoad = ConfigurationData->m_uiProductFlags & A3D_PRODUCT_FLAG_DEFAULT;
+			}
+
+			if (bIsConfigurationToLoad)
+			{
+				TraverseReference(ConfigurationSetData->m_ppPOccurrences[Index]);
+				return;
+			}
 		}
 	}
 
-	// no default configuration, traverse the first
-	if (ConfigurationSetData->m_uiPOccurrencesSize)
+	if (ConfigurationToLoad.IsEmpty())
 	{
-// #ueent_techsoft: Should we take the first occurence or the first configuration???
-		TraverseReference(ConfigurationSetData->m_ppPOccurrences[0]);
+		// no default configuration, traverse the first configuration
+		for (unsigned int Index = 0; Index < ConfigurationSetData->m_uiPOccurrencesSize; ++Index)
+		{
+			ConfigurationData.FillFrom(ConfigurationSetData->m_ppPOccurrences[Index]);
+			if (!ConfigurationData.IsValid())
+			{
+				return;
+			}
+
+			if (ConfigurationData->m_uiProductFlags & A3D_PRODUCT_FLAG_CONFIG)
+			{
+				TraverseReference(ConfigurationSetData->m_ppPOccurrences[Index]);
+			}
+		}
 	}
 }
 
@@ -487,7 +491,7 @@ void FTechSoftFileParser::TraverseReference(const A3DAsmProductOccurrence* Refer
 	ExtractSpecificMetaData(ReferencePtr, MetaData);
 	BuildReferenceName(MetaData.MetaData);
 
-	TraverseMaterialProperties(ReferencePtr);
+	ExtractMaterialProperties(ReferencePtr);
 
 	int32 ComponentId = LastEntityId++;
 	int32 Index = CADFileData.AddComponent(ComponentId);
@@ -555,7 +559,7 @@ FArchiveUnloadedComponent& FTechSoftFileParser::AddUnloadedComponent(FEntityMeta
 
 	if (Format == ECADFormat::SOLIDWORKS /*|| CADFileData.FileFormat() == ECADFormat::JT*/)
 	{
-		if (FString* ConfigurationName = ComponentMetaData.MetaData.Find(TEXT("ConfigurationName")))
+		if (FString* ConfigurationName = Component.MetaData.Find(TEXT("ConfigurationName")))
 		{
 			Instance.ExternalReference.SetConfiguration(*ConfigurationName);
 		}
@@ -593,74 +597,6 @@ FArchiveBody& FTechSoftFileParser::AddBody(FEntityMetaData& BodyMetaData)
 	return Body;
 }
 
-#pragma optimize ("", off)
-#ifndef NEW_CODE
-FCadId FTechSoftFileParser::TraverseOccurrence(const A3DAsmProductOccurrence* OccurrencePtr)
-{
-	FEntityMetaData InstanceMetaData;
-	ExtractMetaData(OccurrencePtr, InstanceMetaData);
-
-	if (InstanceMetaData.bRemoved || !InstanceMetaData.bShow)
-	{
-		return 0;
-	}
-
-	ExtractSpecificMetaData(OccurrencePtr, InstanceMetaData);
-	BuildInstanceName(InstanceMetaData.MetaData);
-
-	TraverseMaterialProperties(OccurrencePtr);
-	TraverseLayer(OccurrencePtr);
-
-	TUniqueTSObj<A3DAsmProductOccurrenceData> OccurrenceData(OccurrencePtr);
-	if (!OccurrenceData.IsValid())
-	{
-		return 0;
-	}
-
-	FArchiveInstance& Instance = AddInstance(InstanceMetaData);
-
-	Instance.TransformMatrix = TraverseTransformation(OccurrenceData->m_pLocation);
-
-	FEntityMetaData PrototypeMetaData;
-	bool bIsUnloadedPrototype = false;
-	if (OccurrenceData->m_pPrototype)
-	{
-		FMatrix PrototypeMatrix;
-		ProcessPrototype(OccurrenceData->m_pPrototype, PrototypeMetaData, PrototypeMatrix);
-		if (!OccurrenceData->m_pLocation)
-		{
-			Instance.TransformMatrix = PrototypeMatrix;
-		}
-	}
-
-	if (PrototypeMetaData.bUnloaded)
-	{
-		FArchiveUnloadedComponent& Prototype = AddUnloadedComponent(PrototypeMetaData, Instance);
-	}
-	else
-	{
-		FArchiveComponent& Component = AddComponent(PrototypeMetaData, Instance);
-
-		for (uint32 Index = 0; Index < OccurrenceData->m_uiPOccurrencesSize; ++Index)
-		{
-			int32 ChildrenId = TraverseOccurrence(OccurrenceData->m_ppPOccurrences[Index]);
-			Component.Children.Add(ChildrenId);
-		}
-
-		if (OccurrenceData->m_pPart)
-		{
-			TraversePartDefinition(OccurrenceData->m_pPart, Component);
-		}
-
-		if (OccurrenceData->m_pPrototype && !OccurrenceData->m_uiPOccurrencesSize && !OccurrenceData->m_pPart)
-		{
-			TraversePrototype(OccurrenceData->m_pPrototype, Component);
-		}
-	}
-
-	return Instance.ObjectId;
-}
-#else
 FCadId FTechSoftFileParser::TraverseOccurrence(const A3DAsmProductOccurrence* OccurrencePtr)
 {
 	TUniqueTSObj<A3DAsmProductOccurrenceData> OccurrenceData(OccurrencePtr);
@@ -689,55 +625,51 @@ FCadId FTechSoftFileParser::TraverseOccurrence(const A3DAsmProductOccurrence* Oc
 	ExtractSpecificMetaData(OccurrencePtr, InstanceMetaData);
 	BuildInstanceName(InstanceMetaData.MetaData);
 
-	TraverseMaterialProperties(OccurrencePtr);
-	TraverseLayer(OccurrencePtr);
+	ExtractMaterialProperties(OccurrencePtr);
+	ExtractLayer(OccurrencePtr);
 
 	FArchiveInstance& Instance = AddInstance(InstanceMetaData);
 
-	if (OccurrenceData->m_eProductLoadStatus == kA3DProductLoadStatusNotLoaded)
+	A3DMiscTransformation* Location = OccurrenceData->m_pLocation;
+
+	FEntityMetaData PrototypeMetaData;
+	if (OccurrenceData->m_pPrototype)
+	{
+		ProcessPrototype(OccurrenceData->m_pPrototype, PrototypeMetaData, &Location);
+	}
+
+	if (Location)
 	{
-		FArchiveUnloadedComponent& UnloadedComponent = AddUnloadedComponent(InstanceMetaData, Instance);
+		Instance.TransformMatrix = TraverseTransformation(Location);
+	}
+
+	if (PrototypeMetaData.bUnloaded)
+	{
+		FArchiveUnloadedComponent& UnloadedComponent = AddUnloadedComponent(PrototypeMetaData, Instance);
 		return Instance.ObjectId;
 	}
 	
 	A3DAsmProductOccurrence* PrototypePtr = OccurrenceData->m_pPrototype;
-	A3DMiscTransformation* LocationPtr = OccurrenceData->m_pLocation;
-
-	while (!LocationPtr && PrototypePtr)
+	while (OccurrenceData->m_pPrototype != nullptr && OccurrenceData->m_pPart == nullptr && OccurrenceData->m_uiPOccurrencesSize == 0)
 	{
-		TUniqueTSObj<A3DAsmProductOccurrenceData> PrototypeOccurrenceData(PrototypePtr);
-		LocationPtr = PrototypeOccurrenceData->m_pLocation;
-		PrototypePtr = PrototypeOccurrenceData->m_pPrototype;
+		OccurrenceData.FillFrom(OccurrenceData->m_pPrototype);
 	}
 
-	Instance.TransformMatrix = TraverseTransformation(LocationPtr);
-
-	A3DAsmPartDefinition* PartDefinition = OccurrenceData->m_pPart;
-	PrototypePtr = OccurrenceData->m_pPrototype;
-
-	while (!PartDefinition && PrototypePtr)
+	if(OccurrenceData->m_pPart == nullptr && OccurrenceData->m_uiPOccurrencesSize == 0)
 	{
-		TUniqueTSObj<A3DAsmProductOccurrenceData> PrototypeOccurrenceData(PrototypePtr);
-		PartDefinition = PrototypeOccurrenceData->m_pPart;
-		PrototypePtr = PrototypeOccurrenceData->m_pPrototype;
+		return Instance.ObjectId;
 	}
 
 	FArchiveComponent& Component = AddComponent(InstanceMetaData, Instance);
 
-	TraversePartDefinition(PartDefinition, Component);
-
-	uint32 ChildrenCount = OccurrenceData->m_uiPOccurrencesSize;
-	A3DAsmProductOccurrence** Children = OccurrenceData->m_ppPOccurrences;
-	PrototypePtr = OccurrenceData->m_pPrototype;
-
-	while (ChildrenCount == 0 && PrototypePtr)
+	if (OccurrenceData->m_pPart != nullptr)
 	{
-		TUniqueTSObj<A3DAsmProductOccurrenceData> PrototypeOccurrenceData(PrototypePtr);
-		ChildrenCount = PrototypeOccurrenceData->m_uiPOccurrencesSize;
-		Children = PrototypeOccurrenceData->m_ppPOccurrences;
-		PrototypePtr = PrototypeOccurrenceData->m_pPrototype;
+		A3DAsmPartDefinition* PartDefinition = OccurrenceData->m_pPart;
+		TraversePartDefinition(PartDefinition, Component);
 	}
 
+	uint32 ChildrenCount = OccurrenceData->m_uiPOccurrencesSize;
+	A3DAsmProductOccurrence** Children = OccurrenceData->m_ppPOccurrences;
 	for (uint32 Index = 0; Index < ChildrenCount; ++Index)
 	{
 		int32 ChildrenId = TraverseOccurrence(Children[Index]);
@@ -746,8 +678,6 @@ FCadId FTechSoftFileParser::TraverseOccurrence(const A3DAsmProductOccurrence* Oc
 
 	return Instance.ObjectId;
 }
-#endif
-#pragma optimize ("", on)
 
 void FTechSoftFileParser::ProcessOccurrence(TUniqueTSObj<A3DAsmProductOccurrenceData>& OccurrenceData, FArchiveComponent& Component)
 {
@@ -813,48 +743,6 @@ void FTechSoftFileParser::CountUnderOccurrence(const A3DAsmProductOccurrence* Oc
 	}
 }
 
-FFileDescriptor FTechSoftFileParser::GetOccurrenceFileName(const A3DAsmProductOccurrence* OccurrencePtr)
-{
-	TUniqueTSObj<A3DAsmProductOccurrenceData> OccurrenceData(OccurrencePtr);
-	if (!OccurrenceData.IsValid())
-	{
-		return FFileDescriptor();
-	}
-
-	if (OccurrenceData->m_pPrototype)
-	{
-		return GetOccurrenceFileName(OccurrenceData->m_pPrototype);
-	}
-
-	A3DUTF8Char* FilePathUTF8Ptr = nullptr;
-	int Ret = A3DAsmProductOccurrenceGetFilePathName(OccurrencePtr, &FilePathUTF8Ptr);
-	if (Ret == A3D_SUCCESS && FilePathUTF8Ptr)
-	{
-		FString OccurrenceFilePath = UTF8_TO_TCHAR(FilePathUTF8Ptr);
-		A3DAsmProductOccurrenceGetFilePathName(NULL, &FilePathUTF8Ptr);
-		FString OccurrenceFileName = FPaths::GetCleanFilename(OccurrenceFilePath);
-		if (OccurrenceFileName != CADFileData.GetCADFileDescription().GetFileName())
-		{
-			return FFileDescriptor(*OccurrenceFilePath, nullptr, *CADFileData.GetCADFileDescription().GetRootFolder());
-		}
-	}
-
-	FilePathUTF8Ptr = nullptr;
-	Ret = A3DAsmProductOccurrenceGetOriginalFilePathName(OccurrencePtr, &FilePathUTF8Ptr);
-	if (Ret == A3D_SUCCESS && FilePathUTF8Ptr)
-	{
-		FString OriginalFilePath = UTF8_TO_TCHAR(FilePathUTF8Ptr);
-		A3DAsmProductOccurrenceGetFilePathName(NULL, &FilePathUTF8Ptr);
-		FString OriginalFileName = FPaths::GetCleanFilename(OriginalFilePath);
-		if (OriginalFileName != CADFileData.GetCADFileDescription().GetFileName())
-		{
-			return FFileDescriptor(*OriginalFilePath, nullptr, *CADFileData.GetCADFileDescription().GetRootFolder());
-		}
-	}
-
-	return FFileDescriptor();
-}
-
 void FTechSoftFileParser::TraversePrototype(const A3DAsmProductOccurrence* InPrototypePtr, FArchiveComponent& Component)
 {
 	TUniqueTSObj<A3DAsmProductOccurrenceData> SubPrototypeData(InPrototypePtr);
@@ -904,27 +792,80 @@ void FTechSoftFileParser::CountUnderSubPrototype(const A3DAsmProductOccurrence*
 	}
 }
 
-void FTechSoftFileParser::ProcessPrototype(const A3DAsmProductOccurrence* InPrototypePtr, FEntityMetaData& OutPrototypeMetaData, FMatrix& OutPrototypeMatrix)
+void FTechSoftFileParser::ProcessPrototype(const A3DAsmProductOccurrence* InPrototypePtr, FEntityMetaData& OutPrototypeMetaData, A3DMiscTransformation** OutLocation)
 {
-	TUniqueTSObj<A3DAsmProductOccurrenceData> PrototypeData(InPrototypePtr);
-	if (!PrototypeData.IsValid())
+	OutPrototypeMetaData.bUnloaded = true;
+
+	const A3DAsmProductOccurrence* PrototypePtr = InPrototypePtr;
+	TUniqueTSObj<A3DAsmProductOccurrenceData> PrototypeData;
+	while (PrototypePtr)
 	{
-		return;
-	}
+		PrototypeData.FillFrom(PrototypePtr);
+		if (!PrototypeData.IsValid())
+		{
+			return;
+		}
 
-	ExtractMetaData(InPrototypePtr, OutPrototypeMetaData);
-	ExtractSpecificMetaData(InPrototypePtr, OutPrototypeMetaData);
+		if (PrototypeData->m_pPart != nullptr || PrototypeData->m_uiPOccurrencesSize != 0 || PrototypeData->m_pPrototype == nullptr)
+		{
+			ExtractMetaData(PrototypePtr, OutPrototypeMetaData);
+			ExtractSpecificMetaData(PrototypePtr, OutPrototypeMetaData);
+			ExtractMaterialProperties(PrototypePtr);
 
-	TraverseMaterialProperties(InPrototypePtr);
+			TUniqueTSObj<A3DUTF8Char*> FilePathUTF8Ptr;
+			FilePathUTF8Ptr.FillWith(&TechSoftUtils::GetFilePathName, PrototypePtr);
+			if (!FilePathUTF8Ptr.IsValid())
+			{
+				FilePathUTF8Ptr.FillWith(&TechSoftUtils::GetOriginalFilePathName, PrototypePtr);
+			}
+			if (FilePathUTF8Ptr.IsValid())
+			{
+				FString FilePath = UTF8_TO_TCHAR(*FilePathUTF8Ptr);
+				FPaths::NormalizeFilename(FilePath);
+				FString FileName = FPaths::GetCleanFilename(FilePath);
+				if (FileName != CADFileData.GetCADFileDescription().GetFileName())
+				{
+					OutPrototypeMetaData.ExternalFile = FFileDescriptor(*FilePath, nullptr, *CADFileData.GetCADFileDescription().GetRootFolder());
+				}
+			}
+		}
+		
+		if (PrototypeData->m_pPart != nullptr || PrototypeData->m_uiPOccurrencesSize != 0)
+		{
+			OutPrototypeMetaData.bUnloaded = false;
+			PrototypePtr = nullptr;
+		}
+		else
+		{
+			PrototypePtr = PrototypeData->m_pPrototype;
+		}
+
+		if (!*OutLocation)
+		{
+			*OutLocation = PrototypeData->m_pLocation;
+		}
+	}
 
-	OutPrototypeMatrix = TraverseTransformation(PrototypeData->m_pLocation);
+	if (!*OutLocation)
+	{
+		while (PrototypeData.IsValid() && PrototypeData->m_pLocation == nullptr && PrototypeData->m_pPrototype != nullptr)
+		{
+			PrototypeData.FillFrom(PrototypeData->m_pPrototype);
+		}
+		if(PrototypeData.IsValid())
+		{
+			*OutLocation = PrototypeData->m_pLocation;
+		}
+	}
 
-	OutPrototypeMetaData.bUnloaded = TechSoftFileParserImpl::IsUnloadedPrototype(*PrototypeData);
 	if (OutPrototypeMetaData.bUnloaded)
 	{
-		OutPrototypeMetaData.ExternalFile = GetOccurrenceFileName(InPrototypePtr);
 		OutPrototypeMetaData.MetaData.Add(TEXT("Name"), OutPrototypeMetaData.ExternalFile.GetFileName());
 	}
+	else
+	{
+		OutPrototypeMetaData.ExternalFile.Empty();
+	}
 
 	BuildReferenceName(OutPrototypeMetaData.MetaData);
 }
@@ -953,8 +894,8 @@ void FTechSoftFileParser::TraversePartDefinition(const A3DAsmPartDefinition* Par
 	ExtractSpecificMetaData(PartDefinitionPtr, PartMetaData);
 	BuildPartName(PartMetaData.MetaData);
 
-	TraverseMaterialProperties(PartDefinitionPtr);
-	TraverseLayer(PartDefinitionPtr);
+	ExtractMaterialProperties(PartDefinitionPtr);
+	ExtractLayer(PartDefinitionPtr);
 
 	TUniqueTSObj<A3DAsmPartDefinitionData> PartData(PartDefinitionPtr);
 	if (PartData.IsValid())
@@ -1037,7 +978,7 @@ FCadId FTechSoftFileParser::TraverseRepresentationSet(const A3DRiSet* Representa
 		return 0;
 	}
 
-	TraverseMaterialProperties(RepresentationSetPtr);
+	ExtractMaterialProperties(RepresentationSetPtr);
 
 	FCadId RepresentationSetId = 0;
 	FArchiveComponent& RepresentationSet = AddOccurence(RepresentationSetMetaData, RepresentationSetId);
@@ -1091,7 +1032,7 @@ FCadId FTechSoftFileParser::TraverseBRepModel(A3DRiBrepModel* BRepModelPtr, FEnt
 	}
 
 	ExtractSpecificMetaData(BRepModelPtr, BRepMetaData);
-	TraverseMaterialProperties(BRepModelPtr);
+	ExtractMaterialProperties(BRepModelPtr);
 
 	FArchiveBody& Body = AddBody(BRepMetaData);
 
@@ -1155,7 +1096,7 @@ FCadId FTechSoftFileParser::TraversePolyBRepModel(const A3DRiPolyBrepModel* Poly
 	}
 
 	ExtractSpecificMetaData(PolygonalPtr, BRepMetaData);
-	TraverseMaterialProperties(PolygonalPtr);
+	ExtractMaterialProperties(PolygonalPtr);
 
 	FArchiveBody& Body = AddBody(BRepMetaData);
 	TraverseRepresentationContent(PolygonalPtr, Body);
@@ -1201,7 +1142,7 @@ void FTechSoftFileParser::ExtractMetaData(const A3DEntity* Entity, FEntityMetaDa
 		{
 			if (MetaDataWithGraphics->m_pGraphics != NULL)
 			{
-				TraverseGraphics(MetaDataWithGraphics->m_pGraphics, OutMetaData);
+				ExtractGraphicProperties(MetaDataWithGraphics->m_pGraphics, OutMetaData);
 			}
 		}
 	}
@@ -1522,7 +1463,7 @@ FArchiveMaterial& FTechSoftFileParser::FindOrAddMaterial(uint32 MaterialIndex)
 	return NewMaterial;
 }
 
-void FTechSoftFileParser::TraverseGraphics(const A3DGraphics* Graphics, FEntityMetaData& OutMetaData)
+void FTechSoftFileParser::ExtractGraphicProperties(const A3DGraphics* Graphics, FEntityMetaData& OutMetaData)
 {
 	TUniqueTSObj<A3DGraphicsData> GraphicsData(Graphics);
 	if (!GraphicsData.IsValid())
@@ -1552,7 +1493,7 @@ void FTechSoftFileParser::TraverseGraphics(const A3DGraphics* Graphics, FEntityM
 	FCADUUID ColorName;
 	FCADUUID MaterialName;
 
-	TraverseGraphStyleData(GraphicsData->m_uiStyleIndex, ColorName, MaterialName);
+	ExtractGraphStyleProperties(GraphicsData->m_uiStyleIndex, ColorName, MaterialName);
 
 	if (ColorName)
 	{
@@ -1565,7 +1506,7 @@ void FTechSoftFileParser::TraverseGraphics(const A3DGraphics* Graphics, FEntityM
 	}
 }
 
-void FTechSoftFileParser::TraverseGraphStyleData(uint32 StyleIndex, FCADUUID& OutColorName, FCADUUID& OutMaterialName)
+void FTechSoftFileParser::ExtractGraphStyleProperties(uint32 StyleIndex, FCADUUID& OutColorName, FCADUUID& OutMaterialName)
 {
 	OutColorName = 0;
 	OutMaterialName = 0;
@@ -1593,7 +1534,7 @@ void FTechSoftFileParser::TraverseGraphStyleData(uint32 StyleIndex, FCADUUID& Ou
 	}
 }
 
-void FTechSoftFileParser::TraverseMaterialProperties(const A3DEntity* Entity)
+void FTechSoftFileParser::ExtractMaterialProperties(const A3DEntity* Entity)
 {
 	TUniqueTSObj<A3DMiscMaterialPropertiesData> MaterialPropertiesData(Entity);
 	if (!MaterialPropertiesData.IsValid())
@@ -1706,7 +1647,7 @@ FMatrix FTechSoftFileParser::TraverseCoordinateSystem(const A3DRiCoordinateSyste
 	return FMatrix::Identity;
 }
 
-void FTechSoftFileParser::TraverseLayer(const A3DAsmProductOccurrence* Occurrence)
+void FTechSoftFileParser::ExtractLayer(const A3DAsmProductOccurrence* Occurrence)
 {
 	// TODO
 
@@ -2344,7 +2285,7 @@ void FTechSoftFileParser::TraverseTessellation3D(const A3DTess3D* TessellationPt
 				A3DUns32 StyleIndex = FaceTessData.m_puiStyleIndexes[0];
 				FCADUUID ColorName;
 				FCADUUID MaterialName;
-				TraverseGraphStyleData(StyleIndex, ColorName, MaterialName);
+				ExtractGraphStyleProperties(StyleIndex, ColorName, MaterialName);
 				if (ColorName)
 				{
 					Tessellation.ColorName = ColorName;
diff --git a/Engine/Plugins/Enterprise/DatasmithCADImporter/Source/CADInterfaces/Private/TechSoftFileParser.h b/Engine/Plugins/Enterprise/DatasmithCADImporter/Source/CADInterfaces/Private/TechSoftFileParser.h
--- a/Engine/Plugins/Enterprise/DatasmithCADImporter/Source/CADInterfaces/Private/TechSoftFileParser.h
+++ b/Engine/Plugins/Enterprise/DatasmithCADImporter/Source/CADInterfaces/Private/TechSoftFileParser.h
@@ -178,7 +178,7 @@ private:
 	void TraverseConfigurationSet(const A3DAsmProductOccurrence* ConfigurationSet);
 	FCadId TraverseOccurrence(const A3DAsmProductOccurrence* Occurrence);
 	void TraversePrototype(const A3DAsmProductOccurrence* Prototype, FArchiveComponent& Component);
-	void ProcessPrototype(const A3DAsmProductOccurrence* InPrototype, FEntityMetaData& OutMetaData, FMatrix& OutPrototypeMatrix);
+	void ProcessPrototype(const A3DAsmProductOccurrence* InPrototype, FEntityMetaData& OutMetaData, A3DMiscTransformation** OutLocation);
 	void ProcessOccurrence(TUniqueTSObj<A3DAsmProductOccurrenceData>& Occurrence, FArchiveComponent& Component);
 	void TraversePartDefinition(const A3DAsmPartDefinition* PartDefinition, FArchiveComponent& Component);
 	FCadId TraverseRepresentationSet(const A3DRiSet* pSet, FEntityMetaData& PartMetaData);
@@ -202,10 +202,10 @@ private:
 	void BuildBodyName(TMap<FString, FString>& MetaData);
 
 	// Graphic properties
-	void TraverseGraphics(const A3DGraphics* Graphics, FEntityMetaData& OutMetaData);
-	void TraverseGraphStyleData(uint32 StyleIndex, FCADUUID& ColorName, FCADUUID& MaterialName);
-	void TraverseMaterialProperties(const A3DEntity* Entity);
-	void TraverseLayer(const A3DAsmProductOccurrence* Occurrence);
+	void ExtractGraphicProperties(const A3DGraphics* Graphics, FEntityMetaData& OutMetaData);
+	void ExtractGraphStyleProperties(uint32 StyleIndex, FCADUUID& ColorName, FCADUUID& MaterialName);
+	void ExtractMaterialProperties(const A3DEntity* Entity);
+	void ExtractLayer(const A3DAsmProductOccurrence* Occurrence);
 	FArchiveColor& FindOrAddColor(uint32 ColorIndex, uint8 Alpha);
 	FArchiveMaterial& FindOrAddMaterial(uint32 MaterialIndex);
 
@@ -215,8 +215,6 @@ private:
 	FMatrix TraverseGeneralTransformation(const A3DMiscTransformation* GeneralTransformation);
 	FMatrix TraverseTransformation3D(const A3DMiscTransformation* CartesianTransformation);
 
-	FFileDescriptor GetOccurrenceFileName(const A3DAsmProductOccurrence* OccurrencePtr);
-
 	// Archive methodes
 	FArchiveInstance& AddInstance(FEntityMetaData& InstanceMetaData);
 	FArchiveComponent& AddComponent(FEntityMetaData& ComponentMetaData, FArchiveInstance& Instance);
diff --git a/Engine/Plugins/Enterprise/DatasmithCADImporter/Source/CADInterfaces/Private/TechSoftInterface.cpp b/Engine/Plugins/Enterprise/DatasmithCADImporter/Source/CADInterfaces/Private/TechSoftInterface.cpp
--- a/Engine/Plugins/Enterprise/DatasmithCADImporter/Source/CADInterfaces/Private/TechSoftInterface.cpp
+++ b/Engine/Plugins/Enterprise/DatasmithCADImporter/Source/CADInterfaces/Private/TechSoftInterface.cpp
@@ -413,6 +413,11 @@ void TUniqueTSObj<A3DTopoWireEdgeData>::InitializeData()
 	A3D_INITIALIZE_DATA(A3DTopoWireEdgeData, Data);
 }
 
+void TUniqueTSObj<A3DUTF8Char*>::InitializeData()
+{
+	Data = nullptr;
+}
+
 A3DStatus TUniqueTSObj<A3DAsmModelFileData>::GetData(const A3DAsmModelFile* InAsmModelFilePtr)
 {
 	return A3DAsmModelFileGet(InAsmModelFilePtr, &Data);
@@ -739,6 +744,10 @@ A3DStatus TUniqueTSObj<A3DTopoWireEdgeData>::GetData(const A3DTopoWireEdge* InTo
 	return A3DTopoWireEdgeGet(InTopoWireEdgePtr, &Data);
 }
 
+A3DStatus TUniqueTSObj<A3DUTF8Char*>::GetData(const A3DAsmProductOccurrence* InOccurrencePtr)
+{
+	return A3DAsmProductOccurrenceGetFilePathName(InOccurrencePtr, &Data);
+}
 
 uint32 TUniqueTSObjFromIndex<A3DGraphMaterialData>::DefaultValue = A3D_DEFAULT_MATERIAL_INDEX;
 uint32 TUniqueTSObjFromIndex<A3DGraphPictureData>::DefaultValue = A3D_DEFAULT_MATERIAL_INDEX;
@@ -772,6 +781,16 @@ A3DStatus GetCurveAsNurbs(const A3DCrvBase* CurvePtr, A3DCrvNurbsData* DataPtr,
 	return A3DCrvBaseGetAsNurbs(CurvePtr, Tolerance, bUseSameParameterization, DataPtr);
 }
 
+A3DStatus GetOriginalFilePathName(const A3DAsmProductOccurrence* A3DOccurrencePtr, A3DUTF8Char** FilePathUTF8Ptr)
+{
+	return A3DAsmProductOccurrenceGetOriginalFilePathName(A3DOccurrencePtr, FilePathUTF8Ptr);
+}
+
+A3DStatus GetFilePathName(const A3DAsmProductOccurrence* A3DOccurrencePtr, A3DUTF8Char** FilePathUTF8Ptr)
+{
+	return A3DAsmProductOccurrenceGetFilePathName(A3DOccurrencePtr, FilePathUTF8Ptr);
+}
+
 A3DStatus GetEntityType(const A3DEntity* EntityPtr, A3DEEntityType* EntityTypePtr)
 {
 	return A3DEntityGetType(EntityPtr, EntityTypePtr);
@@ -786,6 +805,7 @@ bool IsEntityBaseType(const A3DEntity* EntityPtr)
 {
 	return (bool)A3DEntityIsBaseType(EntityPtr);
 }
+
 #endif
 
 }
diff --git a/Engine/Plugins/Enterprise/DatasmithCADImporter/Source/CADInterfaces/Public/TechSoftInterface.h b/Engine/Plugins/Enterprise/DatasmithCADImporter/Source/CADInterfaces/Public/TechSoftInterface.h
--- a/Engine/Plugins/Enterprise/DatasmithCADImporter/Source/CADInterfaces/Public/TechSoftInterface.h
+++ b/Engine/Plugins/Enterprise/DatasmithCADImporter/Source/CADInterfaces/Public/TechSoftInterface.h
@@ -70,6 +70,9 @@ CADINTERFACES_API FTechSoftInterface& GetTechSoftInterface();
 CADINTERFACES_API A3DStatus GetSurfaceAsNurbs(const A3DSurfBase* SurfacePtr, A3DSurfNurbsData* DataPtr, A3DDouble Tolerance, A3DBool bUseSameParameterization);
 CADINTERFACES_API A3DStatus GetCurveAsNurbs(const A3DCrvBase* A3DCurve, A3DCrvNurbsData* DataPtr, A3DDouble Tolerance, A3DBool bUseSameParameterization);
 
+CADINTERFACES_API A3DStatus GetOriginalFilePathName(const A3DAsmProductOccurrence* A3DOccurrencePtr, A3DUTF8Char** FilePathUTF8Ptr);
+CADINTERFACES_API A3DStatus GetFilePathName(const A3DAsmProductOccurrence* A3DOccurrencePtr, A3DUTF8Char** FilePathUTF8Ptr);
+
 CADINTERFACES_API A3DStatus GetEntityType(const A3DEntity* pEntity, A3DEEntityType* peEntityType);
 CADINTERFACES_API bool IsEntityBaseWithGraphicsType(const A3DEntity* pEntity);
 CADINTERFACES_API bool IsEntityBaseType(const A3DEntity* EntityPtr);
diff --git a/Engine/Plugins/Enterprise/DatasmithCADImporter/Source/CADTools/Public/CADData.h b/Engine/Plugins/Enterprise/DatasmithCADImporter/Source/CADTools/Public/CADData.h
--- a/Engine/Plugins/Enterprise/DatasmithCADImporter/Source/CADTools/Public/CADData.h
+++ b/Engine/Plugins/Enterprise/DatasmithCADImporter/Source/CADTools/Public/CADData.h
@@ -105,6 +105,21 @@ public:
 		return (Name.Equals(Other.Name, ESearchCase::IgnoreCase) && (Configuration == Other.Configuration));
 	}
 
+	bool IsEmpty()
+	{
+		return Name.IsEmpty();
+	}
+
+	void Empty()
+	{
+		SourceFilePath.Empty(); 
+		CacheFilePath.Empty();
+		Name.Empty();
+		Configuration.Empty();
+		RootFolder.Empty();
+		DescriptorHash = 0;
+	}
+
 	friend CADTOOLS_API FArchive& operator<<(FArchive& Ar, FFileDescriptor& File);
 
 	friend CADTOOLS_API uint32 GetTypeHash(const FFileDescriptor& FileDescription);
