TechSoft : import SW configuration

diff --git a/Engine/Plugins/Enterprise/DatasmithCADImporter/Source/CADInterfaces/Private/TechSoftFileParser.cpp b/Engine/Plugins/Enterprise/DatasmithCADImporter/Source/CADInterfaces/Private/TechSoftFileParser.cpp
--- a/Engine/Plugins/Enterprise/DatasmithCADImporter/Source/CADInterfaces/Private/TechSoftFileParser.cpp
+++ b/Engine/Plugins/Enterprise/DatasmithCADImporter/Source/CADInterfaces/Private/TechSoftFileParser.cpp
@@ -242,6 +242,12 @@ void UpdateIOOptionAccordingToFormat(const CADLibrary::ECADFormat Format, A3DImp
 		break;
 	}
 
+	case CADLibrary::ECADFormat::SOLIDWORKS:
+	{
+		Importer.m_sLoadData.m_sSpecifics.m_sSolidworks.m_bLoadAllConfigsData = true;
+		break;
+	}
+
 	case CADLibrary::ECADFormat::INVENTOR:
 	case CADLibrary::ECADFormat::CATIA_3DXML:
 	{
