TechSoft imporve Occurrence parsing

diff --git a/Engine/Plugins/Enterprise/DatasmithCADImporter/Source/CADInterfaces/Private/TechSoftFileParser.cpp b/Engine/Plugins/Enterprise/DatasmithCADImporter/Source/CADInterfaces/Private/TechSoftFileParser.cpp
--- a/Engine/Plugins/Enterprise/DatasmithCADImporter/Source/CADInterfaces/Private/TechSoftFileParser.cpp
+++ b/Engine/Plugins/Enterprise/DatasmithCADImporter/Source/CADInterfaces/Private/TechSoftFileParser.cpp
@@ -655,7 +655,6 @@ FCadId FTechSoftFileParser::TraverseOccurrence(const A3DAsmProductOccurrence* Oc
 		return Instance.ObjectId;
 	}
 	
-	A3DAsmProductOccurrence* PrototypePtr = OccurrenceData->m_pPrototype;
 	while (OccurrenceData->m_pPrototype != nullptr && OccurrenceData->m_pPart == nullptr && OccurrenceData->m_uiPOccurrencesSize == 0)
 	{
 		OccurrenceData.FillFrom(OccurrenceData->m_pPrototype);
@@ -668,12 +667,26 @@ FCadId FTechSoftFileParser::TraverseOccurrence(const A3DAsmProductOccurrence* Oc
 
 	FArchiveComponent& Component = AddComponent(InstanceMetaData, Instance);
 
+	A3DAsmProductOccurrence* PrototypePtr = OccurrenceData->m_pPrototype;
+
+	// Add part
+	while (OccurrenceData->m_pPrototype != nullptr && OccurrenceData->m_pPart == nullptr)
+	{
+		OccurrenceData.FillFrom(OccurrenceData->m_pPrototype);
+	}
 	if (OccurrenceData->m_pPart != nullptr)
 	{
 		A3DAsmPartDefinition* PartDefinition = OccurrenceData->m_pPart;
 		TraversePartDefinition(PartDefinition, Component);
 	}
 
+	// Add Occurrence's Children
+	OccurrenceData.FillFrom(PrototypePtr);
+	while (OccurrenceData->m_pPrototype != nullptr && OccurrenceData->m_uiPOccurrencesSize == 0)
+	{
+		OccurrenceData.FillFrom(OccurrenceData->m_pPrototype);
+	}
+
 	uint32 ChildrenCount = OccurrenceData->m_uiPOccurrencesSize;
 	A3DAsmProductOccurrence** Children = OccurrenceData->m_ppPOccurrences;
 	for (uint32 Index = 0; Index < ChildrenCount; ++Index)
