TechSoft fix imporve Occurrence parsing

diff --git a/Engine/Plugins/Enterprise/DatasmithCADImporter/Source/CADInterfaces/Private/TechSoftFileParser.cpp b/Engine/Plugins/Enterprise/DatasmithCADImporter/Source/CADInterfaces/Private/TechSoftFileParser.cpp
--- a/Engine/Plugins/Enterprise/DatasmithCADImporter/Source/CADInterfaces/Private/TechSoftFileParser.cpp
+++ b/Engine/Plugins/Enterprise/DatasmithCADImporter/Source/CADInterfaces/Private/TechSoftFileParser.cpp
@@ -605,6 +605,7 @@ FArchiveBody& FTechSoftFileParser::AddBody(FEntityMetaData& BodyMetaData)
 
 FCadId FTechSoftFileParser::TraverseOccurrence(const A3DAsmProductOccurrence* OccurrencePtr)
 {
+	const A3DAsmProductOccurrence* CurrentOccurrencePtr = OccurrencePtr;
 	TUniqueTSObj<A3DAsmProductOccurrenceData> OccurrenceData(OccurrencePtr);
 	if (!OccurrenceData.IsValid())
 	{
@@ -658,6 +659,7 @@ FCadId FTechSoftFileParser::TraverseOccurrence(const A3DAsmProductOccurrence* Oc
 	while (OccurrenceData->m_pPrototype != nullptr && OccurrenceData->m_pPart == nullptr && OccurrenceData->m_uiPOccurrencesSize == 0)
 	{
 		OccurrenceData.FillFrom(OccurrenceData->m_pPrototype);
+		CurrentOccurrencePtr = OccurrenceData->m_pPrototype;
 	}
 
 	if(OccurrenceData->m_pPart == nullptr && OccurrenceData->m_uiPOccurrencesSize == 0)
@@ -667,8 +669,6 @@ FCadId FTechSoftFileParser::TraverseOccurrence(const A3DAsmProductOccurrence* Oc
 
 	FArchiveComponent& Component = AddComponent(InstanceMetaData, Instance);
 
-	A3DAsmProductOccurrence* PrototypePtr = OccurrenceData->m_pPrototype;
-
 	// Add part
 	while (OccurrenceData->m_pPrototype != nullptr && OccurrenceData->m_pPart == nullptr)
 	{
@@ -681,7 +681,7 @@ FCadId FTechSoftFileParser::TraverseOccurrence(const A3DAsmProductOccurrence* Oc
 	}
 
 	// Add Occurrence's Children
-	OccurrenceData.FillFrom(PrototypePtr);
+	OccurrenceData.FillFrom(CurrentOccurrencePtr);
 	while (OccurrenceData->m_pPrototype != nullptr && OccurrenceData->m_uiPOccurrencesSize == 0)
 	{
 		OccurrenceData.FillFrom(OccurrenceData->m_pPrototype);
