TechSoft Occurrence parsing

diff --git a/Engine/Plugins/Enterprise/DatasmithCADImporter/Source/CADInterfaces/Private/TechSoftFileParser.cpp b/Engine/Plugins/Enterprise/DatasmithCADImporter/Source/CADInterfaces/Private/TechSoftFileParser.cpp
--- a/Engine/Plugins/Enterprise/DatasmithCADImporter/Source/CADInterfaces/Private/TechSoftFileParser.cpp
+++ b/Engine/Plugins/Enterprise/DatasmithCADImporter/Source/CADInterfaces/Private/TechSoftFileParser.cpp
@@ -605,7 +605,8 @@ FArchiveBody& FTechSoftFileParser::AddBody(FEntityMetaData& BodyMetaData)
 
 FCadId FTechSoftFileParser::TraverseOccurrence(const A3DAsmProductOccurrence* OccurrencePtr)
 {
-	const A3DAsmProductOccurrence* CurrentOccurrencePtr = OccurrencePtr;
+	// first product occurence with m_pPart != nullptr || m_uiPOccurrencesSize > 0
+	const A3DAsmProductOccurrence* CachedOccurrencePtr = OccurrencePtr;
 	TUniqueTSObj<A3DAsmProductOccurrenceData> OccurrenceData(OccurrencePtr);
 	if (!OccurrenceData.IsValid())
 	{
@@ -658,8 +659,8 @@ FCadId FTechSoftFileParser::TraverseOccurrence(const A3DAsmProductOccurrence* Oc
 	
 	while (OccurrenceData->m_pPrototype != nullptr && OccurrenceData->m_pPart == nullptr && OccurrenceData->m_uiPOccurrencesSize == 0)
 	{
+		CachedOccurrencePtr = OccurrenceData->m_pPrototype;
 		OccurrenceData.FillFrom(OccurrenceData->m_pPrototype);
-		CurrentOccurrencePtr = OccurrenceData->m_pPrototype;
 	}
 
 	if(OccurrenceData->m_pPart == nullptr && OccurrenceData->m_uiPOccurrencesSize == 0)
@@ -681,7 +682,7 @@ FCadId FTechSoftFileParser::TraverseOccurrence(const A3DAsmProductOccurrence* Oc
 	}
 
 	// Add Occurrence's Children
-	OccurrenceData.FillFrom(CurrentOccurrencePtr);
+	OccurrenceData.FillFrom(CachedOccurrencePtr);
 	while (OccurrenceData->m_pPrototype != nullptr && OccurrenceData->m_uiPOccurrencesSize == 0)
 	{
 		OccurrenceData.FillFrom(OccurrenceData->m_pPrototype);
