Techsoft: fix import material

diff --git a/Engine/Plugins/Enterprise/DatasmithCADImporter/Source/CADInterfaces/Private/TUniqueTechSoftObj.h b/Engine/Plugins/Enterprise/DatasmithCADImporter/Source/CADInterfaces/Private/TUniqueTechSoftObj.h
--- a/Engine/Plugins/Enterprise/DatasmithCADImporter/Source/CADInterfaces/Private/TUniqueTechSoftObj.h
+++ b/Engine/Plugins/Enterprise/DatasmithCADImporter/Source/CADInterfaces/Private/TUniqueTechSoftObj.h
@@ -40,6 +40,7 @@ public:
 	explicit TUniqueTSObjBase()
 	{
 		InitializeData();
+		Status = A3DStatus::A3D_SUCCESS;
 	}
 
 	/**
@@ -51,7 +52,7 @@ public:
 		//TechSoftInterfaceImpl::InitializeData(Data);
 		InitializeData();
 
-		Status = GetData(DataPtr);
+		Status = FillFrom(DataPtr);
 	}
 
 	~TUniqueTSObjBase()
@@ -64,44 +65,41 @@ public:
 	 */
 	A3DStatus FillFrom(IndexerType DataPtr)
 	{
-		if (IsValid())
-		{
-			Status = ResetData();
-		}
-		else
+		if (bHasData)
 		{
-			Status = A3DStatus::A3D_SUCCESS;
+			ResetData();
 		}
 
-		if (!IsValid() || (DataPtr == DefaultValue))
+		if (DataPtr == DefaultValue)
 		{
 			Status = A3DStatus::A3D_ERROR;
-			return Status;
 		}
-
-		Status = GetData(DataPtr);
+		else
+		{
+			Status = GetData(DataPtr);
+			bHasData = true;
+		}
 		return Status;
 	}
 
 	template<typename... InArgTypes>
-	A3DStatus FillWith(A3DStatus (*Getter)(const A3DEntity*, ObjectType*, InArgTypes&&... ), const A3DEntity* DataPtr, InArgTypes&&... Args)
+	A3DStatus FillWith(A3DStatus(*Getter)(const A3DEntity*, ObjectType*, InArgTypes&&...), const A3DEntity* DataPtr, InArgTypes&&... Args)
 	{
-		if (IsValid())
+		if (bHasData)
 		{
-			Status = ResetData();
-		}
-		else
-		{
-			Status = A3DStatus::A3D_SUCCESS;
+			ResetData();
 		}
 
-		if (!IsValid() || (DataPtr == DefaultValue))
+		if (DataPtr == DefaultValue)
 		{
 			Status = A3DStatus::A3D_ERROR;
-			return Status;
+		}
+		else
+		{
+			Status = Getter(DataPtr, &Data, Forward<InArgTypes>(Args)...);
+			bHasData = true;
 		}
 
-		Status = Getter(DataPtr, &Data, Forward<InArgTypes>(Args)...);
 		return Status;
 	}
 
@@ -172,6 +170,7 @@ public:
 
 private:
 	ObjectType Data;
+	bool bHasData = false;
 	A3DStatus Status = A3DStatus::A3D_ERROR;
 
 	/**
@@ -199,9 +198,10 @@ private:
 	}
 #endif
 
-	A3DStatus ResetData()
+	void ResetData()
 	{
-		return GetData(DefaultValue);
+		GetData(DefaultValue);
+		bHasData = false;
 	}
 };
 
@@ -211,8 +211,6 @@ using TUniqueTSObj = TUniqueTSObjBase<ObjectType, const A3DEntity*>;
 template<class ObjectType>
 using TUniqueTSObjFromIndex = TUniqueTSObjBase<ObjectType, uint32>;
 
-template<class ObjectType, class IndexerType>
-IndexerType TUniqueTSObjBase<ObjectType, IndexerType>::DefaultValue = (IndexerType) nullptr;
 
 }
 
diff --git a/Engine/Plugins/Enterprise/DatasmithCADImporter/Source/CADInterfaces/Private/TechSoftFileParser.cpp b/Engine/Plugins/Enterprise/DatasmithCADImporter/Source/CADInterfaces/Private/TechSoftFileParser.cpp
--- a/Engine/Plugins/Enterprise/DatasmithCADImporter/Source/CADInterfaces/Private/TechSoftFileParser.cpp
+++ b/Engine/Plugins/Enterprise/DatasmithCADImporter/Source/CADInterfaces/Private/TechSoftFileParser.cpp
@@ -328,7 +328,7 @@ void FTechSoftFileParser::ReserveCADFileData()
 	FArchiveSceneGraph& SceneGraphArchive = CADFileData.GetSceneGraphArchive();
 	SceneGraphArchive.Reserve(ComponentCount[EComponentType::Occurrence], ComponentCount[EComponentType::Reference], ComponentCount[EComponentType::Body]);
 
-	uint32 MaterialNum = CountMaterial() + CountColor();
+	uint32 MaterialNum = CountColorAndMaterial();
 	SceneGraphArchive.MaterialHIdToMaterial.Reserve(MaterialNum);
 }
 
@@ -599,7 +599,14 @@ FArchiveBody& FTechSoftFileParser::AddBody(FEntityMetaData& BodyMetaData)
 	int32 BodyIndex = CADFileData.AddBody(BodyId);
 	FArchiveBody& Body = CADFileData.GetBodyAt(BodyIndex);
 	Body.MetaData = MoveTemp(BodyMetaData.MetaData);
-
+	if(BodyMetaData.ColorName != 0)
+	{
+		Body.ColorFaceSet.Add(BodyMetaData.ColorName);
+	}
+	if (BodyMetaData.MaterialName != 0)
+	{
+		Body.MaterialFaceSet.Add(BodyMetaData.MaterialName);
+	}
 	return Body;
 }
 
@@ -634,7 +641,6 @@ FCadId FTechSoftFileParser::TraverseOccurrence(const A3DAsmProductOccurrence* Oc
 	BuildInstanceName(InstanceMetaData.MetaData);
 
 	ExtractMaterialProperties(OccurrencePtr);
-	ExtractLayer(OccurrencePtr);
 
 	FArchiveInstance& Instance = AddInstance(InstanceMetaData);
 
@@ -915,7 +921,6 @@ void FTechSoftFileParser::TraversePartDefinition(const A3DAsmPartDefinition* Par
 	BuildPartName(PartMetaData.MetaData);
 
 	ExtractMaterialProperties(PartDefinitionPtr);
-	ExtractLayer(PartDefinitionPtr);
 
 	TUniqueTSObj<A3DAsmPartDefinitionData> PartData(PartDefinitionPtr);
 	if (PartData.IsValid())
@@ -1410,24 +1415,35 @@ FArchiveColor& FTechSoftFileParser::FindOrAddColor(uint32 ColorIndex, uint8 Alph
 	}
 
 	FArchiveColor& NewColor = CADFileData.AddColor(ColorHId);
-	FColor& Color = NewColor.Color;
-	TUniqueTSObjFromIndex<A3DGraphRgbColorData> ColorData(ColorIndex);
-	if (ColorData.IsValid())
-	{
-		Color.R = (uint8)(ColorData->m_dRed * 255);
-		Color.G = (uint8)(ColorData->m_dGreen * 255);
-		Color.B = (uint8)(ColorData->m_dBlue * 255);
-		Color.A = Alpha;
-	}
-	else
-	{
-		Color = FColor(200, 200, 200);
-	}
-
+	NewColor.Color = TechSoftFileParserImpl::GetColorAt(ColorIndex);
+	NewColor.Color.A = Alpha;
+	
 	NewColor.UEMaterialName = BuildColorName(NewColor.Color);
 	return NewColor;
 }
 
+FArchiveMaterial& FTechSoftFileParser::AddMaterialAt(uint32 MaterialIndexToSave, uint32 GraphMaterialIndex)
+{
+	FArchiveMaterial& NewMaterial = CADFileData.AddMaterial(GraphMaterialIndex);
+	FCADMaterial& Material = NewMaterial.Material;
+
+	TUniqueTSObjFromIndex<A3DGraphMaterialData> MaterialData(MaterialIndexToSave);
+	if(MaterialData.IsValid())
+	{
+		Material.Diffuse = TechSoftFileParserImpl::GetColorAt(MaterialData->m_uiDiffuse);
+		Material.Ambient = TechSoftFileParserImpl::GetColorAt(MaterialData->m_uiAmbient);
+		Material.Specular = TechSoftFileParserImpl::GetColorAt(MaterialData->m_uiSpecular);
+		Material.Shininess = MaterialData->m_dShininess;
+		Material.Transparency = 1. - MaterialData->m_dAmbientAlpha;
+		// todo: find how to convert Emissive color into ? reflexion coef...
+		// Material.Emissive = GetColor(MaterialData->m_uiEmissive);
+		// Material.Reflexion;
+	}
+	NewMaterial.UEMaterialName = BuildMaterialName(Material);
+	return NewMaterial;
+}
+
+
 FArchiveMaterial& FTechSoftFileParser::FindOrAddMaterial(uint32 MaterialIndex)
 {
 	if (FArchiveMaterial* MaterialArchive = CADFileData.FindMaterial(MaterialIndex))
@@ -1435,52 +1451,29 @@ FArchiveMaterial& FTechSoftFileParser::FindOrAddMaterial(uint32 MaterialIndex)
 		return *MaterialArchive;
 	}
 
-	TUniqueTSObjFromIndex<A3DGraphRgbColorData> ColorData;
-	TFunction<const FColor(uint32)> GetColor = [&](uint32 ColorIndex) -> const FColor
-	{
-		ColorData.FillFrom(ColorIndex);
-		if (ColorData.IsValid())
-		{
-			return FColor((uint8)ColorData->m_dRed * 255,
-				(uint8)ColorData->m_dGreen * 255,
-				(uint8)ColorData->m_dBlue * 255);
-		}
-		else
-		{
-			return FColor(200, 200, 200);
-		}
-	};
-
-	FArchiveMaterial& NewMaterial = CADFileData.AddMaterial(MaterialIndex);
-	FCADMaterial& Material = NewMaterial.Material;
-
-	A3DBool bIsTexture = false;
-	A3DGlobalIsMaterialTexture(MaterialIndex, &bIsTexture);
+	bool bIsTexture = TechSoftUtils::IsMaterialTexture(MaterialIndex);
 	if (bIsTexture)
 	{
-#ifdef NOTYETDEFINE
-		// style is a texture
-		TUniqueTSObj<A3DGraphTextureDefinitionData> TextureDefinitionData(TextureIndex);
-		if (TextureDefinitionData.IsValid())
+		TUniqueTSObjFromIndex<A3DGraphTextureApplicationData> TextureData(MaterialIndex);
+		if (TextureData.IsValid())
 		{
-			TUniqueTSObj<A3DGraphPictureData> PictureData(TextureDefinitionData->m_uiPictureIndex);
-		}
+			TextureData->m_uiMaterialIndex;
+			return AddMaterialAt(TextureData->m_uiMaterialIndex, MaterialIndex);
+			
+#ifdef NOTYETDEFINE
+			TUniqueTSObj<A3DGraphTextureDefinitionData> TextureDefinitionData(TextureData->m_uiTextureDefinitionIndex);
+			if (TextureDefinitionData.IsValid())
+			{
+				TUniqueTSObj<A3DGraphPictureData> PictureData(TextureDefinitionData->m_uiPictureIndex);
+			}
 #endif
+		}
+		return AddMaterialAt(MaterialIndex, 0);
 	}
 	else
 	{
-		TUniqueTSObjFromIndex<A3DGraphMaterialData> MaterialData(MaterialIndex);
-		Material.Diffuse = GetColor(MaterialData->m_uiDiffuse);
-		Material.Ambient = GetColor(MaterialData->m_uiAmbient);
-		Material.Specular = GetColor(MaterialData->m_uiSpecular);
-		Material.Shininess = MaterialData->m_dShininess;
-		Material.Transparency = MaterialData->m_dAmbientAlpha;
-		// todo: find how to convert Emissive color into ? reflexion coef...
-		// Material.Emissive = GetColor(MaterialData->m_uiEmissive);
-		// Material.Reflexion;
+		return AddMaterial(MaterialIndex);
 	}
-
-	return NewMaterial;
 }
 
 void FTechSoftFileParser::ExtractGraphicProperties(const A3DGraphics* Graphics, FEntityMetaData& OutMetaData)
@@ -1490,31 +1483,38 @@ void FTechSoftFileParser::ExtractGraphicProperties(const A3DGraphics* Graphics,
 	{
 		return;
 	}
-	FEntityBehaviour GraphicsBehaviour;
-
-	GraphicsBehaviour.bFatherHeritColor = GraphicsData->m_usBehaviour & kA3DGraphicsFatherHeritColor;
-	GraphicsBehaviour.bFatherHeritLayer = GraphicsData->m_usBehaviour & kA3DGraphicsFatherHeritLayer;
-	GraphicsBehaviour.bFatherHeritLinePattern = GraphicsData->m_usBehaviour & kA3DGraphicsFatherHeritLinePattern;
-	GraphicsBehaviour.bFatherHeritLineWidth = GraphicsData->m_usBehaviour & kA3DGraphicsFatherHeritLineWidth;
-	GraphicsBehaviour.bFatherHeritShow = GraphicsData->m_usBehaviour & kA3DGraphicsFatherHeritShow;
-	GraphicsBehaviour.bFatherHeritTransparency = GraphicsData->m_usBehaviour & kA3DGraphicsFatherHeritTransparency;
-	GraphicsBehaviour.bRemoved = GraphicsData->m_usBehaviour & kA3DGraphicsRemoved;
-	GraphicsBehaviour.bShow = GraphicsData->m_usBehaviour & kA3DGraphicsShow;
-	GraphicsBehaviour.bSonHeritColor = GraphicsData->m_usBehaviour & kA3DGraphicsSonHeritColor;
-	GraphicsBehaviour.bSonHeritLayer = GraphicsData->m_usBehaviour & kA3DGraphicsSonHeritLayer;
-	GraphicsBehaviour.bSonHeritLinePattern = GraphicsData->m_usBehaviour & kA3DGraphicsSonHeritLinePattern;
-	GraphicsBehaviour.bSonHeritLineWidth = GraphicsData->m_usBehaviour & kA3DGraphicsSonHeritLineWidth;
-	GraphicsBehaviour.bSonHeritShow = GraphicsData->m_usBehaviour & kA3DGraphicsSonHeritShow;
-	GraphicsBehaviour.bSonHeritTransparency = GraphicsData->m_usBehaviour & kA3DGraphicsSonHeritTransparency;
-
-	OutMetaData.bRemoved = GraphicsBehaviour.bRemoved;
-	OutMetaData.bShow = GraphicsBehaviour.bShow;
-
-	FCADUUID ColorName;
-	FCADUUID MaterialName;
+
+	bool bFatherHeritColor = GraphicsData->m_usBehaviour & kA3DGraphicsFatherHeritColor;
+	bool bSonHeritColor = GraphicsData->m_usBehaviour & kA3DGraphicsSonHeritColor;
+
+	//To do if needed
+	//bool bFatherHeritLayer = GraphicsData->m_usBehaviour & kA3DGraphicsFatherHeritLayer;
+	//bool bSonHeritLayer = GraphicsData->m_usBehaviour & kA3DGraphicsSonHeritLayer;
+	//bool bFatherHeritTransparency = GraphicsData->m_usBehaviour & kA3DGraphicsFatherHeritTransparency;
+	//bool bSonHeritTransparency = GraphicsData->m_usBehaviour & kA3DGraphicsSonHeritTransparency;
+
+	OutMetaData.bRemoved = GraphicsData->m_usBehaviour & kA3DGraphicsRemoved;
+	OutMetaData.bShow = GraphicsData->m_usBehaviour & kA3DGraphicsShow;
+
+	if (GraphicsData->m_uiStyleIndex == A3D_DEFAULT_STYLE_INDEX)
+	{
+		return;
+	}
+
+	FCADUUID& ColorName = OutMetaData.ColorName;
+	FCADUUID& MaterialName = OutMetaData.MaterialName;
 
 	ExtractGraphStyleProperties(GraphicsData->m_uiStyleIndex, ColorName, MaterialName);
 
+	if(bSonHeritColor)
+	{
+		OutMetaData.MetaData.Add(TEXT("GraphicsBehaviour"), TEXT("SonHerit"));
+	}
+	else if (bFatherHeritColor)
+	{
+		OutMetaData.MetaData.Add(TEXT("GraphicsBehaviour"), TEXT("FatherHerit"));
+	}
+
 	if (ColorName)
 	{
 		OutMetaData.MetaData.Add(TEXT("ColorName"), FString::Printf(TEXT("%u"), ColorName));
@@ -1528,9 +1528,6 @@ void FTechSoftFileParser::ExtractGraphicProperties(const A3DGraphics* Graphics,
 
 void FTechSoftFileParser::ExtractGraphStyleProperties(uint32 StyleIndex, FCADUUID& OutColorName, FCADUUID& OutMaterialName)
 {
-	OutColorName = 0;
-	OutMaterialName = 0;
-
 	TUniqueTSObjFromIndex<A3DGraphStyleData> GraphStyleData(StyleIndex);
 
 	if (GraphStyleData.IsValid())
@@ -1667,26 +1664,6 @@ FMatrix FTechSoftFileParser::TraverseCoordinateSystem(const A3DRiCoordinateSyste
 	return FMatrix::Identity;
 }
 
-void FTechSoftFileParser::ExtractLayer(const A3DAsmProductOccurrence* Occurrence)
-{
-	// TODO
-
-	A3DUns32 LayerCount = 0;
-	A3DAsmLayer* AsmLayer = 0;
-	if (A3DAsmProductOccurrenceGetLayerList(Occurrence, &LayerCount, &AsmLayer) == A3D_SUCCESS)
-	{
-		if (LayerCount)
-		{
-			for (A3DUns32 Index = 0; Index < LayerCount; ++Index)
-			{
-				//Layer->SetAttribute("Name", asLayers[Index].m_pcLayerName ? asLayers[Index].m_pcLayerName : "null");
-				//Layer->SetAttribute("Layer", asLayers[Index].m_usLayer);
-			}
-		}
-		A3DAsmProductOccurrenceGetLayerList(NULL, &LayerCount, &AsmLayer);
-	}
-}
-
 bool FTechSoftFileParser::IsConfigurationSet(const A3DAsmProductOccurrence* Occurrence)
 {
 	TUniqueTSObj<A3DAsmProductOccurrenceData> OccurrenceData(Occurrence);
@@ -1714,23 +1691,123 @@ bool FTechSoftFileParser::IsConfigurationSet(const A3DAsmProductOccurrence* Occu
 	return bIsConfiguration;
 }
 
-int32 FTechSoftFileParser::CountMaterial()
+
+uint32 FTechSoftFileParser::CountColorAndMaterial()
 {
-	// TODO
+	A3DGlobal* GlobalPtr = nullptr;
+	if (TechSoftUtils::GetGlobalPointer(&GlobalPtr) != A3D_SUCCESS)
+	{
+		return 0;
+	}
 
-	return 0;
+	A3DInt32 iRet = A3D_SUCCESS;
+	TUniqueTSObj<A3DGlobalData> GlobalData(GlobalPtr);
+	if (!GlobalData.IsValid())
+	{
+		return 0;
+	}
+
+	uint32 ColorCount = GlobalData->m_uiColorsSize;
+	uint32 MaterialCount = GlobalData->m_uiMaterialsSize;
+	uint32 TextureDefinitionCount = GlobalData->m_uiTextureDefinitionsSize;
+
+	return ColorCount + MaterialCount + TextureDefinitionCount;
 }
 
-int32 FTechSoftFileParser::CountColor()
+void ExtractTextureDefinition(const A3DGraphTextureDefinitionData& TextureDefinitionData)
 {
-	// TODO
+	// To do
+	//TextureDefinitionData.m_uiPictureIndex;
+	//TextureDefinitionData.m_ucTextureDimension;
+	//TextureDefinitionData.m_eMappingType;
+	//TextureDefinitionData.m_eMappingOperator;
 
-	return 0;
+	//TextureDefinitionData.m_pOperatorTransfo;
+
+	//TextureDefinitionData.m_uiMappingAttributes;
+	//TextureDefinitionData.m_uiMappingAttributesIntensitySize,
+	//TextureDefinitionData.m_pdMappingAttributesIntensity;
+	//TextureDefinitionData.m_uiMappingAttributesComponentsSize,
+	//TextureDefinitionData.m_pucMappingAttributesComponents;
+	//TextureDefinitionData.m_eTextureFunction;
+	//TextureDefinitionData.m_dRed;
+	//TextureDefinitionData.m_dGreen;
+	//TextureDefinitionData.m_dBlue;
+	//TextureDefinitionData.m_dAlpha;
+	//TextureDefinitionData.m_eBlend_src_RGB;
+	//TextureDefinitionData.m_eBlend_dst_RGB;
+	//TextureDefinitionData.m_eBlend_src_Alpha;
+	//TextureDefinitionData.m_eBlend_dst_Alpha;
+	//TextureDefinitionData.m_ucTextureApplyingMode;
+	//TextureDefinitionData.m_eTextureAlphaTest;
+	//TextureDefinitionData.m_dAlphaTestReference;
+	//TextureDefinitionData.m_eTextureWrappingModeS;
+	//TextureDefinitionData.m_eTextureWrappingModeT;
+
+	//if (TextureDefinitionData.m_pTextureTransfo != nullptr)
+	//{
+	//	TUniqueTSObj<A3DGraphTextureTransformationData> TransfoData(TextureDefinitionData.m_pTextureTransfo);
+	//}
 }
 
 void FTechSoftFileParser::ReadMaterialsAndColors()
 {
-	// TODO
+	A3DGlobal* GlobalPtr = nullptr;
+	if(TechSoftUtils::GetGlobalPointer(&GlobalPtr) != A3D_SUCCESS)
+	{
+		return;
+	}
+
+	A3DInt32 iRet = A3D_SUCCESS;
+	TUniqueTSObj<A3DGlobalData> GlobalData(GlobalPtr);
+	if (!GlobalData.IsValid())
+	{
+		return;
+	}
+
+	{
+		uint32 MaterialCount = GlobalData->m_uiMaterialsSize;
+		if (MaterialCount)
+		{
+			for (uint32 MaterialIndex = 0; MaterialIndex < MaterialCount; ++MaterialIndex)
+			{
+				FTechSoftFileParser::FindOrAddMaterial(MaterialIndex);
+			}
+		}
+	}
+
+	{
+		uint32 TextureDefinitionCount = GlobalData->m_uiTextureDefinitionsSize;
+		if (TextureDefinitionCount)
+		{
+			TUniqueTSObjFromIndex<A3DGraphTextureDefinitionData> TextureDefinitionData;
+			for (uint32 TextureIndex = 0; TextureIndex < TextureDefinitionCount; ++TextureIndex)
+			{
+				TextureDefinitionData.FillFrom(TextureIndex);
+				ExtractTextureDefinition(*TextureDefinitionData);
+			}
+		}
+	}
+	
+	{
+		uint32 PictureCount = GlobalData->m_uiPicturesSize;
+		if (PictureCount)
+		{
+			TUniqueTSObjFromIndex<A3DGraphPictureData> PictureData;
+			for (uint32 PictureIndex = 0; PictureIndex < PictureCount; ++PictureIndex)
+			{
+				A3DEntity* PicturePtr = TechSoftUtils::GetPointerFromIndex(PictureIndex, kA3DTypeGraphPicture);
+				if (PicturePtr)
+				{
+					FEntityMetaData PictureMetaData;
+					ExtractMetaData(PicturePtr, PictureMetaData);
+				}
+
+				PictureData.FillFrom(PictureIndex);
+				// To do
+			}
+		}
+	}
 }
 
 void FTechSoftFileParser::MeshRepresentationWithTechSoft(A3DRiRepresentationItem* RepresentationItemPtr, FArchiveBody& Body)
@@ -1749,13 +1826,12 @@ void FTechSoftFileParser::MeshRepresentationWithTechSoft(A3DRiRepresentationItem
 	{
 		// TODO
 		// Check unity conversion
-		ensure(false);
-
 		const FImportParameters& ImportParameters = CADFileData.GetImportParameters();
+
 		////A3DRWParamsTessellationData /*!< The tessellation reading parameters. */
 		TessellationParameters.m_eTessellationLevelOfDetail = kA3DTessLODUserDefined;		/*!< Enum to specify predefined values for some following members. */
 		TessellationParameters.m_bUseHeightInsteadOfRatio = A3D_TRUE;
-		TessellationParameters.m_dMaxChordHeight = ImportParameters.GetChordTolerance();
+		TessellationParameters.m_dMaxChordHeight = ImportParameters.GetChordTolerance() * 10. / FileUnit;
 		TessellationParameters.m_dAngleToleranceDeg = ImportParameters.GetMaxNormalAngle();
 		TessellationParameters.m_dMaximalTriangleEdgeLength = 0; //ImportParameters.MaxEdgeLength;
 
@@ -2255,8 +2331,6 @@ void AddFaceTriangleStripWithTexture(FTessellationData& Tessellation, const A3DT
 
 } //ns FTechSoftFileParserImpl
 
-
-#pragma optimize ("",off)
 void FTechSoftFileParser::TraverseTessellation3D(const A3DTess3D* TessellationPtr, FArchiveBody& Body)
 {
 	FBodyMesh& BodyMesh = CADFileData.AddBodyMesh(Body.ObjectId, Body);
@@ -2266,6 +2340,17 @@ void FTechSoftFileParser::TraverseTessellation3D(const A3DTess3D* TessellationPt
 	const int TessellationFaceDataWithStrip = 0x8888;
 	const int TessellationFaceDataWithOneNormal = 0xE0E0;
 
+	FCADUUID DefaultColorName = 0;
+	FCADUUID DefaultMaterialName = 0;
+	if (Body.ColorFaceSet.Num() > 0)
+	{
+		DefaultColorName = *Body.ColorFaceSet.begin();
+	}
+	if (Body.MaterialFaceSet.Num() > 0)
+	{
+		DefaultMaterialName = *Body.MaterialFaceSet.begin();
+	}
+
 	// Coordinates
 	TUniqueTSObj<A3DTessBaseData> TessellationBaseData(TessellationPtr);
 	{
@@ -2300,23 +2385,24 @@ void FTechSoftFileParser::TraverseTessellation3D(const A3DTess3D* TessellationPt
 			const A3DTessFaceData& FaceTessData = Tessellation3DData->m_psFaceTessData[Index];
 			FTessellationData& Tessellation = BodyMesh.Faces.Emplace_GetRef();
 
+			FCADUUID ColorName = DefaultColorName;
+			FCADUUID MaterialName = DefaultMaterialName;
 			if (FaceTessData.m_uiStyleIndexesSize == 1)
 			{
 				A3DUns32 StyleIndex = FaceTessData.m_puiStyleIndexes[0];
-				FCADUUID ColorName;
-				FCADUUID MaterialName;
 				ExtractGraphStyleProperties(StyleIndex, ColorName, MaterialName);
-				if (ColorName)
-				{
-					Tessellation.ColorName = ColorName;
-					BodyMesh.ColorSet.Add(ColorName);
-				}
+			}
 
-				if (MaterialName)
-				{
-					Tessellation.MaterialName = MaterialName;
-					BodyMesh.MaterialSet.Add(MaterialName);
-				}
+			if (ColorName)
+			{
+				Tessellation.ColorName = ColorName;
+				BodyMesh.ColorSet.Add(ColorName);
+			}
+
+			if (MaterialName)
+			{
+				Tessellation.MaterialName = MaterialName;
+				BodyMesh.MaterialSet.Add(MaterialName);
 			}
 
 			uint32 TriangleCount = TechSoftFileParserImpl::CountTriangles(FaceTessData);
@@ -2474,7 +2560,6 @@ void FTechSoftFileParser::TraverseTessellation3D(const A3DTess3D* TessellationPt
 	Body.ColorFaceSet = BodyMesh.ColorSet;
 	Body.MaterialFaceSet = BodyMesh.MaterialSet;
 }
-#pragma optimize ("",on)
 
 } // ns CADLibrary
 
diff --git a/Engine/Plugins/Enterprise/DatasmithCADImporter/Source/CADInterfaces/Private/TechSoftFileParser.h b/Engine/Plugins/Enterprise/DatasmithCADImporter/Source/CADInterfaces/Private/TechSoftFileParser.h
--- a/Engine/Plugins/Enterprise/DatasmithCADImporter/Source/CADInterfaces/Private/TechSoftFileParser.h
+++ b/Engine/Plugins/Enterprise/DatasmithCADImporter/Source/CADInterfaces/Private/TechSoftFileParser.h
@@ -102,30 +102,14 @@ enum EModellerType : uint32
 	ModellerLast
 };
 
-struct FEntityBehaviour
-{
-	bool bFatherHeritColor = false;
-	bool bFatherHeritLayer = false;
-	bool bFatherHeritLinePattern = false;
-	bool bFatherHeritLineWidth = false;
-	bool bFatherHeritShow = false;
-	bool bFatherHeritTransparency = false;
-	bool bRemoved = false;
-	bool bShow = true;
-	bool bSonHeritColor = false;
-	bool bSonHeritLayer = false;
-	bool bSonHeritLinePattern = false;
-	bool bSonHeritLineWidth = false;
-	bool bSonHeritShow = false;
-	bool bSonHeritTransparency = false;
-};
-
 struct FEntityMetaData
 {
 	TMap<FString, FString> MetaData;
 	bool bRemoved = false;
 	bool bShow = true;
 	bool bUnloaded = false;
+	FCADUUID ColorName = 0;
+	FCADUUID MaterialName = 0;
 	FFileDescriptor ExternalFile;
 };
 
@@ -167,8 +151,7 @@ private:
 	void ReserveCADFileData();
 
 	// Materials and colors
-	int32 CountMaterial();
-	int32 CountColor();
+	uint32 CountColorAndMaterial();
 	void ReadMaterialsAndColors();
 
 	// Traverse ASM tree by starting from the model
@@ -203,11 +186,25 @@ private:
 
 	// Graphic properties
 	void ExtractGraphicProperties(const A3DGraphics* Graphics, FEntityMetaData& OutMetaData);
+
+	/**
+	 * ColorName and MaterialName have to be initialized before.
+	 * This method update the value ColorName or MaterialName accordingly of the GraphStyleData type (material or color)
+	 */
 	void ExtractGraphStyleProperties(uint32 StyleIndex, FCADUUID& ColorName, FCADUUID& MaterialName);
 	void ExtractMaterialProperties(const A3DEntity* Entity);
-	void ExtractLayer(const A3DAsmProductOccurrence* Occurrence);
 	FArchiveColor& FindOrAddColor(uint32 ColorIndex, uint8 Alpha);
-	FArchiveMaterial& FindOrAddMaterial(uint32 MaterialIndex);
+	FArchiveMaterial& FindOrAddMaterial(uint32 MaterialId);
+
+	/**
+	 * @param GraphMaterialIndex is the techsoft index of the graphic data
+	 * @param MaterialIndexToSave is the index of the material really saved (i.e. for texture, at the texture index, with saved the material used by the texture)
+	 */
+	FArchiveMaterial& AddMaterialAt(uint32 MaterialIndexToSave, uint32 GraphMaterialIndex);
+	FArchiveMaterial& AddMaterial(uint32 MaterialIndex)
+	{
+		return AddMaterialAt(MaterialIndex, MaterialIndex);
+	}
 
 	// Transform
 	FMatrix TraverseCoordinateSystem(const A3DRiCoordinateSystem* CoordinateSystem);
@@ -291,6 +288,23 @@ inline void Reserve(CADLibrary::FTessellationData& Tessellation, int32 InTrinang
 	}
 };
 
+#ifdef USE_TECHSOFT_SDK
+inline FColor GetColorAt(uint32 ColorIndex)
+{
+	TUniqueTSObjFromIndex<A3DGraphRgbColorData> ColorData(ColorIndex);
+	if (ColorData.IsValid())
+	{
+		return FColor((uint8)(ColorData->m_dRed * 255)
+			, (uint8)(ColorData->m_dGreen * 255)
+			, (uint8)(ColorData->m_dBlue * 255));
+	}
+	else
+	{
+		return FColor(200, 200, 200);
+	}
+}
+#endif
+
 } // ns TechSoftFileParserImpl
 
 } // ns CADLibrary
\ No newline at end of file
diff --git a/Engine/Plugins/Enterprise/DatasmithCADImporter/Source/CADInterfaces/Private/TechSoftInterface.cpp b/Engine/Plugins/Enterprise/DatasmithCADImporter/Source/CADInterfaces/Private/TechSoftInterface.cpp
--- a/Engine/Plugins/Enterprise/DatasmithCADImporter/Source/CADInterfaces/Private/TechSoftInterface.cpp
+++ b/Engine/Plugins/Enterprise/DatasmithCADImporter/Source/CADInterfaces/Private/TechSoftInterface.cpp
@@ -92,331 +92,421 @@ void TUniqueTSObj<A3DAsmModelFileData>::InitializeData()
 {
 	A3D_INITIALIZE_DATA(A3DAsmModelFileData, Data);
 }
+const A3DEntity* TUniqueTSObj<A3DAsmModelFileData>::DefaultValue = nullptr;
 
 void TUniqueTSObj<A3DAsmPartDefinitionData>::InitializeData()
 {
 	A3D_INITIALIZE_DATA(A3DAsmPartDefinitionData, Data);
 }
+const A3DEntity* TUniqueTSObj<A3DAsmPartDefinitionData>::DefaultValue = nullptr;
 
 void TUniqueTSObj<A3DAsmProductOccurrenceData>::InitializeData()
 {
 	A3D_INITIALIZE_DATA(A3DAsmProductOccurrenceData, Data);
 }
+const A3DEntity* TUniqueTSObj<A3DAsmProductOccurrenceData>::DefaultValue = nullptr;
 
 void TUniqueTSObj<A3DAsmProductOccurrenceDataCV5>::InitializeData()
 {
 	A3D_INITIALIZE_DATA(A3DAsmProductOccurrenceDataCV5, Data);
 }
+const A3DEntity* TUniqueTSObj<A3DAsmProductOccurrenceDataCV5>::DefaultValue = nullptr;
 
 void TUniqueTSObj<A3DAsmProductOccurrenceDataSLW>::InitializeData()
 {
 	A3D_INITIALIZE_DATA(A3DAsmProductOccurrenceDataSLW, Data);
 }
+const A3DEntity* TUniqueTSObj<A3DAsmProductOccurrenceDataSLW>::DefaultValue = nullptr;
 
 void TUniqueTSObj<A3DAsmProductOccurrenceDataUg>::InitializeData()
 {
 	A3D_INITIALIZE_DATA(A3DAsmProductOccurrenceDataUg, Data);
 }
+const A3DEntity* TUniqueTSObj<A3DAsmProductOccurrenceDataUg>::DefaultValue = nullptr;
 
 void TUniqueTSObj<A3DBoundingBoxData>::InitializeData()
 {
 	A3D_INITIALIZE_DATA(A3DBoundingBoxData, Data);
 }
+const A3DEntity* TUniqueTSObj<A3DBoundingBoxData>::DefaultValue = nullptr;
 
 void TUniqueTSObj<A3DCrvCircleData>::InitializeData()
 {
 	A3D_INITIALIZE_DATA(A3DCrvCircleData, Data);
 }
+const A3DEntity* TUniqueTSObj<A3DCrvCircleData>::DefaultValue = nullptr;
 
 void TUniqueTSObj<A3DCrvCompositeData>::InitializeData()
 {
 	A3D_INITIALIZE_DATA(A3DCrvCompositeData, Data);
 }
+const A3DEntity* TUniqueTSObj<A3DCrvCompositeData>::DefaultValue = nullptr;
 
 void TUniqueTSObj<A3DCrvEllipseData>::InitializeData()
 {
 	A3D_INITIALIZE_DATA(A3DCrvEllipseData, Data);
 }
+const A3DEntity* TUniqueTSObj<A3DCrvEllipseData>::DefaultValue = nullptr;
 
 void TUniqueTSObj<A3DCrvHelixData>::InitializeData()
 {
 	A3D_INITIALIZE_DATA(A3DCrvHelixData, Data);
 }
+const A3DEntity* TUniqueTSObj<A3DCrvHelixData>::DefaultValue = nullptr;
 
 void TUniqueTSObj<A3DCrvHyperbolaData>::InitializeData()
 {
 	A3D_INITIALIZE_DATA(A3DCrvHyperbolaData, Data);
 }
+const A3DEntity* TUniqueTSObj<A3DCrvHyperbolaData>::DefaultValue = nullptr;
 
 void TUniqueTSObj<A3DCrvLineData>::InitializeData()
 {
 	A3D_INITIALIZE_DATA(A3DCrvLineData, Data);
 }
+const A3DEntity* TUniqueTSObj<A3DCrvLineData>::DefaultValue = nullptr;
 
 void TUniqueTSObj<A3DCrvNurbsData>::InitializeData()
 {
 	A3D_INITIALIZE_DATA(A3DCrvNurbsData, Data);
 }
+const A3DEntity* TUniqueTSObj<A3DCrvNurbsData>::DefaultValue = nullptr;
 
 void TUniqueTSObj<A3DCrvParabolaData>::InitializeData()
 {
 	A3D_INITIALIZE_DATA(A3DCrvParabolaData, Data);
 }
+const A3DEntity* TUniqueTSObj<A3DCrvParabolaData>::DefaultValue = nullptr;
 
 void TUniqueTSObj<A3DCrvPolyLineData>::InitializeData()
 {
 	A3D_INITIALIZE_DATA(A3DCrvPolyLineData, Data);
 }
+const A3DEntity* TUniqueTSObj<A3DCrvPolyLineData>::DefaultValue = nullptr;
+
+void TUniqueTSObj<A3DGlobalData>::InitializeData()
+{
+	A3D_INITIALIZE_DATA(A3DGlobalData, Data);
+}
+const A3DEntity* TUniqueTSObj<A3DGlobalData>::DefaultValue = nullptr;
 
 void TUniqueTSObj<A3DGraphicsData>::InitializeData()
 {
 	A3D_INITIALIZE_DATA(A3DGraphicsData, Data);
 }
+const A3DEntity* TUniqueTSObj<A3DGraphicsData>::DefaultValue = nullptr;
 
 void TUniqueTSObjFromIndex<A3DGraphMaterialData>::InitializeData()
 {
 	A3D_INITIALIZE_DATA(A3DGraphMaterialData, Data);
 }
+uint32 TUniqueTSObjFromIndex<A3DGraphMaterialData>::DefaultValue = A3D_DEFAULT_MATERIAL_INDEX;
+
+void TUniqueTSObjFromIndex<A3DGraphPictureData>::InitializeData()
+{
+	A3D_INITIALIZE_DATA(A3DGraphPictureData, Data);
+}
+uint32 TUniqueTSObjFromIndex<A3DGraphPictureData>::DefaultValue = A3D_DEFAULT_PICTURE_INDEX;
 
 void TUniqueTSObjFromIndex<A3DGraphRgbColorData>::InitializeData()
 {
 	A3D_INITIALIZE_DATA(A3DGraphRgbColorData, Data);
 }
+uint32 TUniqueTSObjFromIndex<A3DGraphRgbColorData>::DefaultValue = A3D_DEFAULT_COLOR_INDEX;
 
 void TUniqueTSObjFromIndex<A3DGraphStyleData>::InitializeData()
 {
 	A3D_INITIALIZE_DATA(A3DGraphStyleData, Data);
 }
+uint32 TUniqueTSObjFromIndex<A3DGraphStyleData>::DefaultValue = A3D_DEFAULT_STYLE_INDEX;
+
+void TUniqueTSObjFromIndex<A3DGraphTextureApplicationData>::InitializeData()
+{
+	A3D_INITIALIZE_DATA(A3DGraphTextureApplicationData, Data);
+}
+uint32 TUniqueTSObjFromIndex<A3DGraphTextureApplicationData>::DefaultValue = A3D_DEFAULT_TEXTURE_APPLICATION_INDEX;
+
+void TUniqueTSObjFromIndex<A3DGraphTextureDefinitionData>::InitializeData()
+{
+	A3D_INITIALIZE_DATA(A3DGraphTextureDefinitionData, Data);
+}
+uint32 TUniqueTSObjFromIndex<A3DGraphTextureDefinitionData>::DefaultValue = A3D_DEFAULT_TEXTURE_DEFINITION_INDEX;
 
 void TUniqueTSObj<A3DIntervalData>::InitializeData()
 {
 	A3D_INITIALIZE_DATA(A3DIntervalData, Data);
 }
+const A3DEntity* TUniqueTSObj<A3DIntervalData>::DefaultValue = nullptr;
 
 void TUniqueTSObj<A3DMiscAttributeData>::InitializeData()
 {
 	A3D_INITIALIZE_DATA(A3DMiscAttributeData, Data);
 }
+const A3DEntity* TUniqueTSObj<A3DMiscAttributeData>::DefaultValue = nullptr;
 
 void TUniqueTSObj<A3DMiscCartesianTransformationData>::InitializeData()
 {
 	A3D_INITIALIZE_DATA(A3DMiscCartesianTransformationData, Data);
 }
+const A3DEntity* TUniqueTSObj<A3DMiscCartesianTransformationData>::DefaultValue = nullptr;
 
 void TUniqueTSObj<A3DMiscEntityReferenceData>::InitializeData()
 {
 	A3D_INITIALIZE_DATA(A3DMiscEntityReferenceData, Data);
 }
+const A3DEntity* TUniqueTSObj<A3DMiscEntityReferenceData>::DefaultValue = nullptr;
 
 void TUniqueTSObj<A3DMiscGeneralTransformationData>::InitializeData()
 {
 	A3D_INITIALIZE_DATA(A3DMiscGeneralTransformationData, Data);
 }
+const A3DEntity* TUniqueTSObj<A3DMiscGeneralTransformationData>::DefaultValue = nullptr;
 
 void TUniqueTSObj<A3DMiscMaterialPropertiesData>::InitializeData()
 {
 	A3D_INITIALIZE_DATA(A3DMiscMaterialPropertiesData, Data);
 }
+const A3DEntity* TUniqueTSObj<A3DMiscMaterialPropertiesData>::DefaultValue = nullptr;
 
 void TUniqueTSObj<A3DMiscReferenceOnCsysItemData>::InitializeData()
 {
 	A3D_INITIALIZE_DATA(A3DMiscReferenceOnCsysItemData, Data);
 }
+const A3DEntity* TUniqueTSObj<A3DMiscReferenceOnCsysItemData>::DefaultValue = nullptr;
 
 void TUniqueTSObj<A3DMiscReferenceOnTopologyData>::InitializeData()
 {
 	A3D_INITIALIZE_DATA(A3DMiscReferenceOnTopologyData, Data);
 }
+const A3DEntity* TUniqueTSObj<A3DMiscReferenceOnTopologyData>::DefaultValue = nullptr;
 
 void TUniqueTSObj<A3DMiscReferenceOnTessData>::InitializeData()
 {
 	A3D_INITIALIZE_DATA(A3DMiscReferenceOnTessData, Data);
 }
+const A3DEntity* TUniqueTSObj<A3DMiscReferenceOnTessData>::DefaultValue = nullptr;
 
 void TUniqueTSObj<A3DRiBrepModelData>::InitializeData()
 {
 	A3D_INITIALIZE_DATA(A3DRiBrepModelData, Data);
 }
+const A3DEntity* TUniqueTSObj<A3DRiBrepModelData>::DefaultValue = nullptr;
 
 void TUniqueTSObj<A3DRiCoordinateSystemData>::InitializeData()
 {
 	A3D_INITIALIZE_DATA(A3DRiCoordinateSystemData, Data);
 }
+const A3DEntity* TUniqueTSObj<A3DRiCoordinateSystemData>::DefaultValue = nullptr;
 
 void TUniqueTSObj<A3DRiDirectionData>::InitializeData()
 {
 	A3D_INITIALIZE_DATA(A3DRiDirectionData, Data);
 }
+const A3DEntity* TUniqueTSObj<A3DRiDirectionData>::DefaultValue = nullptr;
 
 void TUniqueTSObj<A3DRiPolyBrepModelData>::InitializeData()
 {
 	A3D_INITIALIZE_DATA(A3DRiPolyBrepModelData, Data);
 }
+const A3DEntity* TUniqueTSObj<A3DRiPolyBrepModelData>::DefaultValue = nullptr;
 
 void TUniqueTSObj<A3DRiRepresentationItemData>::InitializeData()
 {
 	A3D_INITIALIZE_DATA(A3DRiRepresentationItemData, Data);
 }
+const A3DEntity* TUniqueTSObj<A3DRiRepresentationItemData>::DefaultValue = nullptr;
 
 void TUniqueTSObj<A3DRiSetData>::InitializeData()
 {
 	A3D_INITIALIZE_DATA(A3DRiSetData, Data);
 }
+const A3DEntity* TUniqueTSObj<A3DRiSetData>::DefaultValue = nullptr;
 
 void TUniqueTSObj<A3DRootBaseData>::InitializeData()
 {
 	A3D_INITIALIZE_DATA(A3DRootBaseData, Data);
 }
+const A3DEntity* TUniqueTSObj<A3DRootBaseData>::DefaultValue = nullptr;
 
 void TUniqueTSObj<A3DRootBaseWithGraphicsData>::InitializeData()
 {
 	A3D_INITIALIZE_DATA(A3DRootBaseWithGraphicsData, Data);
 }
+const A3DEntity* TUniqueTSObj<A3DRootBaseWithGraphicsData>::DefaultValue = nullptr;
 
 void TUniqueTSObj<A3DSurfBlend01Data>::InitializeData()
 {
 	A3D_INITIALIZE_DATA(A3DSurfBlend01Data, Data);
 }
+const A3DEntity* TUniqueTSObj<A3DSurfBlend01Data>::DefaultValue = nullptr;
 
 void TUniqueTSObj<A3DSurfBlend02Data>::InitializeData()
 {
 	A3D_INITIALIZE_DATA(A3DSurfBlend02Data, Data);
 }
+const A3DEntity* TUniqueTSObj<A3DSurfBlend02Data>::DefaultValue = nullptr;
 
 void TUniqueTSObj<A3DSurfBlend03Data>::InitializeData()
 {
 	A3D_INITIALIZE_DATA(A3DSurfBlend03Data, Data);
 }
+const A3DEntity* TUniqueTSObj<A3DSurfBlend03Data>::DefaultValue = nullptr;
 
 void TUniqueTSObj<A3DSurfConeData>::InitializeData()
 {
 	A3D_INITIALIZE_DATA(A3DSurfConeData, Data);
 }
+const A3DEntity* TUniqueTSObj<A3DSurfConeData>::DefaultValue = nullptr;
 
 void TUniqueTSObj<A3DSurfCylinderData>::InitializeData()
 {
 	A3D_INITIALIZE_DATA(A3DSurfCylinderData, Data);
 }
+const A3DEntity* TUniqueTSObj<A3DSurfCylinderData>::DefaultValue = nullptr;
 
 void TUniqueTSObj<A3DSurfCylindricalData>::InitializeData()
 {
 	A3D_INITIALIZE_DATA(A3DSurfCylindricalData, Data);
 }
+const A3DEntity* TUniqueTSObj<A3DSurfCylindricalData>::DefaultValue = nullptr;
 
 void TUniqueTSObj<A3DSurfExtrusionData>::InitializeData()
 {
 	A3D_INITIALIZE_DATA(A3DSurfExtrusionData, Data);
 }
+const A3DEntity* TUniqueTSObj<A3DSurfExtrusionData>::DefaultValue = nullptr;
 
 void TUniqueTSObj<A3DSurfFromCurvesData>::InitializeData()
 {
 	A3D_INITIALIZE_DATA(A3DSurfFromCurvesData, Data);
 }
+const A3DEntity* TUniqueTSObj<A3DSurfFromCurvesData>::DefaultValue = nullptr;
 
 void TUniqueTSObj<A3DSurfNurbsData>::InitializeData()
 {
 	A3D_INITIALIZE_DATA(A3DSurfNurbsData, Data);
 }
+const A3DEntity* TUniqueTSObj<A3DSurfNurbsData>::DefaultValue = nullptr;
 
 void TUniqueTSObj<A3DSurfPipeData>::InitializeData()
 {
 	A3D_INITIALIZE_DATA(A3DSurfPipeData, Data);
 }
+const A3DEntity* TUniqueTSObj<A3DSurfPipeData>::DefaultValue = nullptr;
 
 void TUniqueTSObj<A3DSurfPlaneData>::InitializeData()
 {
 	A3D_INITIALIZE_DATA(A3DSurfPlaneData, Data);
 }
+const A3DEntity* TUniqueTSObj<A3DSurfPlaneData>::DefaultValue = nullptr;
 
 void TUniqueTSObj<A3DSurfRevolutionData>::InitializeData()
 {
 	A3D_INITIALIZE_DATA(A3DSurfRevolutionData, Data);
 }
+const A3DEntity* TUniqueTSObj<A3DSurfRevolutionData>::DefaultValue = nullptr;
 
 void TUniqueTSObj<A3DSurfRuledData>::InitializeData()
 {
 	A3D_INITIALIZE_DATA(A3DSurfRuledData, Data);
 }
+const A3DEntity* TUniqueTSObj<A3DSurfRuledData>::DefaultValue = nullptr;
 
 void TUniqueTSObj<A3DSurfSphereData>::InitializeData()
 {
 	A3D_INITIALIZE_DATA(A3DSurfSphereData, Data);
 }
+const A3DEntity* TUniqueTSObj<A3DSurfSphereData>::DefaultValue = nullptr;
 
 void TUniqueTSObj<A3DSurfTorusData>::InitializeData()
 {
 	A3D_INITIALIZE_DATA(A3DSurfTorusData, Data);
 }
+const A3DEntity* TUniqueTSObj<A3DSurfTorusData>::DefaultValue = nullptr;
 
 void TUniqueTSObj<A3DTess3DData>::InitializeData()
 {
 	A3D_INITIALIZE_DATA(A3DTess3DData, Data);
 }
+const A3DEntity* TUniqueTSObj<A3DTess3DData>::DefaultValue = nullptr;
 
 void TUniqueTSObj<A3DTessBaseData>::InitializeData()
 {
 	A3D_INITIALIZE_DATA(A3DTessBaseData, Data);
 }
+const A3DEntity* TUniqueTSObj<A3DTessBaseData>::DefaultValue = nullptr;
 
 void TUniqueTSObj<A3DTopoBodyData>::InitializeData()
 {
 	A3D_INITIALIZE_DATA(A3DTopoBodyData, Data);
 }
+const A3DEntity* TUniqueTSObj<A3DTopoBodyData>::DefaultValue = nullptr;
 
 void TUniqueTSObj<A3DTopoBrepDataData>::InitializeData()
 {
 	A3D_INITIALIZE_DATA(A3DTopoBrepDataData, Data);
 }
+const A3DEntity* TUniqueTSObj<A3DTopoBrepDataData>::DefaultValue = nullptr;
 
 void TUniqueTSObj<A3DTopoCoEdgeData>::InitializeData()
 {
 	A3D_INITIALIZE_DATA(A3DTopoCoEdgeData, Data);
 }
+const A3DEntity* TUniqueTSObj<A3DTopoCoEdgeData>::DefaultValue = nullptr;
 
 void TUniqueTSObj<A3DTopoConnexData>::InitializeData()
 {
 	A3D_INITIALIZE_DATA(A3DTopoConnexData, Data);
 }
+const A3DEntity* TUniqueTSObj<A3DTopoConnexData>::DefaultValue = nullptr;
 
 void TUniqueTSObj<A3DTopoContextData>::InitializeData()
 {
 	A3D_INITIALIZE_DATA(A3DTopoContextData, Data);
 }
+const A3DEntity* TUniqueTSObj<A3DTopoContextData>::DefaultValue = nullptr;
 
 void TUniqueTSObj<A3DTopoEdgeData>::InitializeData()
 {
 	A3D_INITIALIZE_DATA(A3DTopoEdgeData, Data);
 }
+const A3DEntity* TUniqueTSObj<A3DTopoEdgeData>::DefaultValue = nullptr;
 
 void TUniqueTSObj<A3DTopoFaceData>::InitializeData()
 {
 	A3D_INITIALIZE_DATA(A3DTopoFaceData, Data);
 }
+const A3DEntity* TUniqueTSObj<A3DTopoFaceData>::DefaultValue = nullptr;
 
 void TUniqueTSObj<A3DTopoLoopData>::InitializeData()
 {
 	A3D_INITIALIZE_DATA(A3DTopoLoopData, Data);
 }
+const A3DEntity* TUniqueTSObj<A3DTopoLoopData>::DefaultValue = nullptr;
 
 void TUniqueTSObj<A3DTopoShellData>::InitializeData()
 {
 	A3D_INITIALIZE_DATA(A3DTopoShellData, Data);
 }
+const A3DEntity* TUniqueTSObj<A3DTopoShellData>::DefaultValue = nullptr;
 
 void TUniqueTSObj<A3DTopoUniqueVertexData>::InitializeData()
 {
 	A3D_INITIALIZE_DATA(A3DTopoUniqueVertexData, Data);
 }
+const A3DEntity* TUniqueTSObj<A3DTopoUniqueVertexData>::DefaultValue = nullptr;
 
 void TUniqueTSObj<A3DTopoWireEdgeData>::InitializeData()
 {
 	A3D_INITIALIZE_DATA(A3DTopoWireEdgeData, Data);
 }
+const A3DEntity* TUniqueTSObj<A3DTopoWireEdgeData>::DefaultValue = nullptr;
 
 void TUniqueTSObj<A3DUTF8Char*>::InitializeData()
 {
 	Data = nullptr;
 }
+const A3DEntity* TUniqueTSObj<A3DUTF8Char*>::DefaultValue = nullptr;
 
 A3DStatus TUniqueTSObj<A3DAsmModelFileData>::GetData(const A3DAsmModelFile* InAsmModelFilePtr)
 {
@@ -498,12 +588,11 @@ A3DStatus TUniqueTSObj<A3DCrvPolyLineData>::GetData(const A3DEntity* InCurvePtr)
 	return A3DCrvPolyLineGet(InCurvePtr, &Data);
 }
 
-A3DStatus TUniqueTSObj<A3DIntervalData>::GetData(const A3DEntity* InCurvePtr)
+A3DStatus TUniqueTSObj<A3DGlobalData>::GetData(const A3DEntity* GlobalPtr)
 {
-	return A3DCrvGetInterval(InCurvePtr, &Data);
+	return A3DGlobalGet(GlobalPtr, &Data);
 }
 
-
 A3DStatus TUniqueTSObj<A3DGraphicsData>::GetData(const A3DGraphics* InGraphicsPtr)
 {
 	return A3DGraphicsGet(InGraphicsPtr, &Data);
@@ -514,6 +603,11 @@ A3DStatus TUniqueTSObjFromIndex<A3DGraphMaterialData>::GetData(uint32 InGraphics
 	return A3DGlobalGetGraphMaterialData(InGraphicsIndex, &Data);
 }
 
+A3DStatus TUniqueTSObjFromIndex<A3DGraphPictureData>::GetData(uint32 InGraphicsIndex)
+{
+	return A3DGlobalGetGraphPictureData(InGraphicsIndex, &Data);
+}
+
 A3DStatus TUniqueTSObjFromIndex<A3DGraphRgbColorData>::GetData(uint32 ObjectIndex)
 {
 	return A3DGlobalGetGraphRgbColorData(ObjectIndex, &Data);
@@ -524,6 +618,21 @@ A3DStatus TUniqueTSObjFromIndex<A3DGraphStyleData>::GetData(uint32 ObjectIndex)
 	return A3DGlobalGetGraphStyleData(ObjectIndex, &Data);
 }
 
+A3DStatus TUniqueTSObjFromIndex<A3DGraphTextureApplicationData>::GetData(uint32 ObjectIndex)
+{
+	return A3DGlobalGetGraphTextureApplicationData(ObjectIndex, &Data);
+}
+
+A3DStatus TUniqueTSObjFromIndex<A3DGraphTextureDefinitionData>::GetData(uint32 ObjectIndex)
+{
+	return A3DGlobalGetGraphTextureDefinitionData(ObjectIndex, &Data);
+}
+
+A3DStatus TUniqueTSObj<A3DIntervalData>::GetData(const A3DEntity* InCurvePtr)
+{
+	return A3DCrvGetInterval(InCurvePtr, &Data);
+}
+
 A3DStatus TUniqueTSObj<A3DMiscAttributeData>::GetData(const A3DMiscAttribute* InMiscAttributePtr)
 {
 	return A3DMiscAttributeGet(InMiscAttributePtr, &Data);
@@ -749,10 +858,6 @@ A3DStatus TUniqueTSObj<A3DUTF8Char*>::GetData(const A3DAsmProductOccurrence* InO
 	return A3DAsmProductOccurrenceGetFilePathName(InOccurrencePtr, &Data);
 }
 
-uint32 TUniqueTSObjFromIndex<A3DGraphMaterialData>::DefaultValue = A3D_DEFAULT_MATERIAL_INDEX;
-uint32 TUniqueTSObjFromIndex<A3DGraphPictureData>::DefaultValue = A3D_DEFAULT_MATERIAL_INDEX;
-uint32 TUniqueTSObjFromIndex<A3DGraphRgbColorData>::DefaultValue = A3D_DEFAULT_COLOR_INDEX;
-uint32 TUniqueTSObjFromIndex<A3DGraphStyleData>::DefaultValue = A3D_DEFAULT_STYLE_INDEX;
 #endif
 
 namespace TechSoftUtils
@@ -771,6 +876,11 @@ bool TECHSOFT_InitializeKernel(const TCHAR* InEnginePluginsPath)
 
 #ifdef USE_TECHSOFT_SDK
 
+A3DStatus GetGlobalPointer(A3DGlobal** GlobalPtr)
+{
+	return A3DGlobalGetPointer(GlobalPtr);
+}
+
 A3DStatus GetSurfaceAsNurbs(const A3DSurfBase* SurfacePtr, A3DSurfNurbsData* DataPtr, A3DDouble Tolerance, A3DBool bUseSameParameterization)
 {
 	return A3DSurfBaseGetAsNurbs(SurfacePtr, Tolerance, bUseSameParameterization, DataPtr);
@@ -806,6 +916,26 @@ bool IsEntityBaseType(const A3DEntity* EntityPtr)
 	return (bool)A3DEntityIsBaseType(EntityPtr);
 }
 
+bool IsMaterialTexture(const uint32 MaterialIndex)
+{
+	A3DBool bIsTexture = false;
+	if (A3DGlobalIsMaterialTexture(MaterialIndex, &bIsTexture) != A3DStatus::A3D_SUCCESS)
+	{
+		return false;
+	}
+	return bIsTexture != false;
+}
+
+A3DEntity* GetPointerFromIndex(const uint32 Index, const A3DEEntityType Type)
+{
+	A3DEntity* EntityPtr;
+	if (A3DMiscPointerFromIndexGet(Index, Type, &EntityPtr) != A3DStatus::A3D_SUCCESS)
+	{
+		return nullptr;
+	}
+	return EntityPtr;
+}
+
 #endif
 
 }
diff --git a/Engine/Plugins/Enterprise/DatasmithCADImporter/Source/CADInterfaces/Public/TechSoftInterface.h b/Engine/Plugins/Enterprise/DatasmithCADImporter/Source/CADInterfaces/Public/TechSoftInterface.h
--- a/Engine/Plugins/Enterprise/DatasmithCADImporter/Source/CADInterfaces/Public/TechSoftInterface.h
+++ b/Engine/Plugins/Enterprise/DatasmithCADImporter/Source/CADInterfaces/Public/TechSoftInterface.h
@@ -67,6 +67,9 @@ CADINTERFACES_API bool TECHSOFT_InitializeKernel(const TCHAR* = TEXT(""));
 CADINTERFACES_API FTechSoftInterface& GetTechSoftInterface();
 
 #ifdef USE_TECHSOFT_SDK
+
+CADINTERFACES_API A3DStatus GetGlobalPointer(A3DGlobal** GlobalPtr);
+
 CADINTERFACES_API A3DStatus GetSurfaceAsNurbs(const A3DSurfBase* SurfacePtr, A3DSurfNurbsData* DataPtr, A3DDouble Tolerance, A3DBool bUseSameParameterization);
 CADINTERFACES_API A3DStatus GetCurveAsNurbs(const A3DCrvBase* A3DCurve, A3DCrvNurbsData* DataPtr, A3DDouble Tolerance, A3DBool bUseSameParameterization);
 
@@ -76,6 +79,10 @@ CADINTERFACES_API A3DStatus GetFilePathName(const A3DAsmProductOccurrence* A3DOc
 CADINTERFACES_API A3DStatus GetEntityType(const A3DEntity* pEntity, A3DEEntityType* peEntityType);
 CADINTERFACES_API bool IsEntityBaseWithGraphicsType(const A3DEntity* pEntity);
 CADINTERFACES_API bool IsEntityBaseType(const A3DEntity* EntityPtr);
+CADINTERFACES_API bool IsMaterialTexture(const uint32 MaterialIndex);
+
+CADINTERFACES_API A3DEntity* GetPointerFromIndex(const uint32 Index, const A3DEEntityType Type);
+
 #endif
 
 } // NS TechSoftUtils
diff --git a/Engine/Plugins/Enterprise/DatasmithCADImporter/Source/CADLibrary/Private/MeshDescriptionHelper.cpp b/Engine/Plugins/Enterprise/DatasmithCADImporter/Source/CADLibrary/Private/MeshDescriptionHelper.cpp
--- a/Engine/Plugins/Enterprise/DatasmithCADImporter/Source/CADLibrary/Private/MeshDescriptionHelper.cpp
+++ b/Engine/Plugins/Enterprise/DatasmithCADImporter/Source/CADLibrary/Private/MeshDescriptionHelper.cpp
@@ -369,8 +369,8 @@ namespace CADLibrary
 		TMap<uint32, FPolygonGroupID> MaterialToPolygonGroupMapping;
 		for (const FTessellationData& FaceTessellation : BodyTessellation.Faces)
 		{
-			// we assume that face has only color
-			MaterialToPolygonGroupMapping.Add(FaceTessellation.ColorName, INDEX_NONE);
+			// material is preferred over color
+			MaterialToPolygonGroupMapping.Add(FaceTessellation.MaterialName ? FaceTessellation.MaterialName : FaceTessellation.ColorName, INDEX_NONE);
 		}
 
 		// Add to the mesh, a polygon groups per material
@@ -414,8 +414,9 @@ namespace CADLibrary
 					continue;
 				}
 
-				// Get the polygonGroup
-				const FPolygonGroupID* PolygonGroupID = MaterialToPolygonGroupMapping.Find(Tessellation.ColorName);
+				// Get the polygonGroup (material is preferred over color)
+				FCADUUID GraphicName = Tessellation.MaterialName ? Tessellation.MaterialName : Tessellation.ColorName;
+				const FPolygonGroupID* PolygonGroupID = MaterialToPolygonGroupMapping.Find(GraphicName);
 				if (PolygonGroupID == nullptr)
 				{
 					continue;
@@ -640,7 +641,7 @@ namespace CADLibrary
 		// Set a diffuse color if there's nothing in the BaseColor
 		if (MaterialElement->GetBaseColor().GetExpression() == nullptr)
 		{
-			FLinearColor LinearColor = FLinearColor::FromPow22Color(InMaterial.Diffuse);
+			FLinearColor LinearColor = FLinearColor::FromSRGBColor(InMaterial.Diffuse);
 
 			IDatasmithMaterialExpressionColor* ColorExpression = MaterialElement->AddMaterialExpression<IDatasmithMaterialExpressionColor>();
 			ColorExpression->SetName(TEXT("Base Color"));
diff --git a/Engine/Plugins/Enterprise/DatasmithCADImporter/Source/DatasmithCADTranslator/Private/DatasmithMeshBuilder.cpp b/Engine/Plugins/Enterprise/DatasmithCADImporter/Source/DatasmithCADTranslator/Private/DatasmithMeshBuilder.cpp
--- a/Engine/Plugins/Enterprise/DatasmithCADImporter/Source/DatasmithCADTranslator/Private/DatasmithMeshBuilder.cpp
+++ b/Engine/Plugins/Enterprise/DatasmithCADImporter/Source/DatasmithCADTranslator/Private/DatasmithMeshBuilder.cpp
@@ -56,8 +56,8 @@ TOptional<FMeshDescription> FDatasmithMeshBuilder::GetMeshDescription(TSharedRef
 
 	// FDatasmithSceneBaseGraphBuilder::BuildBody is performing a special treatment for
 	// FBodyMesh without color. Replicate the treatment here too.
-	int32 Num = Body.ColorSet.Num();
-	if (!Body.ColorSet.Num())
+	int32 Num = Body.ColorSet.Num() + Body.MaterialSet.Num();
+	if (!Num)
 	{
 		ensure(OutMeshElement->GetMaterialSlotCount() == 1);
 		int32 MaterialSlotId = OutMeshElement->GetMaterialSlotAt(0)->GetId();
diff --git a/Engine/Plugins/Enterprise/DatasmithCADImporter/Source/DatasmithCADTranslator/Private/DatasmithSceneGraphBuilder.cpp b/Engine/Plugins/Enterprise/DatasmithCADImporter/Source/DatasmithCADTranslator/Private/DatasmithSceneGraphBuilder.cpp
--- a/Engine/Plugins/Enterprise/DatasmithCADImporter/Source/DatasmithCADTranslator/Private/DatasmithSceneGraphBuilder.cpp
+++ b/Engine/Plugins/Enterprise/DatasmithCADImporter/Source/DatasmithCADTranslator/Private/DatasmithSceneGraphBuilder.cpp
@@ -497,7 +497,7 @@ TSharedPtr<IDatasmithActorElement> FDatasmithSceneBaseGraphBuilder::BuildBody(in
 	uint32 MaterialUuid = 0;
 	MaterialUuid = ParentData.MaterialUuid ? ParentData.MaterialUuid : ParentData.ColorUuid;
 
-	if (!Body.ColorFaceSet.Num())
+	if (!(Body.ColorFaceSet.Num() + Body.MaterialFaceSet.Num()))
 	{
 		Body.ColorFaceSet.Add(MaterialUuid);
 	}
@@ -560,16 +560,18 @@ TSharedPtr< IDatasmithMeshElement > FDatasmithSceneBaseGraphBuilder::FindOrAddMe
 	Hash.Set(MD5);
 	MeshElement->SetFileHash(Hash);
 
-	// Currently we assume that face has only colors
-	TSet<uint32>& MaterialSet = Body.ColorFaceSet;
-
-	for (uint32 MaterialSlotId : MaterialSet)
+	TFunction<void(TSet<uint32>&)> SetMaterialToDatasmithMeshElement = [&](TSet<uint32>& MaterialSet)
 	{
-		TSharedPtr< IDatasmithMaterialIDElement > PartMaterialIDElement;
-		PartMaterialIDElement = FindOrAddMaterial(MaterialSlotId);
+		for (uint32 MaterialSlotId : MaterialSet)
+		{
+			TSharedPtr< IDatasmithMaterialIDElement > PartMaterialIDElement;
+			PartMaterialIDElement = FindOrAddMaterial(MaterialSlotId);
+			MeshElement->SetMaterial(PartMaterialIDElement->GetName(), MaterialSlotId);
+		}
+	};
 
-		MeshElement->SetMaterial(PartMaterialIDElement->GetName(), MaterialSlotId);
-	}
+	SetMaterialToDatasmithMeshElement(Body.ColorFaceSet);
+	SetMaterialToDatasmithMeshElement(Body.MaterialFaceSet);
 
 	DatasmithScene->AddMesh(MeshElement);
 
