Techsoft: add tessellation settings

diff --git a/Engine/Plugins/Enterprise/DatasmithCADImporter/Source/CADInterfaces/Private/TechSoftFileParser.cpp b/Engine/Plugins/Enterprise/DatasmithCADImporter/Source/CADInterfaces/Private/TechSoftFileParser.cpp
--- a/Engine/Plugins/Enterprise/DatasmithCADImporter/Source/CADInterfaces/Private/TechSoftFileParser.cpp
+++ b/Engine/Plugins/Enterprise/DatasmithCADImporter/Source/CADInterfaces/Private/TechSoftFileParser.cpp
@@ -1818,29 +1818,19 @@ void FTechSoftFileParser::MeshRepresentationWithTechSoft(A3DRiRepresentationItem
 	A3DRWParamsTessellationData TessellationParameters;
 	A3D_INITIALIZE_DATA(A3DRWParamsTessellationData, TessellationParameters);
 
-	if (true)
-	{
-		TessellationParameters.m_eTessellationLevelOfDetail = kA3DTessLODMedium;		/*!< Enum to specify predefined values for some following members. */
-	}
-	else
-	{
-		// TODO
-		// Check unity conversion
-		const FImportParameters& ImportParameters = CADFileData.GetImportParameters();
+	const FImportParameters& ImportParameters = CADFileData.GetImportParameters();
 
-		////A3DRWParamsTessellationData /*!< The tessellation reading parameters. */
-		TessellationParameters.m_eTessellationLevelOfDetail = kA3DTessLODUserDefined;		/*!< Enum to specify predefined values for some following members. */
-		TessellationParameters.m_bUseHeightInsteadOfRatio = A3D_TRUE;
-		TessellationParameters.m_dMaxChordHeight = ImportParameters.GetChordTolerance() * 10. / FileUnit;
-		TessellationParameters.m_dAngleToleranceDeg = ImportParameters.GetMaxNormalAngle();
-		TessellationParameters.m_dMaximalTriangleEdgeLength = 0; //ImportParameters.MaxEdgeLength;
+	TessellationParameters.m_eTessellationLevelOfDetail = kA3DTessLODUserDefined; // Enum to specify predefined values for some following members.
+	TessellationParameters.m_bUseHeightInsteadOfRatio = A3D_TRUE;
+	TessellationParameters.m_dMaxChordHeight = ImportParameters.GetChordTolerance() * 10. / FileUnit;
+	TessellationParameters.m_dAngleToleranceDeg = ImportParameters.GetMaxNormalAngle();
+	TessellationParameters.m_dMaximalTriangleEdgeLength = 0; //ImportParameters.MaxEdgeLength;
 
-		TessellationParameters.m_bAccurateTessellation = A3D_FALSE;  // A3D_FALSE' indicates the tessellation is set for visualization
-		TessellationParameters.m_bAccurateTessellationWithGrid = A3D_FALSE; /*!< Enable accurate tessellation with faces inner points on a grid.*/
-		TessellationParameters.m_dAccurateTessellationWithGridMaximumStitchLength = 0; 	/*!< Maximal grid stitch length. Disabled if value is 0. Be careful, a too small value can generate a huge tessellation. */
-	}
+	TessellationParameters.m_bAccurateTessellation = A3D_FALSE;  // A3D_FALSE' indicates the tessellation is set for visualization
+	TessellationParameters.m_bAccurateTessellationWithGrid = A3D_FALSE; // Enable accurate tessellation with faces inner points on a grid.
+	TessellationParameters.m_dAccurateTessellationWithGridMaximumStitchLength = 0; 	// Maximal grid stitch length. Disabled if value is 0. Be careful, a too small value can generate a huge tessellation.
 
-	TessellationParameters.m_bKeepUVPoints = A3D_TRUE; /*!< Keep parametric points as texture points. */
+	TessellationParameters.m_bKeepUVPoints = A3D_TRUE; // Keep parametric points as texture points.
 
 	// Get the tessellation
 	A3DStatus Status = A3DRiRepresentationItemComputeTessellation(RepresentationItemPtr, &TessellationParameters);
