TechSoft : Add Stitching::Sew operation

diff --git a/Engine/Plugins/Enterprise/DatasmithCADImporter/Source/CADInterfaces/Private/TechSoftFileParser.cpp b/Engine/Plugins/Enterprise/DatasmithCADImporter/Source/CADInterfaces/Private/TechSoftFileParser.cpp
--- a/Engine/Plugins/Enterprise/DatasmithCADImporter/Source/CADInterfaces/Private/TechSoftFileParser.cpp
+++ b/Engine/Plugins/Enterprise/DatasmithCADImporter/Source/CADInterfaces/Private/TechSoftFileParser.cpp
@@ -301,17 +301,25 @@ ECADParsingResult FTechSoftFileParser::Process()
 		}
 	}
 
-	if (CADFileData.GetImportParameters().GetStitchingTechnique() != StitchingNone && !FImportParameters::bGDisableCADKernelTessellation)
+	A3DAsmModelFile* ModelFile = TechSoftInterface.GetModelFile();
+
+	if (CADFileData.GetImportParameters().GetStitchingTechnique() == StitchingSew && FImportParameters::bGDisableCADKernelTessellation)
 	{
-		// todo
-		//TechSoftInterface.Repair(CADFileData.GetStitchingTechnique());
+		TUniqueTSObj<A3DSewOptionsData> SewData;
+		SewData->m_bComputePreferredOpenShellOrientation = false;
+		double ToleranceMM = 0.01 / FileUnit;
+		A3DStatus Status = TechSoftUtils::SewModel(&ModelFile, ToleranceMM, &*SewData);
+		if(Status != A3DStatus::A3D_SUCCESS)
+		{
+			// To do but what ?
+		}
 	}
 
 	ReserveCADFileData();
 
 	ReadMaterialsAndColors();
 
-	ECADParsingResult Result = TraverseModel(TechSoftInterface.GetModelFile());
+	ECADParsingResult Result = TraverseModel(ModelFile);
 
 	TechSoftInterface.UnloadModel();
 
@@ -1064,7 +1072,24 @@ FCadId FTechSoftFileParser::TraverseBRepModel(A3DRiBrepModel* BRepModelPtr, FEnt
 	TraverseRepresentationContent(BRepModelPtr, Body);
 	if (FImportParameters::bGDisableCADKernelTessellation)
 	{
-		MeshRepresentationWithTechSoft(BRepModelPtr, Body);
+		// todo
+		//if (CADFileData.GetImportParameters().GetStitchingTechnique() == StitchingHeal)
+		//{
+		//	TUniqueTSObj<A3DSewOptionsData> SewData;
+		//	SewData->m_bComputePreferredOpenShellOrientation = false;
+		//	double ToleranceMM = 0.01 / FileUnit;
+		//	A3DRiBrepModel** NewBReps;
+		//	uint32 NewBRepCount = 0;
+		//	A3DStatus Status = TechSoftUtils::HealBRep(&BRepModelPtr, ToleranceMM, &*SewData, &NewBReps, NewBRepCount);
+		//	if (Status == A3DStatus::A3D_SUCCESS)
+		//	{
+		//		MeshRepresentationsWithTechSoft(NewBRepCount, NewBReps, Body);
+		//	}
+		//}
+		//else
+		{
+			MeshRepresentationWithTechSoft(BRepModelPtr, Body);
+		}
 	}
 	else
 	{
@@ -1691,7 +1716,6 @@ bool FTechSoftFileParser::IsConfigurationSet(const A3DAsmProductOccurrence* Occu
 	return bIsConfiguration;
 }
 
-
 uint32 FTechSoftFileParser::CountColorAndMaterial()
 {
 	A3DGlobal* GlobalPtr = nullptr;
diff --git a/Engine/Plugins/Enterprise/DatasmithCADImporter/Source/CADInterfaces/Private/TechSoftInterface.cpp b/Engine/Plugins/Enterprise/DatasmithCADImporter/Source/CADInterfaces/Private/TechSoftInterface.cpp
--- a/Engine/Plugins/Enterprise/DatasmithCADImporter/Source/CADInterfaces/Private/TechSoftInterface.cpp
+++ b/Engine/Plugins/Enterprise/DatasmithCADImporter/Source/CADInterfaces/Private/TechSoftInterface.cpp
@@ -334,6 +334,12 @@ void TUniqueTSObj<A3DRootBaseWithGraphicsData>::InitializeData()
 }
 const A3DEntity* TUniqueTSObj<A3DRootBaseWithGraphicsData>::DefaultValue = nullptr;
 
+void TUniqueTSObj<A3DSewOptionsData>::InitializeData()
+{
+	A3D_INITIALIZE_DATA(A3DSewOptionsData, Data);
+}
+const A3DEntity* TUniqueTSObj<A3DSewOptionsData>::DefaultValue = nullptr;
+
 void TUniqueTSObj<A3DSurfBlend01Data>::InitializeData()
 {
 	A3D_INITIALIZE_DATA(A3DSurfBlend01Data, Data);
@@ -713,6 +719,11 @@ A3DStatus TUniqueTSObj<A3DRootBaseWithGraphicsData>::GetData(const A3DRootBaseWi
 	return A3DRootBaseWithGraphicsGet(InRootBaseWithGraphicsPtr, &Data);
 }
 
+A3DStatus TUniqueTSObj<A3DSewOptionsData>::GetData(const A3DEntity* InRootBaseWithGraphicsPtr)
+{
+	return A3DStatus::A3D_ERROR;
+}
+
 A3DStatus TUniqueTSObj<A3DSurfBlend01Data>::GetData(const A3DEntity* InSurfPtr)
 {
 	return A3DSurfBlend01Get(InSurfPtr, &Data);
@@ -936,6 +947,27 @@ A3DEntity* GetPointerFromIndex(const uint32 Index, const A3DEEntityType Type)
 	return EntityPtr;
 }
 
+A3DStatus HealBRep(A3DRiBrepModel** BRepToHeal, double Tolerance, A3DSewOptionsData const* SewOptions, A3DRiBrepModel*** OutNewBReps, uint32& OutNewBRepCount)
+{
+	A3DUns32 NewBRepCount;
+	A3DStatus Status = A3DSewBrep(&BRepToHeal, 1, Tolerance, SewOptions, OutNewBReps, &NewBRepCount);
+	OutNewBRepCount = NewBRepCount;
+	return Status;
+}
+
+A3DStatus SewBReps(A3DRiBrepModel*** BRepsToSew, uint32 const BRepCount, double Tolerance, A3DSewOptionsData const* SewOptions, A3DRiBrepModel*** OutNewBReps, uint32& OutNewBRepCount)
+{
+	A3DUns32 NewBRepCount;
+	A3DStatus Status = A3DSewBrep(BRepsToSew, BRepCount, Tolerance, SewOptions, OutNewBReps, &NewBRepCount);
+	OutNewBRepCount = NewBRepCount;
+	return Status;
+}
+
+A3DStatus SewModel(A3DAsmModelFile** ModelPtr, double Tolerance, A3DSewOptionsData const* SewOptions)
+{
+	return A3DAsmModelFileSew(ModelPtr, Tolerance, SewOptions);
+}
+
 #endif
 
 }
diff --git a/Engine/Plugins/Enterprise/DatasmithCADImporter/Source/CADInterfaces/Public/TechSoftInterface.h b/Engine/Plugins/Enterprise/DatasmithCADImporter/Source/CADInterfaces/Public/TechSoftInterface.h
--- a/Engine/Plugins/Enterprise/DatasmithCADImporter/Source/CADInterfaces/Public/TechSoftInterface.h
+++ b/Engine/Plugins/Enterprise/DatasmithCADImporter/Source/CADInterfaces/Public/TechSoftInterface.h
@@ -81,6 +81,10 @@ CADINTERFACES_API bool IsEntityBaseWithGraphicsType(const A3DEntity* pEntity);
 CADINTERFACES_API bool IsEntityBaseType(const A3DEntity* EntityPtr);
 CADINTERFACES_API bool IsMaterialTexture(const uint32 MaterialIndex);
 
+CADINTERFACES_API A3DStatus HealBRep(A3DRiBrepModel** BRepToHeal, double Tolerance, A3DSewOptionsData const* SewOptions, A3DRiBrepModel*** OutNewBReps, uint32& OutNewBRepCount);
+CADINTERFACES_API A3DStatus SewBReps(A3DRiBrepModel*** BRepsToSew, uint32 const BRepCount, double Tolerance, A3DSewOptionsData const* SewOptions, A3DRiBrepModel*** OutNewBReps, uint32& OutNewBRepCount);
+CADINTERFACES_API A3DStatus SewModel(A3DAsmModelFile** ModelPtr, double Tolerance, A3DSewOptionsData const* SewOptions);
+
 CADINTERFACES_API A3DEntity* GetPointerFromIndex(const uint32 Index, const A3DEEntityType Type);
 
 #endif
