TechSoft: Fix material transparency

diff --git a/Engine/Plugins/Enterprise/DatasmithCADImporter/Source/CADInterfaces/Private/TechSoftFileParser.cpp b/Engine/Plugins/Enterprise/DatasmithCADImporter/Source/CADInterfaces/Private/TechSoftFileParser.cpp
--- a/Engine/Plugins/Enterprise/DatasmithCADImporter/Source/CADInterfaces/Private/TechSoftFileParser.cpp
+++ b/Engine/Plugins/Enterprise/DatasmithCADImporter/Source/CADInterfaces/Private/TechSoftFileParser.cpp
@@ -1459,7 +1459,7 @@ FArchiveMaterial& FTechSoftFileParser::AddMaterialAt(uint32 MaterialIndexToSave,
 		Material.Ambient = TechSoftFileParserImpl::GetColorAt(MaterialData->m_uiAmbient);
 		Material.Specular = TechSoftFileParserImpl::GetColorAt(MaterialData->m_uiSpecular);
 		Material.Shininess = MaterialData->m_dShininess;
-		Material.Transparency = 1. - MaterialData->m_dAmbientAlpha;
+		Material.Transparency = 1. - MaterialData->m_dDiffuseAlpha;
 		// todo: find how to convert Emissive color into ? reflexion coef...
 		// Material.Emissive = GetColor(MaterialData->m_uiEmissive);
 		// Material.Reflexion;
