Update Analytics data information

diff --git a/Engine/Plugins/Enterprise/DatasmithCADImporter/Source/CADInterfaces/Private/TechSoftFileParser.cpp b/Engine/Plugins/Enterprise/DatasmithCADImporter/Source/CADInterfaces/Private/TechSoftFileParser.cpp
--- a/Engine/Plugins/Enterprise/DatasmithCADImporter/Source/CADInterfaces/Private/TechSoftFileParser.cpp
+++ b/Engine/Plugins/Enterprise/DatasmithCADImporter/Source/CADInterfaces/Private/TechSoftFileParser.cpp
@@ -321,6 +321,12 @@ ECADParsingResult FTechSoftFileParser::Process()
 
 	ECADParsingResult Result = TraverseModel(ModelFile);
 
+	FString TechsoftVersion = TechSoftUtils::GetTechSoftVersion();
+	if (!TechsoftVersion.IsEmpty() && CADFileData.ComponentCount() > 0)
+	{
+		CADFileData.GetComponentAt(0).MetaData.Add(TEXT("TechsoftVersion"), TechsoftVersion);
+	}
+
 	TechSoftInterface.UnloadModel();
 
 	return Result;
diff --git a/Engine/Plugins/Enterprise/DatasmithCADImporter/Source/CADInterfaces/Private/TechSoftInterface.cpp b/Engine/Plugins/Enterprise/DatasmithCADImporter/Source/CADInterfaces/Private/TechSoftInterface.cpp
--- a/Engine/Plugins/Enterprise/DatasmithCADImporter/Source/CADInterfaces/Private/TechSoftInterface.cpp
+++ b/Engine/Plugins/Enterprise/DatasmithCADImporter/Source/CADInterfaces/Private/TechSoftInterface.cpp
@@ -880,6 +880,17 @@ FTechSoftInterface& GetTechSoftInterface()
 	return TechSoftInterface;
 }
 
+FString GetTechSoftVersion()
+{
+#ifdef USE_TECHSOFT_SDK
+	A3DInt32 MajorVersion = 0, MinorVersion = 0;
+	A3DDllGetVersion(&MajorVersion, &MinorVersion);
+	return FString::Printf(TEXT("Techsoft %d.%d"), MajorVersion, MinorVersion);
+#endif
+	return FString();
+}
+
+
 bool TECHSOFT_InitializeKernel(const TCHAR* InEnginePluginsPath)
 {
 	return GetTechSoftInterface().InitializeKernel(InEnginePluginsPath);
diff --git a/Engine/Plugins/Enterprise/DatasmithCADImporter/Source/CADInterfaces/Public/CADFileData.h b/Engine/Plugins/Enterprise/DatasmithCADImporter/Source/CADInterfaces/Public/CADFileData.h
--- a/Engine/Plugins/Enterprise/DatasmithCADImporter/Source/CADInterfaces/Public/CADFileData.h
+++ b/Engine/Plugins/Enterprise/DatasmithCADImporter/Source/CADInterfaces/Public/CADFileData.h
@@ -124,6 +124,10 @@ public:
 		return SceneGraphArchive.CADIdToComponentIndex.Find(ComponentId) != nullptr;
 	}
 
+	int32 ComponentCount()
+	{
+		return SceneGraphArchive.Components.Num();
+	}
 
 	int32 AddComponent(FCadId ComponentId)
 	{
diff --git a/Engine/Plugins/Enterprise/DatasmithCADImporter/Source/CADInterfaces/Public/TechSoftInterface.h b/Engine/Plugins/Enterprise/DatasmithCADImporter/Source/CADInterfaces/Public/TechSoftInterface.h
--- a/Engine/Plugins/Enterprise/DatasmithCADImporter/Source/CADInterfaces/Public/TechSoftInterface.h
+++ b/Engine/Plugins/Enterprise/DatasmithCADImporter/Source/CADInterfaces/Public/TechSoftInterface.h
@@ -65,6 +65,7 @@ namespace TechSoftUtils
 
 CADINTERFACES_API bool TECHSOFT_InitializeKernel(const TCHAR* = TEXT(""));
 CADINTERFACES_API FTechSoftInterface& GetTechSoftInterface();
+CADINTERFACES_API FString GetTechSoftVersion();
 
 #ifdef USE_TECHSOFT_SDK
 
diff --git a/Engine/Plugins/Enterprise/DatasmithCADImporter/Source/CADTools/Private/CADData.cpp b/Engine/Plugins/Enterprise/DatasmithCADImporter/Source/CADTools/Private/CADData.cpp
--- a/Engine/Plugins/Enterprise/DatasmithCADImporter/Source/CADTools/Private/CADData.cpp
+++ b/Engine/Plugins/Enterprise/DatasmithCADImporter/Source/CADTools/Private/CADData.cpp
@@ -236,11 +236,14 @@ ECADFormat FileFormat(const FString& Extension)
 	{
 		return ECADFormat::CATIAV4;
 	}
-	else if (Extension == TEXT("prt.*") || Extension == TEXT("asm.*"))
+	else if (Extension == TEXT("prt.*") || Extension == TEXT("asm.*") 
+		|| Extension == TEXT("creo") || Extension == TEXT("creo.*")
+		|| Extension == TEXT("neu") || Extension == TEXT("neu.*")
+		|| Extension == TEXT("xas") || Extension == TEXT("xpr"))
 	{
 		return ECADFormat::CREO;
 	}
-	else if (Extension == TEXT("prt"))
+	else if (Extension == TEXT("prt") || Extension == TEXT("asm"))
 	{
 		return ECADFormat::NX;
 	}
@@ -256,10 +259,26 @@ ECADFormat FileFormat(const FString& Extension)
 	{
 		return ECADFormat::PARASOLID;
 	}
-	else if (Extension == TEXT("3dxml"))
+	else if (Extension == TEXT("3dxml") || Extension == TEXT("3drep"))
 	{
 		return ECADFormat::CATIA_3DXML;
 	}
+	else if (Extension == TEXT("par") || Extension == TEXT("psm"))
+	{
+		return ECADFormat::SOLID_EDGE;
+	}
+	else if (Extension == TEXT("dwg"))
+	{
+		return ECADFormat::AUTOCAD;
+	}
+	else if (Extension == TEXT("dgn"))
+	{
+		return ECADFormat::MICROSTATION;
+	}
+	else if (Extension == TEXT("hsf"))
+	{
+		return ECADFormat::TECHSOFT;
+	}
 	else
 	{
 		return ECADFormat::OTHER;
diff --git a/Engine/Plugins/Enterprise/DatasmithCADImporter/Source/CADTools/Public/CADData.h b/Engine/Plugins/Enterprise/DatasmithCADImporter/Source/CADTools/Public/CADData.h
--- a/Engine/Plugins/Enterprise/DatasmithCADImporter/Source/CADTools/Public/CADData.h
+++ b/Engine/Plugins/Enterprise/DatasmithCADImporter/Source/CADTools/Public/CADData.h
@@ -21,6 +21,7 @@ namespace CADLibrary
 enum class ECADFormat
 {
 	ACIS,
+	AUTOCAD,
 	CATIA,
 	CATIA_CGR,
 	CATIA_3DXML,
@@ -28,11 +29,15 @@ enum class ECADFormat
 	CREO,
 	DWG,
 	DGN,
+	TECHSOFT,
+	IFC,
 	IGES,
 	INVENTOR,
 	JT,
 	NX,
+	MICROSTATION,
 	PARASOLID,
+	SOLID_EDGE,
 	SOLIDWORKS,
 	STEP,
 	OTHER
diff --git a/Engine/Plugins/Enterprise/DatasmithCADImporter/Source/DatasmithCADTranslator/Private/DatasmithCADTranslator.cpp b/Engine/Plugins/Enterprise/DatasmithCADImporter/Source/DatasmithCADTranslator/Private/DatasmithCADTranslator.cpp
--- a/Engine/Plugins/Enterprise/DatasmithCADImporter/Source/DatasmithCADTranslator/Private/DatasmithCADTranslator.cpp
+++ b/Engine/Plugins/Enterprise/DatasmithCADImporter/Source/DatasmithCADTranslator/Private/DatasmithCADTranslator.cpp
@@ -37,7 +37,6 @@ void FDatasmithCADTranslator::Initialize(FDatasmithTranslatorCapabilities& OutCa
 	OutCapabilities.SupportedFileFormats.Add(FFileFormatInfo{ TEXT("model"), TEXT("CATIA V4 files") });
 
 	OutCapabilities.SupportedFileFormats.Add(FFileFormatInfo{ TEXT("asm.*"), TEXT("Creo Assembly files") });
-	OutCapabilities.SupportedFileFormats.Add(FFileFormatInfo{ TEXT("asm"), TEXT("Creo, NX Assembly files") });
 	OutCapabilities.SupportedFileFormats.Add(FFileFormatInfo{ TEXT("creo.*"), TEXT("Creo Assembly files") });
 	OutCapabilities.SupportedFileFormats.Add(FFileFormatInfo{ TEXT("creo"), TEXT("Creo Assembly files") });
 	OutCapabilities.SupportedFileFormats.Add(FFileFormatInfo{ TEXT("neu.*"), TEXT("Creo Assembly files") });
@@ -45,7 +44,6 @@ void FDatasmithCADTranslator::Initialize(FDatasmithTranslatorCapabilities& OutCa
 	OutCapabilities.SupportedFileFormats.Add(FFileFormatInfo{ TEXT("prt.*"), TEXT("Creo Part files") });
 	OutCapabilities.SupportedFileFormats.Add(FFileFormatInfo{ TEXT("xas"), TEXT("Creo Assembly files") });
 	OutCapabilities.SupportedFileFormats.Add(FFileFormatInfo{ TEXT("xpr"), TEXT("Creo Part files") });
-	OutCapabilities.SupportedFileFormats.Add(FFileFormatInfo{ TEXT("prt"), TEXT("Creo, NX Part files") });
 
 	OutCapabilities.SupportedFileFormats.Add(FFileFormatInfo{ TEXT("iam"), TEXT("Inventor Assembly files") });
 	OutCapabilities.SupportedFileFormats.Add(FFileFormatInfo{ TEXT("ipt"), TEXT("Inventor Part files") });
diff --git a/Engine/Plugins/Enterprise/DatasmithCADImporter/Source/DatasmithCADTranslator/Private/DatasmithSceneGraphBuilder.cpp b/Engine/Plugins/Enterprise/DatasmithCADImporter/Source/DatasmithCADTranslator/Private/DatasmithSceneGraphBuilder.cpp
--- a/Engine/Plugins/Enterprise/DatasmithCADImporter/Source/DatasmithCADTranslator/Private/DatasmithSceneGraphBuilder.cpp
+++ b/Engine/Plugins/Enterprise/DatasmithCADImporter/Source/DatasmithCADTranslator/Private/DatasmithSceneGraphBuilder.cpp
@@ -216,9 +216,9 @@ bool FDatasmithSceneBaseGraphBuilder::Build()
 
 	// Set ProductName, ProductVersion in DatasmithScene for Analytics purpose
 	// application_name is something like "Catia V5"
-	DatasmithScene->SetVendor(TEXT("CoreTechnologie"));
+	DatasmithScene->SetVendor(TEXT("Techsoft"));
 
-	if (const FString* ProductVersion = Component.MetaData.Find(TEXT("KernelIOVersion")))
+	if (const FString* ProductVersion = Component.MetaData.Find(TEXT("TechsoftVersion")))
 	{
 		DatasmithScene->SetProductVersion(**ProductVersion);
 	}
