CVar ds.CADTranslator.PreferJtFileEmbeddedTessellation

diff --git a/Engine/Config/ConsoleVariables.ini b/Engine/Config/ConsoleVariables.ini
--- a/Engine/Config/ConsoleVariables.ini
+++ b/Engine/Config/ConsoleVariables.ini
@@ -127,4 +127,8 @@ fx.SkipVectorVMBackendOptimizations=1
 ;ds.CADTranslator.DisableCADKernelTessellation=false
 
 ; Uncomment to select the SDK to load CAD files: "KernelIO" or "TechSoft". Default is "KernelIO"
-;ds.CADTranslator.CADLibrary="KernelIO"
\ No newline at end of file
+;ds.CADTranslator.CADLibrary="KernelIO"
+
+; Uncomment to force to import embedded tessellation instead of meshing BRep model. Default is false
+;ds.CADTranslator.PreferJtFileEmbeddedTessellation=true
+
diff --git a/Engine/Plugins/Enterprise/DatasmithCADImporter/Source/CADInterfaces/Private/TechSoftFileParser.cpp b/Engine/Plugins/Enterprise/DatasmithCADImporter/Source/CADInterfaces/Private/TechSoftFileParser.cpp
--- a/Engine/Plugins/Enterprise/DatasmithCADImporter/Source/CADInterfaces/Private/TechSoftFileParser.cpp
+++ b/Engine/Plugins/Enterprise/DatasmithCADImporter/Source/CADInterfaces/Private/TechSoftFileParser.cpp
@@ -2,6 +2,8 @@
 
 #include "TechSoftFileParser.h"
 
+#include "CADOptions.h"
+
 #ifdef USE_TECHSOFT_SDK
 
 #include "TechSoftInterface.h"
@@ -248,6 +250,16 @@ void UpdateIOOptionAccordingToFormat(const CADLibrary::ECADFormat Format, A3DImp
 		break;
 	}
 
+	case CADLibrary::ECADFormat::JT:
+	{
+		if (FImportParameters::bGPreferJtFileEmbeddedTessellation)
+		{
+			Importer.m_sLoadData.m_sGeneral.m_eReadGeomTessMode = kA3DReadTessOnly;
+			Importer.m_sLoadData.m_sSpecifics.m_sJT.m_eReadTessellationLevelOfDetail = A3DEJTReadTessellationLevelOfDetail::kA3DJTTessLODHigh;
+			break;
+		}
+	}
+
 	case CADLibrary::ECADFormat::INVENTOR:
 	case CADLibrary::ECADFormat::CATIA_3DXML:
 	{
diff --git a/Engine/Plugins/Enterprise/DatasmithCADImporter/Source/CADTools/Private/CADOption.cpp b/Engine/Plugins/Enterprise/DatasmithCADImporter/Source/CADTools/Private/CADOption.cpp
--- a/Engine/Plugins/Enterprise/DatasmithCADImporter/Source/CADTools/Private/CADOption.cpp
+++ b/Engine/Plugins/Enterprise/DatasmithCADImporter/Source/CADTools/Private/CADOption.cpp
@@ -54,6 +54,13 @@ FAutoConsoleVariableRef GCADTranslatorEnableTimeControl(
 	TEXT("Enable the timer that kill the worker if the import time is unusually long. With this time control, the load of the corrupted file is canceled but the rest of the scene is imported.\n"),
 	ECVF_Default);
 
+bool FImportParameters::bGPreferJtFileEmbeddedTessellation = false;
+FAutoConsoleVariableRef GCADTranslatorJtFileEmbeddedTessellation(
+	TEXT("ds.CADTranslator.PreferJtFileEmbeddedTessellation"),
+	FImportParameters::bGPreferJtFileEmbeddedTessellation,
+	TEXT("If both (tessellation and BRep) exist in the file, import embedded tessellation instead of meshing BRep.\n"),
+	ECVF_Default);
+
 uint32 GetTypeHash(const FImportParameters& ImportParameters)
 {
 	uint32 ParametersHash = ::GetTypeHash(ImportParameters.bGDisableCADKernelTessellation);
diff --git a/Engine/Plugins/Enterprise/DatasmithCADImporter/Source/CADTools/Public/CADOptions.h b/Engine/Plugins/Enterprise/DatasmithCADImporter/Source/CADTools/Public/CADOptions.h
--- a/Engine/Plugins/Enterprise/DatasmithCADImporter/Source/CADTools/Public/CADOptions.h
+++ b/Engine/Plugins/Enterprise/DatasmithCADImporter/Source/CADTools/Public/CADOptions.h
@@ -54,6 +54,7 @@ namespace CADLibrary
 		CADTOOLS_API static bool bGEnableTimeControl;
 		CADTOOLS_API static bool bGEnableCADCache;
 		CADTOOLS_API static bool bGOverwriteCache;
+		CADTOOLS_API static bool bGPreferJtFileEmbeddedTessellation;
 
 	public:
 		FImportParameters(double InMetricUnit = 0.001, double InScaleFactor = 0.1, FDatasmithUtils::EModelCoordSystem NewCoordinateSystem = FDatasmithUtils::EModelCoordSystem::ZUp_RightHanded)
@@ -104,6 +105,7 @@ namespace CADLibrary
 				Hash = HashCombine(Hash, ::GetTypeHash(Param));
 			}
 			Hash = HashCombine(Hash, ::GetTypeHash(bScaleUVMap));
+			Hash = HashCombine(Hash, ::GetTypeHash(bGPreferJtFileEmbeddedTessellation));
 			return Hash;
 		}
 
@@ -119,12 +121,13 @@ namespace CADLibrary
 			Ar << (uint8&) ImportParameters.Propagation;
 			Ar << (uint8&) ImportParameters.DisplayPreference;
 			Ar << ImportParameters.bScaleUVMap;
-			// these 4 static variables have to be serialized to be transmitted to CADWorkers
+			// these 5 static variables have to be serialized to be transmitted to CADWorkers
 			// because CADWorker has not access to CVars
 			Ar << ImportParameters.bGOverwriteCache;
 			Ar << ImportParameters.bGDisableCADKernelTessellation;
 			Ar << ImportParameters.bGEnableTimeControl;
 			Ar << ImportParameters.bGEnableCADCache;
+			Ar << ImportParameters.bGPreferJtFileEmbeddedTessellation;
 			return Ar;
 		}
 
