Implement body caching using Hoops format for retessellation
Fixed some logic flaws when CAD cache is disabled: retessellation would not work.

diff --git a/Engine/Plugins/Enterprise/DatasmithCADImporter/Source/CADInterfaces/Private/TechSoftFileParser.cpp b/Engine/Plugins/Enterprise/DatasmithCADImporter/Source/CADInterfaces/Private/TechSoftFileParser.cpp
--- a/Engine/Plugins/Enterprise/DatasmithCADImporter/Source/CADInterfaces/Private/TechSoftFileParser.cpp
+++ b/Engine/Plugins/Enterprise/DatasmithCADImporter/Source/CADInterfaces/Private/TechSoftFileParser.cpp
@@ -333,6 +333,11 @@ ECADParsingResult FTechSoftFileParser::Process()
 
 	ECADParsingResult Result = TraverseModel(ModelFile);
 
+	if (Result == ECADParsingResult::ProcessOk)
+	{
+		GenerateBodyMeshes();
+	}
+
 	FString TechsoftVersion = TechSoftUtils::GetTechSoftVersion();
 	if (!TechsoftVersion.IsEmpty() && CADFileData.ComponentCount() > 0)
 	{
@@ -344,6 +349,98 @@ ECADParsingResult FTechSoftFileParser::Process()
 	return Result;
 }
 
+void FTechSoftFileParser::GenerateBodyMeshes()
+{
+	if (FImportParameters::bGDisableCADKernelTessellation)
+	{
+		for (TPair<A3DRiRepresentationItem*, int32>& Entry : RepresentationItemsCache)
+		{
+			A3DRiRepresentationItem* RepresentationItemPtr = Entry.Key;
+
+			FArchiveBody& Body = CADFileData.GetBodyAt(Entry.Value);
+
+			FBodyMesh& BodyMesh = CADFileData.AddBodyMesh(Body.ObjectId, Body);
+
+			// todo
+			//if (CADFileData.GetImportParameters().GetStitchingTechnique() == StitchingHeal)
+			//{
+			//	TUniqueTSObj<A3DSewOptionsData> SewData;
+			//	SewData->m_bComputePreferredOpenShellOrientation = false;
+			//	double ToleranceMM = 0.01 / FileUnit;
+			//	A3DRiBrepModel** NewBReps;
+			//	uint32 NewBRepCount = 0;
+			//	A3DStatus Status = TechSoftUtils::HealBRep(&BRepModelPtr, ToleranceMM, &*SewData, &NewBReps, NewBRepCount);
+			//	if (Status == A3DStatus::A3D_SUCCESS)
+			//	{
+			//		// Aggregate all the bodies into one?
+			//		MeshRepresentationsWithTechSoft(NewBRepCount, NewBReps, Body);
+			//		TechSoftInterface.FillBodyMesh(NewBRepCount, NewBReps, CADFileData.GetImportParameters(), FileUnit, BodyMesh);???
+			//	}
+			//}
+			//else
+			{
+				TechSoftInterface.FillBodyMesh(RepresentationItemPtr, CADFileData.GetImportParameters(), FileUnit, BodyMesh);
+			}
+
+			// Convert material
+			FCADUUID DefaultColorName = Body.ColorFaceSet.Num() > 0 ? *Body.ColorFaceSet.begin()  : 0;
+			FCADUUID DefaultMaterialName = Body.MaterialFaceSet.Num() > 0 ? *Body.MaterialFaceSet.begin() : 0;
+
+			for (FTessellationData& Tessellation : BodyMesh.Faces)
+			{
+				FCADUUID ColorName = DefaultColorName;
+				FCADUUID MaterialName = DefaultMaterialName;
+
+				// Extract proper color or material based on style index
+				uint32 CachedStyleIndex = Tessellation.MaterialName;
+				Tessellation.MaterialName = 0;
+
+				if (CachedStyleIndex != 0xffffffff)
+				{
+					ExtractGraphStyleProperties(CachedStyleIndex, ColorName, MaterialName);
+				}
+
+				if (ColorName)
+				{
+					Tessellation.ColorName = ColorName;
+					BodyMesh.ColorSet.Add(ColorName);
+				}
+
+				if (MaterialName)
+				{
+					Tessellation.MaterialName = MaterialName;
+					BodyMesh.MaterialSet.Add(MaterialName);
+				}
+			}
+
+			Body.ColorFaceSet = BodyMesh.ColorSet;
+			Body.MaterialFaceSet = BodyMesh.MaterialSet;
+
+			// Write part's representation as hsf file
+			A3DEEntityType Type;
+			A3DEntityGetType(RepresentationItemPtr, &Type);
+
+			if (Type == kA3DTypeRiBrepModel)
+			{
+				FString FilePath = CADFileData.GetBodyCachePath(Body.MeshActorName);
+				if (!FilePath.IsEmpty())
+				{
+					this->TechSoftInterface.SaveBodyToHsfFile(RepresentationItemPtr, FilePath);
+				}
+			}
+
+			// #ueent_techsoft: Take care of materials stored in BodyMesh
+
+			Body.ColorFaceSet = BodyMesh.ColorSet;
+			Body.MaterialFaceSet = BodyMesh.MaterialSet;
+		}
+	}
+	else
+	{
+		// Mesh with CADKernel
+	}
+}
+
 void FTechSoftFileParser::ReserveCADFileData()
 {
 	// TODO: Could be more accurate
@@ -546,14 +643,6 @@ void FTechSoftFileParser::TraverseReference(const A3DAsmProductOccurrence* Refer
 	{
 		TraversePartDefinition(ReferenceData->m_pPart, Component);
 	}
-
-#ifndef NEW_CODE
-	// Is this really necessary ????
-	if (ReferenceData->m_pPrototype && !ReferenceData->m_uiPOccurrencesSize && !ReferenceData->m_pPart)
-	{
-		TraversePrototype(ReferenceData->m_pPrototype, Component);
-	}
-#endif
 }
 
 FArchiveInstance& FTechSoftFileParser::AddInstance(FEntityMetaData& InstanceMetaData)
@@ -619,7 +708,7 @@ FArchiveComponent& FTechSoftFileParser::AddOccurence(FEntityMetaData& InstanceMe
 	return Prototype;
 }
 
-FArchiveBody& FTechSoftFileParser::AddBody(FEntityMetaData& BodyMetaData)
+int32 FTechSoftFileParser::AddBody(FEntityMetaData& BodyMetaData)
 {
 	FCadId BodyId = LastEntityId++;
 	int32 BodyIndex = CADFileData.AddBody(BodyId);
@@ -633,7 +722,7 @@ FArchiveBody& FTechSoftFileParser::AddBody(FEntityMetaData& BodyMetaData)
 	{
 		Body.MaterialFaceSet.Add(BodyMetaData.MaterialName);
 	}
-	return Body;
+	return BodyIndex;
 }
 
 FCadId FTechSoftFileParser::TraverseOccurrence(const A3DAsmProductOccurrence* OccurrencePtr)
@@ -731,31 +820,6 @@ FCadId FTechSoftFileParser::TraverseOccurrence(const A3DAsmProductOccurrence* Oc
 	return Instance.ObjectId;
 }
 
-void FTechSoftFileParser::ProcessOccurrence(TUniqueTSObj<A3DAsmProductOccurrenceData>& OccurrenceData, FArchiveComponent& Component)
-{
-	if (!OccurrenceData.IsValid())
-	{
-		return;
-	}
-
-	if (OccurrenceData->m_pPart)
-	{
-		TraversePartDefinition(OccurrenceData->m_pPart, Component);
-	}
-
-	if (OccurrenceData->m_pExternalData)
-	{
-		int32 ChildrenId = TraverseOccurrence(OccurrenceData->m_pExternalData);
-		Component.Children.Add(ChildrenId);
-	}
-
-	for (uint32 Index = 0; Index < OccurrenceData->m_uiPOccurrencesSize; ++Index)
-	{
-		int32 ChildrenId = TraverseOccurrence(OccurrenceData->m_ppPOccurrences[Index]);
-		Component.Children.Add(ChildrenId);
-	}
-};
-
 void FTechSoftFileParser::CountUnderOccurrence(const A3DAsmProductOccurrence* Occurrence)
 {
 	TUniqueTSObj<A3DAsmProductOccurrenceData> OccurrenceData(Occurrence);
@@ -795,31 +859,6 @@ void FTechSoftFileParser::CountUnderOccurrence(const A3DAsmProductOccurrence* Oc
 	}
 }
 
-void FTechSoftFileParser::TraversePrototype(const A3DAsmProductOccurrence* InPrototypePtr, FArchiveComponent& Component)
-{
-	TUniqueTSObj<A3DAsmProductOccurrenceData> SubPrototypeData(InPrototypePtr);
-	if (!SubPrototypeData.IsValid())
-	{
-		return;
-	}
-
-	for (uint32 Index = 0; Index < SubPrototypeData->m_uiPOccurrencesSize; ++Index)
-	{
-		int32 ChildrenId = TraverseOccurrence(SubPrototypeData->m_ppPOccurrences[Index]);
-		Component.Children.Add(ChildrenId);
-	}
-
-	if (SubPrototypeData->m_pPart)
-	{
-		TraversePartDefinition(SubPrototypeData->m_pPart, Component);
-	}
-
-	if (SubPrototypeData->m_pPrototype && !SubPrototypeData->m_uiPOccurrencesSize && !SubPrototypeData->m_pPart)
-	{
-		TraversePrototype(SubPrototypeData->m_pPrototype, Component);
-	}
-}
-
 void FTechSoftFileParser::CountUnderSubPrototype(const A3DAsmProductOccurrence* InPrototypePtr)
 {
 	TUniqueTSObj<A3DAsmProductOccurrenceData> SubPrototypeData(InPrototypePtr);
@@ -1064,17 +1103,6 @@ FCadId FTechSoftFileParser::TraverseBRepModel(A3DRiBrepModel* BRepModelPtr, FEnt
 		return 0;
 	}
 
-	if (FCadId* CadIdPtr = RepresentationItemsCache.Find(BRepModelPtr))
-	{
-		return *CadIdPtr;
-	}
-
-	TUniqueTSObj<A3DRiBrepModelData> BodyData(BRepModelPtr);
-	if (!BodyData.IsValid())
-	{
-		return 0;
-	}
-
 	FEntityMetaData BRepMetaData;
 	ExtractMetaData(BRepModelPtr, BRepMetaData);
 	if (!BRepMetaData.bShow || BRepMetaData.bRemoved)
@@ -1082,80 +1110,30 @@ FCadId FTechSoftFileParser::TraverseBRepModel(A3DRiBrepModel* BRepModelPtr, FEnt
 		return 0;
 	}
 
-	ExtractSpecificMetaData(BRepModelPtr, BRepMetaData);
-	ExtractMaterialProperties(BRepModelPtr);
-
-	FArchiveBody& Body = AddBody(BRepMetaData);
 
-	TraverseRepresentationContent(BRepModelPtr, Body);
-	if (FImportParameters::bGDisableCADKernelTessellation)
+	if (int32* BodyIndexPtr = RepresentationItemsCache.Find(BRepModelPtr))
 	{
-		// todo
-		//if (CADFileData.GetImportParameters().GetStitchingTechnique() == StitchingHeal)
-		//{
-		//	TUniqueTSObj<A3DSewOptionsData> SewData;
-		//	SewData->m_bComputePreferredOpenShellOrientation = false;
-		//	double ToleranceMM = 0.01 / FileUnit;
-		//	A3DRiBrepModel** NewBReps;
-		//	uint32 NewBRepCount = 0;
-		//	A3DStatus Status = TechSoftUtils::HealBRep(&BRepModelPtr, ToleranceMM, &*SewData, &NewBReps, NewBRepCount);
-		//	if (Status == A3DStatus::A3D_SUCCESS)
-		//	{
-		//		MeshRepresentationsWithTechSoft(NewBRepCount, NewBReps, Body);
-		//	}
-		//}
-		//else
-		{
-			MeshRepresentationWithTechSoft(BRepModelPtr, Body);
-		}
-	}
-	else
-	{
-		// Mesh with CADKernel
+		return CADFileData.GetBodyAt(*BodyIndexPtr).ObjectId;
 	}
 
-	RepresentationItemsCache.Add(BRepModelPtr, Body.ObjectId);
+	ExtractSpecificMetaData(BRepModelPtr, BRepMetaData);
+	ExtractMaterialProperties(BRepModelPtr);
 
-	return Body.ObjectId;
-}
+	int32 BodyIndex = AddBody(BRepMetaData);
+	FArchiveBody& Body = CADFileData.GetBodyAt(BodyIndex);
 
-void FTechSoftFileParser::TraverseRepresentationContent(const A3DRiRepresentationItem* RepresentationItemPtr, FArchiveBody& Body)
-{
-	TUniqueTSObj<A3DRiRepresentationItemData> RepresentationItemData(RepresentationItemPtr);
-	if (!RepresentationItemData.IsValid())
-	{
-		return;
-	}
+	RepresentationItemsCache.Add(BRepModelPtr, BodyIndex);
 
-	if (RepresentationItemData->m_pCoordinateSystem)
-	{
-		TraverseCoordinateSystem(RepresentationItemData->m_pCoordinateSystem);
-	}
-
-	if (RepresentationItemData->m_pTessBase)
-	{
-		TraverseTessellationBase(RepresentationItemData->m_pTessBase, Body);
-	}
+	return Body.ObjectId;
 }
 
-FCadId FTechSoftFileParser::TraversePolyBRepModel(const A3DRiPolyBrepModel* PolygonalPtr, FEntityMetaData& PartMetaData)
+FCadId FTechSoftFileParser::TraversePolyBRepModel(A3DRiPolyBrepModel* PolygonalPtr, FEntityMetaData& PartMetaData)
 {
 	if (!PolygonalPtr)
 	{
 		return 0;
 	}
 
-	if (FCadId* CadIdPtr = RepresentationItemsCache.Find(PolygonalPtr))
-	{
-		return *CadIdPtr;
-	}
-
-	TUniqueTSObj<A3DRiPolyBrepModelData> BodyData(PolygonalPtr);
-	if (!BodyData.IsValid())
-	{
-		return 0;
-	}
-
 	FEntityMetaData BRepMetaData;
 	ExtractMetaData(PolygonalPtr, BRepMetaData);
 	if (!BRepMetaData.bShow || BRepMetaData.bRemoved)
@@ -1163,13 +1141,18 @@ FCadId FTechSoftFileParser::TraversePolyBRepModel(const A3DRiPolyBrepModel* Poly
 		return 0;
 	}
 
+	if (int32* BodyIndexPtr = RepresentationItemsCache.Find(PolygonalPtr))
+	{
+		return CADFileData.GetBodyAt(*BodyIndexPtr).ObjectId;
+	}
+
 	ExtractSpecificMetaData(PolygonalPtr, BRepMetaData);
 	ExtractMaterialProperties(PolygonalPtr);
 
-	FArchiveBody& Body = AddBody(BRepMetaData);
-	TraverseRepresentationContent(PolygonalPtr, Body);
+	int32 BodyIndex = AddBody(BRepMetaData);
+	FArchiveBody& Body = CADFileData.GetBodyAt(BodyIndex);
 
-	RepresentationItemsCache.Add(PolygonalPtr, Body.ObjectId);
+	RepresentationItemsCache.Add(PolygonalPtr, BodyIndex);
 
 	return Body.ObjectId;
 }
@@ -1852,747 +1835,6 @@ void FTechSoftFileParser::ReadMaterialsAndColors()
 	}
 }
 
-void FTechSoftFileParser::MeshRepresentationWithTechSoft(A3DRiRepresentationItem* RepresentationItemPtr, FArchiveBody& Body)
-{
-	TUniqueTSObj<A3DRiRepresentationItemData> RepresentationItemData;
-
-	// TUniqueTechSoftObj does not work in this case
-	A3DRWParamsTessellationData TessellationParameters;
-	A3D_INITIALIZE_DATA(A3DRWParamsTessellationData, TessellationParameters);
-
-	const FImportParameters& ImportParameters = CADFileData.GetImportParameters();
-
-	TessellationParameters.m_eTessellationLevelOfDetail = kA3DTessLODUserDefined; // Enum to specify predefined values for some following members.
-	TessellationParameters.m_bUseHeightInsteadOfRatio = A3D_TRUE;
-	TessellationParameters.m_dMaxChordHeight = ImportParameters.GetChordTolerance() * 10. / FileUnit;
-	TessellationParameters.m_dAngleToleranceDeg = ImportParameters.GetMaxNormalAngle();
-	TessellationParameters.m_dMaximalTriangleEdgeLength = 0; //ImportParameters.MaxEdgeLength;
-
-	TessellationParameters.m_bAccurateTessellation = A3D_FALSE;  // A3D_FALSE' indicates the tessellation is set for visualization
-	TessellationParameters.m_bAccurateTessellationWithGrid = A3D_FALSE; // Enable accurate tessellation with faces inner points on a grid.
-	TessellationParameters.m_dAccurateTessellationWithGridMaximumStitchLength = 0; 	// Maximal grid stitch length. Disabled if value is 0. Be careful, a too small value can generate a huge tessellation.
-
-	TessellationParameters.m_bKeepUVPoints = A3D_TRUE; // Keep parametric points as texture points.
-
-	// Get the tessellation
-	A3DStatus Status = A3DRiRepresentationItemComputeTessellation(RepresentationItemPtr, &TessellationParameters);
-	Status = A3DRiRepresentationItemGet(RepresentationItemPtr, RepresentationItemData.GetEmptyDataPtr());
-	TraverseTessellationBase(RepresentationItemData->m_pTessBase, Body);
-}
-
-void FTechSoftFileParser::TraverseTessellationBase(const A3DTessBase* Tessellation, FArchiveBody& Body)
-{
-	A3DEEntityType eType;
-	if (A3DEntityGetType(Tessellation, &eType) == A3D_SUCCESS)
-	{
-		switch (eType)
-		{
-		case kA3DTypeTess3D:
-			TraverseTessellation3D(Tessellation, Body);
-			break;
-		case kA3DTypeTess3DWire:
-		case kA3DTypeTessMarkup:
-		default:
-			break;
-		}
-	}
-}
-
-namespace TechSoftFileParserImpl
-{
-
-uint32 CountTriangles(const A3DTessFaceData& FaceTessData)
-{
-	const int32 TessellationFaceDataWithTriangle = 0x2222;
-	const int32 TessellationFaceDataWithFan = 0x4444;
-	const int32 TessellationFaceDataWithStrip = 0x8888;
-	const int32 TessellationFaceDataWithOneNormal = 0xE0E0;
-
-	uint32 UsedEntitiesFlags = FaceTessData.m_usUsedEntitiesFlags;
-
-	uint32 TriangleCount = 0;
-	uint32 FaceSetIndex = 0;
-	if (UsedEntitiesFlags & TessellationFaceDataWithTriangle)
-	{
-		TriangleCount += FaceTessData.m_puiSizesTriangulated[FaceSetIndex];
-		FaceSetIndex++;
-	}
-
-	if (FaceTessData.m_uiSizesTriangulatedSize > FaceSetIndex)
-	{
-		if (UsedEntitiesFlags & TessellationFaceDataWithFan)
-		{
-			uint32 LastFanIndex = 1 + FaceSetIndex + FaceTessData.m_puiSizesTriangulated[FaceSetIndex];
-			FaceSetIndex++;
-			for (; FaceSetIndex < LastFanIndex; FaceSetIndex++)
-			{
-				uint32 FanSize = (FaceTessData.m_puiSizesTriangulated[FaceSetIndex] & kA3DTessFaceDataNormalMask);
-				TriangleCount += (FanSize - 2);
-			}
-		}
-	}
-
-	if (FaceTessData.m_uiSizesTriangulatedSize > FaceSetIndex)
-	{
-		FaceSetIndex++;
-		for (; FaceSetIndex < FaceTessData.m_uiSizesTriangulatedSize; FaceSetIndex++)
-		{
-			uint32 StripSize = (FaceTessData.m_puiSizesTriangulated[FaceSetIndex] & kA3DTessFaceDataNormalMask);
-			TriangleCount += (StripSize - 2);
-		}
-	}
-	return TriangleCount;
-}
-
-void AddFaceTriangleWithUniqueNormal(FTessellationData& Tessellation, const A3DTess3DData& Tessellation3DData, const uint32 InTriangleCount, uint32& InOutLastTriangleIndex, int32& InOutLastVertexIndex)
-{
-	int32 FaceIndex[3] = { 0, 0, 0 };
-	int32 NormalIndex[3] = { 0, 0, 0 };
-
-	// Get Triangles
-	for (uint32 TriangleIndex = 0; TriangleIndex < InTriangleCount; TriangleIndex++)
-	{
-		NormalIndex[0] = Tessellation3DData.m_puiTriangulatedIndexes[InOutLastTriangleIndex++];
-		NormalIndex[1] = NormalIndex[0];
-		NormalIndex[2] = NormalIndex[0];
-
-		FaceIndex[0] = Tessellation3DData.m_puiTriangulatedIndexes[InOutLastTriangleIndex++] / 3;
-		FaceIndex[1] = Tessellation3DData.m_puiTriangulatedIndexes[InOutLastTriangleIndex++] / 3;
-		FaceIndex[2] = Tessellation3DData.m_puiTriangulatedIndexes[InOutLastTriangleIndex++] / 3;
-
-		if (!TechSoftFileParserImpl::AddFace(FaceIndex, Tessellation, InOutLastVertexIndex))
-		{
-			continue;
-		}
-
-		TechSoftFileParserImpl::AddNormals(Tessellation3DData.m_pdNormals, NormalIndex, Tessellation.NormalArray);
-	}
-}
-
-void AddFaceTriangleWithUniqueNormalAndTexture(FTessellationData& Tessellation, const A3DTess3DData& Tessellation3DData, const uint32 InTriangleCount, const uint32 TextureCount, uint32& InOutStartIndex, int32& InOutLastVertexIndex)
-{
-	int32 FaceIndex[3] = { 0, 0, 0 };
-	int32 NormalIndex[3] = { 0, 0, 0 };
-	int32 TextureIndex[3] = { 0, 0, 0 };
-
-	// Get Triangles
-	for (uint32 TriangleIndex = 0; TriangleIndex < InTriangleCount; TriangleIndex++)
-	{
-		NormalIndex[0] = Tessellation3DData.m_puiTriangulatedIndexes[InOutStartIndex++];
-		NormalIndex[1] = NormalIndex[0];
-		NormalIndex[2] = NormalIndex[0];
-
-		TextureIndex[0] = Tessellation3DData.m_puiTriangulatedIndexes[InOutStartIndex];
-		InOutStartIndex += TextureCount;
-		FaceIndex[0] = Tessellation3DData.m_puiTriangulatedIndexes[InOutStartIndex++] / 3;
-		TextureIndex[1] = Tessellation3DData.m_puiTriangulatedIndexes[InOutStartIndex];
-		InOutStartIndex += TextureCount;
-		FaceIndex[1] = Tessellation3DData.m_puiTriangulatedIndexes[InOutStartIndex++] / 3;
-		TextureIndex[2] = Tessellation3DData.m_puiTriangulatedIndexes[InOutStartIndex];
-		InOutStartIndex += TextureCount;
-		FaceIndex[2] = Tessellation3DData.m_puiTriangulatedIndexes[InOutStartIndex++] / 3;
-
-		if (!TechSoftFileParserImpl::AddFace(FaceIndex, Tessellation, InOutLastVertexIndex))
-		{
-			continue;
-		}
-
-		TechSoftFileParserImpl::AddNormals(Tessellation3DData.m_pdNormals, NormalIndex, Tessellation.NormalArray);
-		TechSoftFileParserImpl::AddTextureCoordinates(Tessellation3DData.m_pdTextureCoords, TextureIndex, Tessellation.TexCoordArray);
-	}
-}
-
-void AddFaceTriangle(FTessellationData& Tessellation, const A3DTess3DData& Tessellation3DData, const uint32 InTriangleCount, uint32& InOutStartIndex, int32& InOutLastVertexIndex)
-{
-	int32 FaceIndex[3] = { 0, 0, 0 };
-	int32 NormalIndex[3] = { 0, 0, 0 };
-
-	// Get Triangles
-	for (uint32 TriangleIndex = 0; TriangleIndex < InTriangleCount; TriangleIndex++)
-	{
-		NormalIndex[0] = Tessellation3DData.m_puiTriangulatedIndexes[InOutStartIndex++];
-		FaceIndex[0] = Tessellation3DData.m_puiTriangulatedIndexes[InOutStartIndex++] / 3;
-		NormalIndex[1] = Tessellation3DData.m_puiTriangulatedIndexes[InOutStartIndex++];
-		FaceIndex[1] = Tessellation3DData.m_puiTriangulatedIndexes[InOutStartIndex++] / 3;
-		NormalIndex[2] = Tessellation3DData.m_puiTriangulatedIndexes[InOutStartIndex++];
-		FaceIndex[2] = Tessellation3DData.m_puiTriangulatedIndexes[InOutStartIndex++] / 3;
-
-		if (TechSoftFileParserImpl::AddFace(FaceIndex, Tessellation, InOutLastVertexIndex))
-		{
-			TechSoftFileParserImpl::AddNormals(Tessellation3DData.m_pdNormals, NormalIndex, Tessellation.NormalArray);
-		}
-	}
-}
-
-void AddFaceTriangleWithTexture(FTessellationData& Tessellation, const A3DTess3DData& Tessellation3DData, const uint32 InTriangleCount, const uint32 InTextureCount, uint32& InOutStartIndex, int32& inOutLastVertexIndex)
-{
-	int32 FaceIndex[3] = { 0, 0, 0 };
-	int32 NormalIndex[3] = { 0, 0, 0 };
-	int32 TextureIndex[3] = { 0, 0, 0 };
-
-	// Get Triangles
-	for (uint64 TriangleIndex = 0; TriangleIndex < InTriangleCount; TriangleIndex++)
-	{
-		NormalIndex[0] = Tessellation3DData.m_puiTriangulatedIndexes[InOutStartIndex++];
-		TextureIndex[0] = Tessellation3DData.m_puiTriangulatedIndexes[InOutStartIndex];
-		InOutStartIndex += InTextureCount;
-		FaceIndex[0] = Tessellation3DData.m_puiTriangulatedIndexes[InOutStartIndex++] / 3;
-		NormalIndex[1] = Tessellation3DData.m_puiTriangulatedIndexes[InOutStartIndex++];
-		TextureIndex[1] = Tessellation3DData.m_puiTriangulatedIndexes[InOutStartIndex];
-		InOutStartIndex += InTextureCount;
-		FaceIndex[1] = Tessellation3DData.m_puiTriangulatedIndexes[InOutStartIndex++] / 3;
-		NormalIndex[2] = Tessellation3DData.m_puiTriangulatedIndexes[InOutStartIndex++];
-		TextureIndex[2] = Tessellation3DData.m_puiTriangulatedIndexes[InOutStartIndex];
-		InOutStartIndex += InTextureCount;
-		FaceIndex[2] = Tessellation3DData.m_puiTriangulatedIndexes[InOutStartIndex++] / 3;
-
-		if (TechSoftFileParserImpl::AddFace(FaceIndex, Tessellation, inOutLastVertexIndex))
-		{
-			TechSoftFileParserImpl::AddNormals(Tessellation3DData.m_pdNormals, NormalIndex, Tessellation.NormalArray);
-			TechSoftFileParserImpl::AddTextureCoordinates(Tessellation3DData.m_pdTextureCoords, TextureIndex, Tessellation.TexCoordArray);
-		}
-	}
-}
-
-void AddFaceTriangleFanWithUniqueNormal(FTessellationData& Tessellation, const A3DTess3DData& Tessellation3DData, const uint32 InTriangleCount, uint32& InOutLastTriangleIndex, int32& LastVertexIndex)
-{
-	int32 FaceIndex[3] = { 0, 0, 0 };
-	int32 NormalIndex[3] = { 0, 0, 0 };
-
-	NormalIndex[0] = Tessellation3DData.m_puiTriangulatedIndexes[InOutLastTriangleIndex++];
-	NormalIndex[1] = NormalIndex[0];
-	NormalIndex[2] = NormalIndex[0];
-
-	FaceIndex[0] = Tessellation3DData.m_puiTriangulatedIndexes[InOutLastTriangleIndex++] / 3;
-	FaceIndex[1] = Tessellation3DData.m_puiTriangulatedIndexes[InOutLastTriangleIndex++] / 3;
-
-	// Get Triangles
-	for (uint32 TriangleIndex = 2; TriangleIndex < InTriangleCount; TriangleIndex++)
-	{
-		FaceIndex[2] = Tessellation3DData.m_puiTriangulatedIndexes[InOutLastTriangleIndex++] / 3;
-
-		if (TechSoftFileParserImpl::AddFace(FaceIndex, Tessellation, LastVertexIndex))
-		{
-			TechSoftFileParserImpl::AddNormals(Tessellation3DData.m_pdNormals, NormalIndex, Tessellation.NormalArray);
-		}
-
-		FaceIndex[1] = FaceIndex[2];
-	}
-}
-
-void AddFaceTriangleFanWithUniqueNormalAndTexture(FTessellationData& Tessellation, const A3DTess3DData& Tessellation3DData, const uint32 InTriangleCount, const uint32 TextureCount, uint32& InOutLastTriangleIndex, int32& InOutLastVertexIndex)
-{
-	int32 FaceIndex[3] = { 0, 0, 0 };
-	int32 NormalIndex[3] = { 0, 0, 0 };
-	int32 TextureIndex[3] = { 0, 0, 0 };
-
-	NormalIndex[0] = Tessellation3DData.m_puiTriangulatedIndexes[InOutLastTriangleIndex++];
-	NormalIndex[1] = NormalIndex[0];
-	NormalIndex[2] = NormalIndex[0];
-
-	TextureIndex[0] = Tessellation3DData.m_puiTriangulatedIndexes[InOutLastTriangleIndex];
-	InOutLastTriangleIndex += TextureCount;
-	FaceIndex[0] = Tessellation3DData.m_puiTriangulatedIndexes[InOutLastTriangleIndex++] / 3;
-
-	TextureIndex[1] = Tessellation3DData.m_puiTriangulatedIndexes[InOutLastTriangleIndex];
-	InOutLastTriangleIndex += TextureCount;
-	FaceIndex[1] = Tessellation3DData.m_puiTriangulatedIndexes[InOutLastTriangleIndex++] / 3;
-
-	for (uint32 TriangleIndex = 2; TriangleIndex < InTriangleCount; TriangleIndex++)
-	{
-		TextureIndex[2] = Tessellation3DData.m_puiTriangulatedIndexes[InOutLastTriangleIndex];
-		InOutLastTriangleIndex += TextureCount;
-		FaceIndex[2] = Tessellation3DData.m_puiTriangulatedIndexes[InOutLastTriangleIndex++] / 3;
-
-		if (TechSoftFileParserImpl::AddFace(FaceIndex, Tessellation, InOutLastVertexIndex))
-		{
-			TechSoftFileParserImpl::AddNormals(Tessellation3DData.m_pdNormals, NormalIndex, Tessellation.NormalArray);
-			TechSoftFileParserImpl::AddTextureCoordinates(Tessellation3DData.m_pdTextureCoords, TextureIndex, Tessellation.TexCoordArray);
-		}
-
-		FaceIndex[1] = FaceIndex[2];
-		TextureIndex[1] = TextureIndex[2];
-	}
-}
-
-void AddFaceTriangleFan(FTessellationData& Tessellation, const A3DTess3DData& Tessellation3DData, const uint32 InTriangleCount, uint32& InOutLastTriangleIndex, int32& InOutLastVertexIndex)
-{
-	int32 FaceIndex[3] = { 0, 0, 0 };
-	int32 NormalIndex[3] = { 0, 0, 0 };
-
-	NormalIndex[0] = Tessellation3DData.m_puiTriangulatedIndexes[InOutLastTriangleIndex++];
-	FaceIndex[0] = Tessellation3DData.m_puiTriangulatedIndexes[InOutLastTriangleIndex++] / 3;
-	NormalIndex[1] = Tessellation3DData.m_puiTriangulatedIndexes[InOutLastTriangleIndex++];
-	FaceIndex[1] = Tessellation3DData.m_puiTriangulatedIndexes[InOutLastTriangleIndex++] / 3;
-
-	for (unsigned long TriangleIndex = 2; TriangleIndex < InTriangleCount; TriangleIndex++)
-	{
-		NormalIndex[2] = Tessellation3DData.m_puiTriangulatedIndexes[InOutLastTriangleIndex++];
-		FaceIndex[2] = Tessellation3DData.m_puiTriangulatedIndexes[InOutLastTriangleIndex++] / 3;
-
-		if (TechSoftFileParserImpl::AddFace(FaceIndex, Tessellation, InOutLastVertexIndex))
-		{
-			TechSoftFileParserImpl::AddNormals(Tessellation3DData.m_pdNormals, NormalIndex, Tessellation.NormalArray);
-		}
-
-		NormalIndex[1] = NormalIndex[2];
-		FaceIndex[1] = FaceIndex[2];
-	}
-}
-
-void AddFaceTriangleFanWithTexture(FTessellationData& Tessellation, const A3DTess3DData& Tessellation3DData, const uint32 InTriangleCount, const uint32 TextureCount, uint32& InOutLastTriangleIndex, int32& InOutLastVertexIndex)
-{
-	int32 FaceIndex[3] = { 0, 0, 0 };
-	int32 NormalIndex[3] = { 0, 0, 0 };
-	int32 TextureIndex[3] = { 0, 0, 0 };
-
-	NormalIndex[0] = Tessellation3DData.m_puiTriangulatedIndexes[InOutLastTriangleIndex++];
-	TextureIndex[0] = Tessellation3DData.m_puiTriangulatedIndexes[InOutLastTriangleIndex];
-	InOutLastTriangleIndex += TextureCount;
-	FaceIndex[0] = Tessellation3DData.m_puiTriangulatedIndexes[InOutLastTriangleIndex++] / 3;
-
-	NormalIndex[1] = Tessellation3DData.m_puiTriangulatedIndexes[InOutLastTriangleIndex++];
-	TextureIndex[1] = Tessellation3DData.m_puiTriangulatedIndexes[InOutLastTriangleIndex];
-	InOutLastTriangleIndex += TextureCount;
-	FaceIndex[1] = Tessellation3DData.m_puiTriangulatedIndexes[InOutLastTriangleIndex++] / 3;
-
-	for (uint32 TriangleIndex = 2; TriangleIndex < InTriangleCount; TriangleIndex++)
-	{
-		NormalIndex[2] = Tessellation3DData.m_puiTriangulatedIndexes[InOutLastTriangleIndex++];
-		TextureIndex[2] = Tessellation3DData.m_puiTriangulatedIndexes[InOutLastTriangleIndex];
-		InOutLastTriangleIndex += TextureCount;
-		FaceIndex[2] = Tessellation3DData.m_puiTriangulatedIndexes[InOutLastTriangleIndex++] / 3;
-
-		if (TechSoftFileParserImpl::AddFace(FaceIndex, Tessellation, InOutLastVertexIndex))
-		{
-			TechSoftFileParserImpl::AddNormals(Tessellation3DData.m_pdNormals, NormalIndex, Tessellation.NormalArray);
-			TechSoftFileParserImpl::AddTextureCoordinates(Tessellation3DData.m_pdTextureCoords, TextureIndex, Tessellation.TexCoordArray);
-		}
-
-		NormalIndex[1] = NormalIndex[2];
-		TextureIndex[1] = TextureIndex[2];
-		FaceIndex[1] = FaceIndex[2];
-	}
-}
-
-void AddFaceTriangleStripWithUniqueNormal(FTessellationData& Tessellation, const A3DTess3DData& Tessellation3DData, const uint32 InTriangleCount, uint32& InOutLastTriangleIndex, int32& InOutLastVertexIndex)
-{
-	int32 FaceIndex[3] = { 0, 0, 0 };
-	int32 NormalIndex[3] = { 0, 0, 0 };
-
-	NormalIndex[0] = Tessellation3DData.m_puiTriangulatedIndexes[InOutLastTriangleIndex++];
-	NormalIndex[1] = NormalIndex[0];
-	NormalIndex[2] = NormalIndex[0];
-
-	FaceIndex[0] = Tessellation3DData.m_puiTriangulatedIndexes[InOutLastTriangleIndex++] / 3;
-	FaceIndex[1] = Tessellation3DData.m_puiTriangulatedIndexes[InOutLastTriangleIndex++] / 3;
-
-	for (unsigned long TriangleIndex = 0; TriangleIndex < InTriangleCount; TriangleIndex++)
-	{
-		FaceIndex[2] = Tessellation3DData.m_puiTriangulatedIndexes[InOutLastTriangleIndex++] / 3;
-
-		if (TechSoftFileParserImpl::AddFace(FaceIndex, Tessellation, InOutLastVertexIndex))
-		{
-			TechSoftFileParserImpl::AddNormals(Tessellation3DData.m_pdNormals, NormalIndex, Tessellation.NormalArray);
-		}
-
-		TriangleIndex++;
-		if (TriangleIndex == InTriangleCount)
-		{
-			break;
-		}
-
-		Swap(FaceIndex[1], FaceIndex[2]);
-
-		NormalIndex[0] = Tessellation3DData.m_puiTriangulatedIndexes[InOutLastTriangleIndex++];
-		FaceIndex[0] = Tessellation3DData.m_puiTriangulatedIndexes[InOutLastTriangleIndex++] / 3;
-
-		if (TechSoftFileParserImpl::AddFace(FaceIndex, Tessellation, InOutLastVertexIndex))
-		{
-			TechSoftFileParserImpl::AddNormals(Tessellation3DData.m_pdNormals, NormalIndex, Tessellation.NormalArray);
-		}
-
-		Swap(FaceIndex[0], FaceIndex[1]);
-	}
-}
-
-void AddFaceTriangleStripWithUniqueNormalAndTexture(FTessellationData& Tessellation, const A3DTess3DData& Tessellation3DData, const uint32 InTriangleCount, const uint32 TextureCount, uint32& InOutLastTriangleIndex, int32& InOutLastVertexIndex)
-{
-	int32 FaceIndex[3] = { 0, 0, 0 };
-	int32 NormalIndex[3] = { 0, 0, 0 };
-	int32 TextureIndex[3] = { 0, 0, 0 };
-
-	NormalIndex[0] = Tessellation3DData.m_puiTriangulatedIndexes[InOutLastTriangleIndex++];
-	NormalIndex[1] = NormalIndex[0];
-	NormalIndex[2] = NormalIndex[0];
-
-	TextureIndex[0] = Tessellation3DData.m_puiTriangulatedIndexes[InOutLastTriangleIndex];
-	InOutLastTriangleIndex += TextureCount;
-	FaceIndex[0] = Tessellation3DData.m_puiTriangulatedIndexes[InOutLastTriangleIndex++];
-
-	TextureIndex[1] = Tessellation3DData.m_puiTriangulatedIndexes[InOutLastTriangleIndex];
-	InOutLastTriangleIndex += TextureCount;
-	FaceIndex[1] = Tessellation3DData.m_puiTriangulatedIndexes[InOutLastTriangleIndex++];
-
-	for (unsigned long TriangleIndex = 0; TriangleIndex < InTriangleCount; TriangleIndex++)
-	{
-		TextureIndex[2] = Tessellation3DData.m_puiTriangulatedIndexes[InOutLastTriangleIndex];
-		InOutLastTriangleIndex += TextureCount;
-		FaceIndex[2] = Tessellation3DData.m_puiTriangulatedIndexes[InOutLastTriangleIndex++];
-
-		if (TechSoftFileParserImpl::AddFace(FaceIndex, Tessellation, InOutLastVertexIndex))
-		{
-			TechSoftFileParserImpl::AddNormals(Tessellation3DData.m_pdNormals, NormalIndex, Tessellation.NormalArray);
-			TechSoftFileParserImpl::AddTextureCoordinates(Tessellation3DData.m_pdTextureCoords, TextureIndex, Tessellation.TexCoordArray);
-		}
-
-		TriangleIndex++;
-		if (TriangleIndex == InTriangleCount)
-		{
-			break;
-		}
-
-		Swap(FaceIndex[1], FaceIndex[2]);
-		Swap(TextureIndex[1], TextureIndex[2]);
-
-		FaceIndex[0] = Tessellation3DData.m_puiTriangulatedIndexes[InOutLastTriangleIndex++] / 3;
-
-		if (TechSoftFileParserImpl::AddFace(FaceIndex, Tessellation, InOutLastVertexIndex))
-		{
-			TechSoftFileParserImpl::AddNormals(Tessellation3DData.m_pdNormals, NormalIndex, Tessellation.NormalArray);
-		}
-
-		Swap(FaceIndex[0], FaceIndex[1]);
-		Swap(TextureIndex[0], TextureIndex[1]);
-	}
-}
-
-void AddFaceTriangleStrip(FTessellationData& Tessellation, const A3DTess3DData& Tessellation3DData, const uint32 InTriangleCount, uint32& InOutLastTriangleIndex, int32& LastVertexIndex)
-{
-	int32 FaceIndex[3] = { 0, 0, 0 };
-	int32 NormalIndex[3] = { 0, 0, 0 };
-
-	NormalIndex[0] = Tessellation3DData.m_puiTriangulatedIndexes[InOutLastTriangleIndex++];
-	FaceIndex[0] = Tessellation3DData.m_puiTriangulatedIndexes[InOutLastTriangleIndex++] / 3;
-	NormalIndex[1] = Tessellation3DData.m_puiTriangulatedIndexes[InOutLastTriangleIndex++];
-	FaceIndex[1] = Tessellation3DData.m_puiTriangulatedIndexes[InOutLastTriangleIndex++] / 3;
-
-	for (unsigned long TriangleIndex = 2; TriangleIndex < InTriangleCount; TriangleIndex++)
-	{
-		NormalIndex[2] = Tessellation3DData.m_puiTriangulatedIndexes[InOutLastTriangleIndex++];
-		FaceIndex[2] = Tessellation3DData.m_puiTriangulatedIndexes[InOutLastTriangleIndex++] / 3;
-
-		if (TechSoftFileParserImpl::AddFace(FaceIndex, Tessellation, LastVertexIndex))
-		{
-			TechSoftFileParserImpl::AddNormals(Tessellation3DData.m_pdNormals, NormalIndex, Tessellation.NormalArray);
-		}
-
-		TriangleIndex++;
-		if (TriangleIndex == InTriangleCount)
-		{
-			break;
-		}
-
-		Swap(FaceIndex[1], FaceIndex[2]);
-		Swap(NormalIndex[1], NormalIndex[2]);
-
-		NormalIndex[0] = Tessellation3DData.m_puiTriangulatedIndexes[InOutLastTriangleIndex++];
-		FaceIndex[0] = Tessellation3DData.m_puiTriangulatedIndexes[InOutLastTriangleIndex++] / 3;
-
-		if (TechSoftFileParserImpl::AddFace(FaceIndex, Tessellation, LastVertexIndex))
-		{
-			TechSoftFileParserImpl::AddNormals(Tessellation3DData.m_pdNormals, NormalIndex, Tessellation.NormalArray);
-		}
-
-		Swap(FaceIndex[0], FaceIndex[1]);
-		Swap(NormalIndex[0], NormalIndex[1]);
-	}
-}
-
-void AddFaceTriangleStripWithTexture(FTessellationData& Tessellation, const A3DTess3DData& Tessellation3DData, const uint32 InTriangleCount, const uint32 TextureCount, uint32& InOutLastTriangleIndex, int32& InOutLastVertexIndex)
-{
-	int32 FaceIndex[3] = { 0, 0, 0 };
-	int32 NormalIndex[3] = { 0, 0, 0 };
-	int32 TextureIndex[3] = { 0, 0, 0 };
-
-	NormalIndex[0] = Tessellation3DData.m_puiTriangulatedIndexes[InOutLastTriangleIndex++];
-	TextureIndex[0] = Tessellation3DData.m_puiTriangulatedIndexes[InOutLastTriangleIndex];
-	InOutLastTriangleIndex += TextureCount;
-	FaceIndex[0] = Tessellation3DData.m_puiTriangulatedIndexes[InOutLastTriangleIndex++];
-	NormalIndex[1] = Tessellation3DData.m_puiTriangulatedIndexes[InOutLastTriangleIndex++];
-	TextureIndex[1] = Tessellation3DData.m_puiTriangulatedIndexes[InOutLastTriangleIndex];
-	InOutLastTriangleIndex += TextureCount;
-	FaceIndex[1] = Tessellation3DData.m_puiTriangulatedIndexes[InOutLastTriangleIndex++];
-
-	for (unsigned long TriangleIndex = 0; TriangleIndex < InTriangleCount; TriangleIndex++)
-	{
-		NormalIndex[2] = Tessellation3DData.m_puiTriangulatedIndexes[InOutLastTriangleIndex++];
-		TextureIndex[2] = Tessellation3DData.m_puiTriangulatedIndexes[InOutLastTriangleIndex];
-		InOutLastTriangleIndex += TextureCount;
-		FaceIndex[2] = Tessellation3DData.m_puiTriangulatedIndexes[InOutLastTriangleIndex++];
-
-		if (TechSoftFileParserImpl::AddFace(FaceIndex, Tessellation, InOutLastVertexIndex))
-		{
-			TechSoftFileParserImpl::AddNormals(Tessellation3DData.m_pdNormals, NormalIndex, Tessellation.NormalArray);
-			TechSoftFileParserImpl::AddTextureCoordinates(Tessellation3DData.m_pdTextureCoords, TextureIndex, Tessellation.TexCoordArray);
-		}
-
-		TriangleIndex++;
-		if (TriangleIndex == InTriangleCount)
-		{
-			break;
-		}
-
-		Swap(FaceIndex[1], FaceIndex[2]);
-		Swap(NormalIndex[1], NormalIndex[2]);
-		Swap(TextureIndex[1], TextureIndex[2]);
-
-		NormalIndex[0] = Tessellation3DData.m_puiTriangulatedIndexes[InOutLastTriangleIndex++];
-		FaceIndex[0] = Tessellation3DData.m_puiTriangulatedIndexes[InOutLastTriangleIndex++] / 3;
-
-		if (TechSoftFileParserImpl::AddFace(FaceIndex, Tessellation, InOutLastVertexIndex))
-		{
-			TechSoftFileParserImpl::AddNormals(Tessellation3DData.m_pdNormals, NormalIndex, Tessellation.NormalArray);
-		}
-
-		Swap(FaceIndex[0], FaceIndex[1]);
-		Swap(NormalIndex[0], NormalIndex[1]);
-		Swap(TextureIndex[0], TextureIndex[1]);
-	}
-}
-
-} //ns FTechSoftFileParserImpl
-
-void FTechSoftFileParser::TraverseTessellation3D(const A3DTess3D* TessellationPtr, FArchiveBody& Body)
-{
-	FBodyMesh& BodyMesh = CADFileData.AddBodyMesh(Body.ObjectId, Body);
-
-	const int TessellationFaceDataWithTriangle = 0x2222;
-	const int TessellationFaceDataWithFan = 0x4444;
-	const int TessellationFaceDataWithStrip = 0x8888;
-	const int TessellationFaceDataWithOneNormal = 0xE0E0;
-
-	FCADUUID DefaultColorName = 0;
-	FCADUUID DefaultMaterialName = 0;
-	if (Body.ColorFaceSet.Num() > 0)
-	{
-		DefaultColorName = *Body.ColorFaceSet.begin();
-	}
-	if (Body.MaterialFaceSet.Num() > 0)
-	{
-		DefaultMaterialName = *Body.MaterialFaceSet.begin();
-	}
-
-	// Coordinates
-	TUniqueTSObj<A3DTessBaseData> TessellationBaseData(TessellationPtr);
-	{
-		if (TessellationBaseData.IsValid() && TessellationBaseData->m_uiCoordSize > 0)
-		{
-			int32 VertexCount = TessellationBaseData->m_uiCoordSize / 3;
-			BodyMesh.VertexArray.Reserve(VertexCount);
-
-			double* Coordinates = TessellationBaseData->m_pdCoords;
-			for (unsigned int Index = 0; Index < TessellationBaseData->m_uiCoordSize; ++Index)
-			{
-				Coordinates[Index] *= FileUnit;
-			}
-
-			for (unsigned int Index = 0; Index < TessellationBaseData->m_uiCoordSize; Index += 3)
-			{
-				BodyMesh.VertexArray.Emplace(Coordinates[Index], Coordinates[Index + 1], Coordinates[Index + 2]);
-			}
-		}
-		else
-		{
-			// No vertex, no mesh...
-			return;
-		}
-	}
-
-	TUniqueTSObj<A3DTess3DData> Tessellation3DData(TessellationPtr);
-	if (Tessellation3DData.IsValid())
-	{
-		for (unsigned int Index = 0; Index < Tessellation3DData->m_uiFaceTessSize; ++Index)
-		{
-			const A3DTessFaceData& FaceTessData = Tessellation3DData->m_psFaceTessData[Index];
-			FTessellationData& Tessellation = BodyMesh.Faces.Emplace_GetRef();
-
-			FCADUUID ColorName = DefaultColorName;
-			FCADUUID MaterialName = DefaultMaterialName;
-			if (FaceTessData.m_uiStyleIndexesSize == 1)
-			{
-				A3DUns32 StyleIndex = FaceTessData.m_puiStyleIndexes[0];
-				ExtractGraphStyleProperties(StyleIndex, ColorName, MaterialName);
-			}
-
-			if (ColorName)
-			{
-				Tessellation.ColorName = ColorName;
-				BodyMesh.ColorSet.Add(ColorName);
-			}
-
-			if (MaterialName)
-			{
-				Tessellation.MaterialName = MaterialName;
-				BodyMesh.MaterialSet.Add(MaterialName);
-			}
-
-			uint32 TriangleCount = TechSoftFileParserImpl::CountTriangles(FaceTessData);
-			TechSoftFileParserImpl::Reserve(Tessellation, TriangleCount, /*bWithTexture*/ FaceTessData.m_uiTextureCoordIndexesSize > 0);
-
-			uint32 UsedEntitiesFlags = FaceTessData.m_usUsedEntitiesFlags;
-			uint32 LastTrianguleIndex = FaceTessData.m_uiStartTriangulated;
-
-			uint32 FaceSetIndex = 0;
-			int32 LastVertexIndex = 0;
-
-			if (UsedEntitiesFlags & TessellationFaceDataWithTriangle)
-			{
-				bool bTessellationWithOneNormal = (UsedEntitiesFlags & TessellationFaceDataWithOneNormal);
-				if (bTessellationWithOneNormal)
-				{
-					if (FaceTessData.m_uiTextureCoordIndexesSize)
-					{
-						TechSoftFileParserImpl::AddFaceTriangleWithUniqueNormalAndTexture(Tessellation, *Tessellation3DData, FaceTessData.m_puiSizesTriangulated[0], FaceTessData.m_uiTextureCoordIndexesSize, LastTrianguleIndex, LastVertexIndex);
-					}
-					else
-					{
-						TechSoftFileParserImpl::AddFaceTriangleWithUniqueNormal(Tessellation, *Tessellation3DData, FaceTessData.m_puiSizesTriangulated[0], LastTrianguleIndex, LastVertexIndex);
-					}
-				}
-				else
-				{
-					if (FaceTessData.m_uiTextureCoordIndexesSize)
-					{
-						TechSoftFileParserImpl::AddFaceTriangleWithTexture(Tessellation, *Tessellation3DData, FaceTessData.m_puiSizesTriangulated[0], FaceTessData.m_uiTextureCoordIndexesSize, LastTrianguleIndex, LastVertexIndex);
-					}
-					else
-					{
-						TechSoftFileParserImpl::AddFaceTriangle(Tessellation, *Tessellation3DData, FaceTessData.m_puiSizesTriangulated[0], LastTrianguleIndex, LastVertexIndex);
-					}
-				}
-				FaceSetIndex++;
-			}
-
-			if (FaceTessData.m_uiSizesTriangulatedSize > FaceSetIndex)
-			{
-				//if (UsedEntitiesFlags & TessellationFaceDataWithFan)
-				//{
-				//	unsigned int LastFanIndex = 1 + FaceSetIndex + FaceTessData.m_puiSizesTriangulated[FaceSetIndex];
-				//	FaceSetIndex++;
-				//	for (; FaceSetIndex < LastFanIndex; FaceSetIndex++)
-				//	{
-				//		bool bTessellationWithOneNormal = (UsedEntitiesFlags & TessellationFaceDataWithOneNormal) && (FaceTessData.m_puiSizesTriangulated[FaceSetIndex] & kA3DTessFaceDataNormalSingle);
-				//		A3DUns32 FanSize = (FaceTessData.m_puiSizesTriangulated[FaceSetIndex] & kA3DTessFaceDataNormalMask);
-				//		if (bTessellationWithOneNormal)
-				//		{
-				//			if (FaceTessData.m_uiTextureCoordIndexesSize)
-				//			{
-				//				TechSoftFileParserImpl::AddFaceTriangleFanWithUniqueNormalAndTexture(Tessellation, *Tessellation3DData, FaceTessData.m_puiSizesTriangulated[FaceSetIndex], FaceTessData.m_uiTextureCoordIndexesSize, LastTrianguleIndex, LastVertexIndex);
-				//			}
-				//			else
-				//			{
-				//				TechSoftFileParserImpl::AddFaceTriangleFanWithUniqueNormal(Tessellation, *Tessellation3DData, FaceTessData.m_puiSizesTriangulated[FaceSetIndex], LastTrianguleIndex, LastVertexIndex);
-				//			}
-				//		}
-				//		else
-				//		{
-				//			if (FaceTessData.m_uiTextureCoordIndexesSize)
-				//			{
-				//				TechSoftFileParserImpl::AddFaceTriangleFanWithTexture(Tessellation, *Tessellation3DData, FaceTessData.m_puiSizesTriangulated[FaceSetIndex], FaceTessData.m_uiTextureCoordIndexesSize, LastTrianguleIndex, LastVertexIndex);
-				//			}
-				//			else
-				//			{
-				//				TechSoftFileParserImpl::AddFaceTriangleFan(Tessellation, *Tessellation3DData, FaceTessData.m_puiSizesTriangulated[FaceSetIndex], LastTrianguleIndex, LastVertexIndex);
-				//			}
-				//		}
-				//	}
-				//}
-
-				if (UsedEntitiesFlags & kA3DTessFaceDataTriangleFan)
-				{
-					uint32 FanCount = FaceTessData.m_puiSizesTriangulated[FaceSetIndex++];
-					for (uint32 FanIndex = 0; FanIndex < FanCount; ++FanIndex)
-					{
-						uint32 VertexCount = FaceTessData.m_puiSizesTriangulated[FaceSetIndex++];
-						TechSoftFileParserImpl::AddFaceTriangleFan(Tessellation, *Tessellation3DData, VertexCount, LastTrianguleIndex, LastVertexIndex);
-					}
-				}
-
-				if (UsedEntitiesFlags & kA3DTessFaceDataTriangleFanOneNormal)
-				{
-					uint32 FanCount = FaceTessData.m_puiSizesTriangulated[FaceSetIndex++] & kA3DTessFaceDataNormalMask;
-					for (uint32 FanIndex = 0; FanIndex < FanCount; ++FanIndex)
-					{
-						ensure((FaceTessData.m_puiSizesTriangulated[FaceSetIndex] & kA3DTessFaceDataNormalSingle) != 0);
-
-						uint32 VertexCount = FaceTessData.m_puiSizesTriangulated[FaceSetIndex++] & kA3DTessFaceDataNormalMask;
-						TechSoftFileParserImpl::AddFaceTriangleFanWithUniqueNormal(Tessellation, *Tessellation3DData, VertexCount, LastTrianguleIndex, LastVertexIndex);
-					}
-				}
-
-				if (UsedEntitiesFlags & kA3DTessFaceDataTriangleFanTextured)
-				{
-					uint32 FanCount = FaceTessData.m_puiSizesTriangulated[FaceSetIndex++];
-					for (uint32 FanIndex = 0; FanIndex < FanCount; ++FanIndex)
-					{
-						uint32 VertexCount = FaceTessData.m_puiSizesTriangulated[FaceSetIndex++];
-						TechSoftFileParserImpl::AddFaceTriangleFanWithTexture(Tessellation, *Tessellation3DData, VertexCount, FaceTessData.m_uiTextureCoordIndexesSize, LastTrianguleIndex, LastVertexIndex);
-					}
-				}
-
-				if (UsedEntitiesFlags & kA3DTessFaceDataTriangleFanOneNormalTextured)
-				{
-					uint32 FanCount = FaceTessData.m_puiSizesTriangulated[FaceSetIndex++] & kA3DTessFaceDataNormalMask;
-					for (uint32 FanIndex = 0; FanIndex < FanCount; ++FanIndex)
-					{
-						ensure((FaceTessData.m_puiSizesTriangulated[FaceSetIndex] & kA3DTessFaceDataNormalSingle) != 0);
-
-						uint32 VertexCount = FaceTessData.m_puiSizesTriangulated[FaceSetIndex++] & kA3DTessFaceDataNormalMask;
-						TechSoftFileParserImpl::AddFaceTriangleFanWithUniqueNormalAndTexture(Tessellation, *Tessellation3DData, VertexCount, FaceTessData.m_uiTextureCoordIndexesSize, LastTrianguleIndex, LastVertexIndex);
-					}
-				}
-			}
-
-			if (FaceTessData.m_uiSizesTriangulatedSize > FaceSetIndex)
-			{
-				FaceSetIndex++;
-				for (; FaceSetIndex < FaceTessData.m_uiSizesTriangulatedSize; FaceSetIndex++)
-				{
-					bool bTessellationWithOneNormal = (UsedEntitiesFlags & TessellationFaceDataWithOneNormal) && (FaceTessData.m_puiSizesTriangulated[FaceSetIndex] & kA3DTessFaceDataNormalSingle);
-					A3DUns32 StripSize = (FaceTessData.m_puiSizesTriangulated[FaceSetIndex] & kA3DTessFaceDataNormalMask);
-					if (bTessellationWithOneNormal)
-					{
-						if (FaceTessData.m_uiTextureCoordIndexesSize)
-						{
-							TechSoftFileParserImpl::AddFaceTriangleStripWithUniqueNormalAndTexture(Tessellation, *Tessellation3DData, FaceTessData.m_puiSizesTriangulated[FaceSetIndex], FaceTessData.m_uiTextureCoordIndexesSize, LastTrianguleIndex, LastVertexIndex);
-						}
-						else
-						{
-							TechSoftFileParserImpl::AddFaceTriangleStripWithUniqueNormal(Tessellation, *Tessellation3DData, FaceTessData.m_puiSizesTriangulated[FaceSetIndex], LastTrianguleIndex, LastVertexIndex);
-						}
-					}
-					else
-					{
-						if (FaceTessData.m_uiTextureCoordIndexesSize)
-						{
-							TechSoftFileParserImpl::AddFaceTriangleStripWithTexture(Tessellation, *Tessellation3DData, FaceTessData.m_puiSizesTriangulated[FaceSetIndex], FaceTessData.m_uiTextureCoordIndexesSize, LastTrianguleIndex, LastVertexIndex);
-						}
-						else
-						{
-							TechSoftFileParserImpl::AddFaceTriangleStrip(Tessellation, *Tessellation3DData, FaceTessData.m_puiSizesTriangulated[FaceSetIndex], LastTrianguleIndex, LastVertexIndex);
-						}
-					}
-				}
-			}
-
-		}
-	}
-
-	Body.ColorFaceSet = BodyMesh.ColorSet;
-	Body.MaterialFaceSet = BodyMesh.MaterialSet;
-}
-
 } // ns CADLibrary
 
 #endif  
\ No newline at end of file
diff --git a/Engine/Plugins/Enterprise/DatasmithCADImporter/Source/CADInterfaces/Private/TechSoftFileParser.h b/Engine/Plugins/Enterprise/DatasmithCADImporter/Source/CADInterfaces/Private/TechSoftFileParser.h
--- a/Engine/Plugins/Enterprise/DatasmithCADImporter/Source/CADInterfaces/Private/TechSoftFileParser.h
+++ b/Engine/Plugins/Enterprise/DatasmithCADImporter/Source/CADInterfaces/Private/TechSoftFileParser.h
@@ -11,8 +11,6 @@
 
 #include "TUniqueTechSoftObj.h"
 
-#define NEW_CODE
-
 typedef void A3DAsmModelFile;
 typedef void A3DAsmPartDefinition;
 typedef void A3DAsmProductOccurrence;
@@ -160,20 +158,15 @@ private:
 	bool IsConfigurationSet(const A3DAsmProductOccurrence* Occurrence);
 	void TraverseConfigurationSet(const A3DAsmProductOccurrence* ConfigurationSet);
 	FCadId TraverseOccurrence(const A3DAsmProductOccurrence* Occurrence);
-	void TraversePrototype(const A3DAsmProductOccurrence* Prototype, FArchiveComponent& Component);
 	void ProcessPrototype(const A3DAsmProductOccurrence* InPrototype, FEntityMetaData& OutMetaData, A3DMiscTransformation** OutLocation);
-	void ProcessOccurrence(TUniqueTSObj<A3DAsmProductOccurrenceData>& Occurrence, FArchiveComponent& Component);
 	void TraversePartDefinition(const A3DAsmPartDefinition* PartDefinition, FArchiveComponent& Component);
 	FCadId TraverseRepresentationSet(const A3DRiSet* pSet, FEntityMetaData& PartMetaData);
 	FCadId TraverseRepresentationItem(A3DRiRepresentationItem* RepresentationItem, FEntityMetaData& PartMetaData);
 	FCadId TraverseBRepModel(A3DRiBrepModel* BrepModel, FEntityMetaData& PartMetaData);
-	FCadId TraversePolyBRepModel(const A3DRiPolyBrepModel* PolygonalBrepModel, FEntityMetaData& PartMetaData);
+	FCadId TraversePolyBRepModel(A3DRiPolyBrepModel* PolygonalBrepModel, FEntityMetaData& PartMetaData);
 
 	// Tessellation methods
-	void MeshRepresentationWithTechSoft(A3DRiRepresentationItem* RepresentationItem, FArchiveBody& Body);
-	void TraverseRepresentationContent(const A3DRiRepresentationItem* RepresentationItem, FArchiveBody& Body);
-	void TraverseTessellationBase(const A3DTessBase* Tessellation, FArchiveBody& Body);
-	void TraverseTessellation3D(const A3DTess3D* Tessellation, FArchiveBody& Body);
+	void GenerateBodyMeshes();
 
 	// MetaData
 	void ExtractMetaData(const A3DEntity* Entity, FEntityMetaData& OutMetaData);
@@ -218,7 +211,7 @@ private:
 	FArchiveUnloadedComponent& AddUnloadedComponent(FEntityMetaData& ComponentMetaData, FArchiveInstance& Instance);
 	FArchiveComponent& AddOccurence(FEntityMetaData& InstanceMetaData, FEntityMetaData& ReferenceMetaData, FCadId& OutComponentId);
 	FArchiveComponent& AddOccurence(FEntityMetaData& InstanceMetaData, FCadId& OutComponentId);
-	FArchiveBody& AddBody(FEntityMetaData& BodyMetaData);
+	int32 AddBody(FEntityMetaData& BodyMetaData);
 #endif
 
 private:
@@ -236,75 +229,27 @@ private:
 	FCadId LastEntityId = 1;
 
 	// 
-	TMap<const A3DRiRepresentationItem*, FCadId> RepresentationItemsCache;
+	TMap<A3DRiRepresentationItem*, int32> RepresentationItemsCache;
 };
 
 namespace TechSoftFileParserImpl
 {
-// Methodes used by TraverseTessellation3D
-
-typedef double A3DDouble;
-inline bool AddFace(int32 FaceIndex[3], CADLibrary::FTessellationData& Tessellation, int32& InOutVertexIndex)
-{
-	if (FaceIndex[0] == FaceIndex[1] || FaceIndex[0] == FaceIndex[2] || FaceIndex[1] == FaceIndex[2])
-	{
-		return false;
-	}
-
-	for (int32 Index = 0; Index < 3; ++Index)
-	{
-		Tessellation.VertexIndices.Add(InOutVertexIndex++);
-	}
-	Tessellation.PositionIndices.Append(FaceIndex, 3);
-	return true;
-};
-
-inline void AddNormals(const A3DDouble* Normals, const int32 Indices[3], TArray<FVector>& TessellationNormals)
-{
-	for (int32 Index = 0; Index < 3; ++Index)
-	{
-		int32 NormalIndex = Indices[Index];
-		TessellationNormals.Emplace(Normals[NormalIndex], Normals[NormalIndex + 1], Normals[NormalIndex + 2]);
-	}
-};
-
-inline void AddTextureCoordinates(const A3DDouble* TextureCoords, const int32 Indices[3], TArray<FVector2D>& TessellationTextures)
-{
-	for (int32 Index = 0; Index < 3; ++Index)
-	{
-		int32 TextureIndex = Indices[Index];
-		TessellationTextures.Emplace(TextureCoords[TextureIndex], TextureCoords[TextureIndex + 1]);
-	}
-};
-
-inline void Reserve(CADLibrary::FTessellationData& Tessellation, int32 InTrinangleCount, bool bWithTexture)
-{
-	Tessellation.PositionIndices.Reserve(3 * InTrinangleCount);
-	Tessellation.VertexIndices.Reserve(3 * InTrinangleCount);
-	Tessellation.NormalArray.Reserve(3 * InTrinangleCount);
-	if (bWithTexture)
-	{
-		Tessellation.TexCoordArray.Reserve(3 * InTrinangleCount);
-	}
-};
-
 #ifdef USE_TECHSOFT_SDK
-inline FColor GetColorAt(uint32 ColorIndex)
-{
-	TUniqueTSObjFromIndex<A3DGraphRgbColorData> ColorData(ColorIndex);
-	if (ColorData.IsValid())
+	inline FColor GetColorAt(uint32 ColorIndex)
 	{
-		return FColor((uint8)(ColorData->m_dRed * 255)
-			, (uint8)(ColorData->m_dGreen * 255)
-			, (uint8)(ColorData->m_dBlue * 255));
+		TUniqueTSObjFromIndex<A3DGraphRgbColorData> ColorData(ColorIndex);
+		if (ColorData.IsValid())
+		{
+			return FColor((uint8)(ColorData->m_dRed * 255)
+				, (uint8)(ColorData->m_dGreen * 255)
+				, (uint8)(ColorData->m_dBlue * 255));
+		}
+		else
+		{
+			return FColor(200, 200, 200);
+		}
 	}
-	else
-	{
-		return FColor(200, 200, 200);
-	}
-}
 #endif
-
 } // ns TechSoftFileParserImpl
 
 } // ns CADLibrary
\ No newline at end of file
diff --git a/Engine/Plugins/Enterprise/DatasmithCADImporter/Source/CADInterfaces/Private/TechSoftInterface.cpp b/Engine/Plugins/Enterprise/DatasmithCADImporter/Source/CADInterfaces/Private/TechSoftInterface.cpp
--- a/Engine/Plugins/Enterprise/DatasmithCADImporter/Source/CADInterfaces/Private/TechSoftInterface.cpp
+++ b/Engine/Plugins/Enterprise/DatasmithCADImporter/Source/CADInterfaces/Private/TechSoftInterface.cpp
@@ -4,11 +4,14 @@
 
 #include "TechSoftInterface.h"
 
+#include "CADOptions.h"
 #include "CADInterfacesModule.h"
+
 #include "HAL/PlatformProcess.h"
 #include "Misc/Paths.h"
 #include "TUniqueTechSoftObj.h"
 
+#include "TechSoftInterfaceUtils.inl"
 
 namespace CADLibrary
 {
@@ -65,6 +68,146 @@ bool FTechSoftInterface::InitializeKernel(const TCHAR* InEnginePluginsPath)
 #endif
 }
 
+void FTechSoftInterface::SaveBodyToHsfFile(void* BodyPtr, const FString& Filename)
+{
+	if (!BodyPtr)
+	{
+		return;
+	}
+
+#ifdef USE_TECHSOFT_SDK
+	A3DRiRepresentationItem* RepresentationItem = (A3DRiRepresentationItem*)BodyPtr;
+	A3DAsmPartDefinition* PartDefinition = nullptr;
+
+	A3DAsmPartDefinitionData PartDefinitionData;
+	A3D_INITIALIZE_DATA(A3DAsmPartDefinitionData, PartDefinitionData);
+
+	PartDefinitionData.m_uiRepItemsSize = 1;
+	PartDefinitionData.m_ppRepItems = &RepresentationItem;
+
+	A3DAsmPartDefinitionCreate(&PartDefinitionData, &PartDefinition);
+
+	A3DAsmProductOccurrence* ProductOccurrence = nullptr;
+
+	A3DAsmProductOccurrenceData ProductOccurrenceData;
+	A3D_INITIALIZE_DATA(A3DAsmProductOccurrenceData, ProductOccurrenceData);
+
+	ProductOccurrenceData.m_pPart = PartDefinition;
+
+	A3DInt32 iRet = A3DAsmProductOccurrenceCreate(&ProductOccurrenceData, &ProductOccurrence);
+
+	A3DAsmModelFile* ModelFile = nullptr;
+
+	A3DAsmModelFileData ModelFileData;
+	A3D_INITIALIZE_DATA(A3DAsmModelFileData, ModelFileData);
+
+	ModelFileData.m_uiPOccurrencesSize = 1;
+	ModelFileData.m_dUnit = 1.0;
+	ModelFileData.m_ppPOccurrences = &ProductOccurrence;
+
+	A3DAsmModelFileCreate(&ModelFileData, &ModelFile);
+
+	A3DRWParamsExportPrcData sParamsExportData;
+	A3D_INITIALIZE_DATA(A3DRWParamsExportPrcData, sParamsExportData);
+
+	sParamsExportData.m_bCompressBrep = false;
+	sParamsExportData.m_bCompressTessellation = false;
+
+	A3DUTF8Char HsfFileName[_MAX_PATH];
+	FCStringAnsi::Strncpy(HsfFileName, TCHAR_TO_UTF8(*Filename), _MAX_PATH);
+
+	iRet = A3DAsmModelFileExportToPrcFile(ModelFile, &sParamsExportData, HsfFileName, nullptr);
+
+	// #ueent_techsoft: Deleting the model seems to delete the entire content. To be double-checked
+	A3DEntityDelete(ModelFile);
+#endif
+}
+
+bool FTechSoftInterface::GetBodyFromHsfFile(const FString& Filename, const FImportParameters& ImportParameters, double FileUnit, FBodyMesh& BodyMesh)
+{
+#ifdef USE_TECHSOFT_SDK
+	A3DRWParamsPrcReadHelper* ReadHelper = nullptr;
+	A3DAsmModelFile* ModelFile = nullptr;
+
+	A3DInt32 iRet = A3DAsmModelFileLoadFromPrcFile(TCHAR_TO_UTF8(*Filename), &ReadHelper, &ModelFile);
+
+	if (iRet != A3D_SUCCESS || !ModelFile)
+	{
+		return false;
+	}
+
+	TUniqueTSObj<A3DAsmModelFileData> ModelFileData(ModelFile);
+	if (!ModelFileData.IsValid() || ModelFileData->m_uiPOccurrencesSize == 0)
+	{
+		return false;
+	}
+
+	TUniqueTSObj<A3DAsmProductOccurrenceData> OccurenceData(ModelFileData->m_ppPOccurrences[0]);
+	if (!OccurenceData.IsValid() || !OccurenceData->m_pPart)
+	{
+		return false;
+	}
+
+	TUniqueTSObj<A3DAsmPartDefinitionData> PartDefinitionData(OccurenceData->m_pPart);
+	if (!PartDefinitionData.IsValid() || PartDefinitionData->m_uiRepItemsSize == 0)
+	{
+		return false;
+	}
+
+	return FillBodyMesh(PartDefinitionData->m_ppRepItems[0], ImportParameters, FileUnit, BodyMesh);
+#else
+	return false;
+#endif
+}
+
+bool FTechSoftInterface::FillBodyMesh(void* BodyPtr, const FImportParameters& ImportParameters, double FileUnit, FBodyMesh& BodyMesh)
+{
+#ifdef USE_TECHSOFT_SDK
+	A3DRiRepresentationItem* RepresentationItemPtr = (A3DRiRepresentationItem*)BodyPtr;
+
+	A3DEEntityType Type;
+	A3DEntityGetType(RepresentationItemPtr, &Type);
+	if (Type == kA3DTypeRiPolyBrepModel)
+	{
+		TUniqueTSObj<A3DRiRepresentationItemData> RepresentationItemData(RepresentationItemPtr);
+		TechSoftInterfaceUtils::FTechSoftTessellationExtractor Extractor(RepresentationItemData->m_pTessBase);
+		return Extractor.FillBodyMesh(BodyMesh, FileUnit);
+	}
+
+	TUniqueTSObj<A3DRiRepresentationItemData> RepresentationItemData;
+
+	// TUniqueTechSoftObj does not work in this case
+	A3DRWParamsTessellationData TessellationParameters;
+	A3D_INITIALIZE_DATA(A3DRWParamsTessellationData, TessellationParameters);
+
+	TessellationParameters.m_eTessellationLevelOfDetail = kA3DTessLODUserDefined; // Enum to specify predefined values for some following members.
+	TessellationParameters.m_bUseHeightInsteadOfRatio = A3D_TRUE;
+	TessellationParameters.m_dMaxChordHeight = ImportParameters.GetChordTolerance() * 10. / FileUnit;
+	TessellationParameters.m_dAngleToleranceDeg = ImportParameters.GetMaxNormalAngle();
+	TessellationParameters.m_dMaximalTriangleEdgeLength = 0; //ImportParameters.MaxEdgeLength;
+
+	TessellationParameters.m_bAccurateTessellation = A3D_FALSE;  // A3D_FALSE' indicates the tessellation is set for visualization
+	TessellationParameters.m_bAccurateTessellationWithGrid = A3D_FALSE; // Enable accurate tessellation with faces inner points on a grid.
+	TessellationParameters.m_dAccurateTessellationWithGridMaximumStitchLength = 0; 	// Maximal grid stitch length. Disabled if value is 0. Be careful, a too small value can generate a huge tessellation.
+
+	TessellationParameters.m_bKeepUVPoints = A3D_TRUE; // Keep parametric points as texture points.
+
+	// Get the tessellation
+	A3DStatus Status = A3DRiRepresentationItemComputeTessellation(RepresentationItemPtr, &TessellationParameters);
+	Status = A3DRiRepresentationItemGet(RepresentationItemPtr, RepresentationItemData.GetEmptyDataPtr());
+
+	{
+		A3DEntityGetType(RepresentationItemData->m_pTessBase, &Type);
+		ensure(Type == kA3DTypeTess3D);
+	}
+
+	TechSoftInterfaceUtils::FTechSoftTessellationExtractor Extractor(RepresentationItemData->m_pTessBase);
+	return Extractor.FillBodyMesh(BodyMesh, FileUnit);
+#else
+return false;
+#endif
+}
+
 #ifdef USE_TECHSOFT_SDK
 
 A3DStatus FTechSoftInterface::Import(const A3DImport& Import)
diff --git a/Engine/Plugins/Enterprise/DatasmithCADImporter/Source/CADInterfaces/Private/TechSoftInterfaceUtils.inl b/Engine/Plugins/Enterprise/DatasmithCADImporter/Source/CADInterfaces/Private/TechSoftInterfaceUtils.inl
new file mode 100644
--- /dev/null
+++ b/Engine/Plugins/Enterprise/DatasmithCADImporter/Source/CADInterfaces/Private/TechSoftInterfaceUtils.inl
@@ -0,0 +1,741 @@
+#pragma once
+
+#include "CADData.h"
+#include "CADSceneGraph.h"
+
+#include "TUniqueTechSoftObj.h"
+
+#ifdef USE_TECHSOFT_SDK
+
+namespace TechSoftInterfaceUtils
+{
+	class FTechSoftTessellationExtractor
+	{
+	public:
+		FTechSoftTessellationExtractor(const A3DTess3D* InTessellationPtr)
+			: TessellationPtr(InTessellationPtr)
+		{
+		}
+
+		bool  FillBodyMesh(CADLibrary::FBodyMesh& BodyMesh, double FileUnit)
+		{
+			FillVertexArray(FileUnit, BodyMesh.VertexArray);
+
+			if (BodyMesh.VertexArray.Num() == 0)
+			{
+				return false;
+			}
+
+			FillFaceArray(BodyMesh);
+
+			return BodyMesh.Faces.Num() > 0 ? true : false;
+		}
+
+	private:
+		void FillVertexArray(double FileUnit, TArray<FVector>& VertexArray)
+		{
+			using namespace CADLibrary;
+
+			TUniqueTSObj<A3DTessBaseData> TessellationBaseData(TessellationPtr);
+
+			if (!TessellationBaseData.IsValid() || TessellationBaseData->m_uiCoordSize == 0)
+			{
+				return;
+			}
+
+			int32 VertexCount = TessellationBaseData->m_uiCoordSize / 3;
+			VertexArray.Reserve(VertexCount);
+
+			double* Coordinates = TessellationBaseData->m_pdCoords;
+			for (unsigned int Index = 0; Index < TessellationBaseData->m_uiCoordSize; ++Index)
+			{
+				Coordinates[Index] *= FileUnit;
+			}
+
+			for (unsigned int Index = 0; Index < TessellationBaseData->m_uiCoordSize; Index += 3)
+			{
+				VertexArray.Emplace(Coordinates[Index], Coordinates[Index + 1], Coordinates[Index + 2]);
+			}
+		}
+
+		// #ueent_techsoft: TODO: Make it more in line with the actual implementation of FillFaceArray
+		uint32 CountTriangles(const A3DTessFaceData& FaceTessData)
+		{
+			const int32 TessellationFaceDataWithTriangle = 0x2222;
+			const int32 TessellationFaceDataWithFan = 0x4444;
+			const int32 TessellationFaceDataWithStripe = 0x8888;
+			const int32 TessellationFaceDataWithOneNormal = 0xE0E0;
+
+			uint32 UsedEntitiesFlags = FaceTessData.m_usUsedEntitiesFlags;
+
+			uint32 TriangleCount = 0;
+			uint32 FaceSetIndex = 0;
+			if (UsedEntitiesFlags & TessellationFaceDataWithTriangle)
+			{
+				TriangleCount += FaceTessData.m_puiSizesTriangulated[FaceSetIndex];
+				FaceSetIndex++;
+			}
+
+			if (FaceTessData.m_uiSizesTriangulatedSize > FaceSetIndex)
+			{
+				if (UsedEntitiesFlags & TessellationFaceDataWithFan)
+				{
+					uint32 LastFanIndex = 1 + FaceSetIndex + FaceTessData.m_puiSizesTriangulated[FaceSetIndex];
+					FaceSetIndex++;
+					for (; FaceSetIndex < LastFanIndex; FaceSetIndex++)
+					{
+						uint32 FanSize = (FaceTessData.m_puiSizesTriangulated[FaceSetIndex] & kA3DTessFaceDataNormalMask);
+						TriangleCount += (FanSize - 2);
+					}
+				}
+			}
+
+			if (FaceTessData.m_uiSizesTriangulatedSize > FaceSetIndex)
+			{
+				FaceSetIndex++;
+				for (; FaceSetIndex < FaceTessData.m_uiSizesTriangulatedSize; FaceSetIndex++)
+				{
+					uint32 StripeSize = (FaceTessData.m_puiSizesTriangulated[FaceSetIndex] & kA3DTessFaceDataNormalMask);
+					TriangleCount += (StripeSize - 2);
+				}
+			}
+			return TriangleCount;
+		}
+
+#pragma optimize ("", off)
+		void FillFaceArray(CADLibrary::FBodyMesh& BodyMesh)
+		{
+			using namespace CADLibrary;
+
+			TArray<FTessellationData>& Faces = BodyMesh.Faces;
+
+			TUniqueTSObj<A3DTess3DData> TessellationData(TessellationPtr);
+
+			if (!TessellationData.IsValid() || TessellationData->m_uiFaceTessSize == 0)
+			{
+				return;
+			}
+
+			TessellationNormals = TessellationData->m_pdNormals;
+			TessellationTexCoords = TessellationData->m_pdTextureCoords;
+			TriangulatedIndexes = TessellationData->m_puiTriangulatedIndexes;
+
+			for (unsigned int Index = 0; Index < TessellationData->m_uiFaceTessSize; ++Index)
+			{
+				const A3DTessFaceData& FaceTessData = TessellationData->m_psFaceTessData[Index];
+				FTessellationData& Tessellation = Faces.Emplace_GetRef();
+
+				Tessellation.MaterialName = 0xffffffff; // Assumes -1 is an invalid value for m_puiStyleIndexes
+				if (FaceTessData.m_uiStyleIndexesSize == 1)
+				{
+					// Store the StyleIndex on the MaterialName. It will be processed after tessellation
+					Tessellation.MaterialName = FaceTessData.m_puiStyleIndexes[0];
+				}
+
+				// Pre-allocate memory for triangles' data
+				uint32 TriangleCount = CountTriangles(FaceTessData);
+				Tessellation.PositionIndices.Reserve(3 * TriangleCount);
+				Tessellation.VertexIndices.Reserve(3 * TriangleCount);
+				Tessellation.NormalArray.Reserve(3 * TriangleCount);
+				if (FaceTessData.m_uiTextureCoordIndexesSize > 0)
+				{
+					Tessellation.TexCoordArray.Reserve(3 * TriangleCount);
+				}
+
+				uint32 UsedEntitiesFlags = FaceTessData.m_usUsedEntitiesFlags;
+
+				LastTrianguleIndex = FaceTessData.m_uiStartTriangulated;
+				LastVertexIndex = 0;
+
+				uint32 FaceSetIndex = 0;
+				bool bMustProcess = FaceTessData.m_uiSizesTriangulatedSize > FaceSetIndex;
+
+				if (bMustProcess && UsedEntitiesFlags & kA3DTessFaceDataTriangle)
+				{
+					AddFaceTriangle(Tessellation, FaceTessData.m_puiSizesTriangulated[FaceSetIndex++], LastTrianguleIndex, LastVertexIndex);
+					bMustProcess = FaceTessData.m_uiSizesTriangulatedSize > FaceSetIndex;
+				}
+
+				if (bMustProcess && UsedEntitiesFlags & kA3DTessFaceDataTriangleOneNormal)
+				{
+					AddFaceTriangleWithUniqueNormal(Tessellation, FaceTessData.m_puiSizesTriangulated[FaceSetIndex++], LastTrianguleIndex, LastVertexIndex);
+					bMustProcess = FaceTessData.m_uiSizesTriangulatedSize > FaceSetIndex;
+				}
+
+				if (bMustProcess && UsedEntitiesFlags & kA3DTessFaceDataTriangleTextured)
+				{
+					AddFaceTriangleWithTexture(Tessellation, FaceTessData.m_puiSizesTriangulated[FaceSetIndex++], FaceTessData.m_uiTextureCoordIndexesSize, LastTrianguleIndex, LastVertexIndex);
+					bMustProcess = FaceTessData.m_uiSizesTriangulatedSize > FaceSetIndex;
+				}
+
+				if (bMustProcess && UsedEntitiesFlags & kA3DTessFaceDataTriangleOneNormalTextured)
+				{
+					AddFaceTriangleWithUniqueNormalAndTexture(Tessellation, FaceTessData.m_puiSizesTriangulated[FaceSetIndex++], FaceTessData.m_uiTextureCoordIndexesSize, LastTrianguleIndex, LastVertexIndex);
+					bMustProcess = FaceTessData.m_uiSizesTriangulatedSize > FaceSetIndex;
+				}
+
+				if (bMustProcess && UsedEntitiesFlags & kA3DTessFaceDataTriangleFan)
+				{
+					uint32 FanCount = FaceTessData.m_puiSizesTriangulated[FaceSetIndex++];
+					for (uint32 FanIndex = 0; FanIndex < FanCount; ++FanIndex)
+					{
+						uint32 VertexCount = FaceTessData.m_puiSizesTriangulated[FaceSetIndex++];
+						AddFaceTriangleFan(Tessellation, VertexCount, LastTrianguleIndex, LastVertexIndex);
+					}
+					bMustProcess = FaceTessData.m_uiSizesTriangulatedSize > FaceSetIndex;
+				}
+
+				if (bMustProcess && UsedEntitiesFlags & kA3DTessFaceDataTriangleFanOneNormal)
+				{
+					uint32 FanCount = FaceTessData.m_puiSizesTriangulated[FaceSetIndex++] & kA3DTessFaceDataNormalMask;
+					for (uint32 FanIndex = 0; FanIndex < FanCount; ++FanIndex)
+					{
+						ensure((FaceTessData.m_puiSizesTriangulated[FaceSetIndex] & kA3DTessFaceDataNormalSingle) != 0);
+
+						uint32 VertexCount = FaceTessData.m_puiSizesTriangulated[FaceSetIndex++] & kA3DTessFaceDataNormalMask;
+						AddFaceTriangleFanWithUniqueNormal(Tessellation, VertexCount, LastTrianguleIndex, LastVertexIndex);
+					}
+					bMustProcess = FaceTessData.m_uiSizesTriangulatedSize > FaceSetIndex;
+				}
+
+				if (bMustProcess && UsedEntitiesFlags & kA3DTessFaceDataTriangleFanTextured)
+				{
+					uint32 FanCount = FaceTessData.m_puiSizesTriangulated[FaceSetIndex++];
+					for (uint32 FanIndex = 0; FanIndex < FanCount; ++FanIndex)
+					{
+						uint32 VertexCount = FaceTessData.m_puiSizesTriangulated[FaceSetIndex++];
+						AddFaceTriangleFanWithTexture(Tessellation, VertexCount, FaceTessData.m_uiTextureCoordIndexesSize, LastTrianguleIndex, LastVertexIndex);
+					}
+					bMustProcess = FaceTessData.m_uiSizesTriangulatedSize > FaceSetIndex;
+				}
+
+				if (bMustProcess && UsedEntitiesFlags & kA3DTessFaceDataTriangleFanOneNormalTextured)
+				{
+					uint32 FanCount = FaceTessData.m_puiSizesTriangulated[FaceSetIndex++] & kA3DTessFaceDataNormalMask;
+					for (uint32 FanIndex = 0; FanIndex < FanCount; ++FanIndex)
+					{
+						ensure((FaceTessData.m_puiSizesTriangulated[FaceSetIndex] & kA3DTessFaceDataNormalSingle) != 0);
+
+						uint32 VertexCount = FaceTessData.m_puiSizesTriangulated[FaceSetIndex++] & kA3DTessFaceDataNormalMask;
+						AddFaceTriangleFanWithUniqueNormalAndTexture(Tessellation, VertexCount, FaceTessData.m_uiTextureCoordIndexesSize, LastTrianguleIndex, LastVertexIndex);
+					}
+					bMustProcess = FaceTessData.m_uiSizesTriangulatedSize > FaceSetIndex;
+				}
+
+				if (bMustProcess && UsedEntitiesFlags & kA3DTessFaceDataTriangleStripe)
+				{
+					A3DUns32 StripeSize = FaceTessData.m_puiSizesTriangulated[FaceSetIndex++];
+					for (A3DUns32 StripeIndex = 0; StripeIndex < StripeSize; ++StripeIndex)
+					{
+						A3DUns32 PointCount = FaceTessData.m_puiSizesTriangulated[FaceSetIndex++];
+						AddFaceTriangleStripe(Tessellation, PointCount, LastTrianguleIndex, LastVertexIndex);
+					}
+					bMustProcess = FaceTessData.m_uiSizesTriangulatedSize > FaceSetIndex;
+				}
+
+				if (bMustProcess && UsedEntitiesFlags & kA3DTessFaceDataTriangleStripeOneNormal)
+				{
+					A3DUns32 StripeSize = FaceTessData.m_puiSizesTriangulated[FaceSetIndex++] & kA3DTessFaceDataNormalMask;
+					for (A3DUns32 StripeIndex = 0; StripeIndex < StripeSize; ++StripeIndex)
+					{
+						// #ueent_techsoft: There are some specificities with the case of oneNormal
+						//					See TessConnector.cpp around line 237
+						//	bool bIsOneNormal = (pFaceTessData->m_puiSizesTriangulated[uiCurrentSize] & kA3DTessFaceDataNormalSingle) != 0;????
+						//	A3DUns32* pStripeNormalIndice = puiTriangulatedIndexes;
+						//	A3DUns32 uiNbPoint = pFaceTessData->m_puiSizesTriangulated[uiCurrentSize++] & kA3DTessFaceDataNormalMask;
+						A3DUns32 PointCount = FaceTessData.m_puiSizesTriangulated[FaceSetIndex++] & kA3DTessFaceDataNormalMask;
+						AddFaceTriangleStripeWithUniqueNormal(Tessellation, PointCount, LastTrianguleIndex, LastVertexIndex);
+					}
+					bMustProcess = FaceTessData.m_uiSizesTriangulatedSize > FaceSetIndex;
+				}
+
+				if (bMustProcess && UsedEntitiesFlags & kA3DTessFaceDataTriangleStripeTextured)
+				{
+					A3DUns32 StripeSize = FaceTessData.m_puiSizesTriangulated[FaceSetIndex++];
+					for (A3DUns32 StripeIndex = 0; StripeIndex < StripeSize; ++StripeIndex)
+					{
+						A3DUns32 PointCount = FaceTessData.m_puiSizesTriangulated[FaceSetIndex++];
+						AddFaceTriangleStripeWithTexture(Tessellation, PointCount, FaceTessData.m_uiTextureCoordIndexesSize, LastTrianguleIndex, LastVertexIndex);
+					}
+					bMustProcess = FaceTessData.m_uiSizesTriangulatedSize > FaceSetIndex;
+				}
+
+				if (bMustProcess && UsedEntitiesFlags & kA3DTessFaceDataTriangleStripeOneNormalTextured)
+				{
+					A3DUns32 StripeSize = FaceTessData.m_puiSizesTriangulated[FaceSetIndex++] & kA3DTessFaceDataNormalMask;
+					for (A3DUns32 StripeIndex = 0; StripeIndex < StripeSize; ++StripeIndex)
+					{
+						A3DUns32 PointCount = FaceTessData.m_puiSizesTriangulated[FaceSetIndex++] & kA3DTessFaceDataNormalMask;
+						AddFaceTriangleStripeWithUniqueNormalAndTexture(Tessellation, PointCount, FaceTessData.m_uiTextureCoordIndexesSize, LastTrianguleIndex, LastVertexIndex);
+					}
+					bMustProcess = FaceTessData.m_uiSizesTriangulatedSize > FaceSetIndex;
+				}
+
+				ensure(!bMustProcess);
+			}
+		}
+#pragma optimize ("", on)
+
+		typedef double A3DDouble;
+		bool AddFace(int32 FaceIndex[3], CADLibrary::FTessellationData& Tessellation, int32& InOutVertexIndex)
+		{
+			if (FaceIndex[0] == FaceIndex[1] || FaceIndex[0] == FaceIndex[2] || FaceIndex[1] == FaceIndex[2])
+			{
+				return false;
+			}
+
+			for (int32 Index = 0; Index < 3; ++Index)
+			{
+				Tessellation.VertexIndices.Add(InOutVertexIndex++);
+			}
+			Tessellation.PositionIndices.Append(FaceIndex, 3);
+			return true;
+		};
+
+		void AddNormals(const A3DDouble* Normals, const int32 Indices[3], TArray<FVector>& NormalsArray)
+		{
+			for (int32 Index = 0; Index < 3; ++Index)
+			{
+				int32 NormalIndex = Indices[Index];
+				NormalsArray.Emplace(Normals[NormalIndex], Normals[NormalIndex + 1], Normals[NormalIndex + 2]);
+			}
+		};
+
+		void AddTextureCoordinates(const A3DDouble* TextureCoords, const int32 Indices[3], TArray<FVector2D>& TessellationTextures)
+		{
+			for (int32 Index = 0; Index < 3; ++Index)
+			{
+				int32 TextureIndex = Indices[Index];
+				TessellationTextures.Emplace(TextureCoords[TextureIndex], TextureCoords[TextureIndex + 1]);
+			}
+		};
+
+		void AddFaceTriangle(CADLibrary::FTessellationData& Tessellation, const uint32 InTriangleCount, uint32& InOutStartIndex, int32& InOutLastVertexIndex)
+		{
+			int32 FaceIndex[3] = { 0, 0, 0 };
+			int32 NormalIndex[3] = { 0, 0, 0 };
+
+			// Get Triangles
+			for (uint32 TriangleIndex = 0; TriangleIndex < InTriangleCount; TriangleIndex++)
+			{
+				NormalIndex[0] = TriangulatedIndexes[InOutStartIndex++];
+				FaceIndex[0] = TriangulatedIndexes[InOutStartIndex++] / 3;
+				NormalIndex[1] = TriangulatedIndexes[InOutStartIndex++];
+				FaceIndex[1] = TriangulatedIndexes[InOutStartIndex++] / 3;
+				NormalIndex[2] = TriangulatedIndexes[InOutStartIndex++];
+				FaceIndex[2] = TriangulatedIndexes[InOutStartIndex++] / 3;
+
+				if (AddFace(FaceIndex, Tessellation, InOutLastVertexIndex))
+				{
+					AddNormals(TessellationNormals, NormalIndex, Tessellation.NormalArray);
+				}
+			}
+		}
+
+		void AddFaceTriangleWithUniqueNormal(CADLibrary::FTessellationData& Tessellation, const uint32 InTriangleCount, uint32& InOutLastTriangleIndex, int32& InOutLastVertexIndex)
+		{
+			int32 FaceIndex[3] = { 0, 0, 0 };
+			int32 NormalIndex[3] = { 0, 0, 0 };
+
+			// Get Triangles
+			for (uint32 TriangleIndex = 0; TriangleIndex < InTriangleCount; TriangleIndex++)
+			{
+				NormalIndex[0] = TriangulatedIndexes[InOutLastTriangleIndex++];
+				NormalIndex[1] = NormalIndex[0];
+				NormalIndex[2] = NormalIndex[0];
+
+				FaceIndex[0] = TriangulatedIndexes[InOutLastTriangleIndex++] / 3;
+				FaceIndex[1] = TriangulatedIndexes[InOutLastTriangleIndex++] / 3;
+				FaceIndex[2] = TriangulatedIndexes[InOutLastTriangleIndex++] / 3;
+
+				if (!AddFace(FaceIndex, Tessellation, InOutLastVertexIndex))
+				{
+					continue;
+				}
+
+				AddNormals(TessellationNormals, NormalIndex, Tessellation.NormalArray);
+			}
+		}
+
+		void AddFaceTriangleWithUniqueNormalAndTexture(CADLibrary::FTessellationData& Tessellation, const uint32 InTriangleCount, const uint32 TextureCount, uint32& InOutStartIndex, int32& InOutLastVertexIndex)
+		{
+			int32 FaceIndex[3] = { 0, 0, 0 };
+			int32 NormalIndex[3] = { 0, 0, 0 };
+			int32 TextureIndex[3] = { 0, 0, 0 };
+
+			// Get Triangles
+			for (uint32 TriangleIndex = 0; TriangleIndex < InTriangleCount; TriangleIndex++)
+			{
+				NormalIndex[0] = TriangulatedIndexes[InOutStartIndex++];
+				NormalIndex[1] = NormalIndex[0];
+				NormalIndex[2] = NormalIndex[0];
+
+				TextureIndex[0] = TriangulatedIndexes[InOutStartIndex];
+				InOutStartIndex += TextureCount;
+				FaceIndex[0] = TriangulatedIndexes[InOutStartIndex++] / 3;
+				TextureIndex[1] = TriangulatedIndexes[InOutStartIndex];
+				InOutStartIndex += TextureCount;
+				FaceIndex[1] = TriangulatedIndexes[InOutStartIndex++] / 3;
+				TextureIndex[2] = TriangulatedIndexes[InOutStartIndex];
+				InOutStartIndex += TextureCount;
+				FaceIndex[2] = TriangulatedIndexes[InOutStartIndex++] / 3;
+
+				if (!AddFace(FaceIndex, Tessellation, InOutLastVertexIndex))
+				{
+					continue;
+				}
+
+				AddNormals(TessellationNormals, NormalIndex, Tessellation.NormalArray);
+				AddTextureCoordinates(TessellationTexCoords, TextureIndex, Tessellation.TexCoordArray);
+			}
+		}
+
+		void AddFaceTriangleWithTexture(CADLibrary::FTessellationData& Tessellation, const uint32 InTriangleCount, const uint32 InTextureCount, uint32& InOutStartIndex, int32& inOutLastVertexIndex)
+		{
+			int32 FaceIndex[3] = { 0, 0, 0 };
+			int32 NormalIndex[3] = { 0, 0, 0 };
+			int32 TextureIndex[3] = { 0, 0, 0 };
+
+			// Get Triangles
+			for (uint64 TriangleIndex = 0; TriangleIndex < InTriangleCount; TriangleIndex++)
+			{
+				NormalIndex[0] = TriangulatedIndexes[InOutStartIndex++];
+				TextureIndex[0] = TriangulatedIndexes[InOutStartIndex];
+				InOutStartIndex += InTextureCount;
+				FaceIndex[0] = TriangulatedIndexes[InOutStartIndex++] / 3;
+				NormalIndex[1] = TriangulatedIndexes[InOutStartIndex++];
+				TextureIndex[1] = TriangulatedIndexes[InOutStartIndex];
+				InOutStartIndex += InTextureCount;
+				FaceIndex[1] = TriangulatedIndexes[InOutStartIndex++] / 3;
+				NormalIndex[2] = TriangulatedIndexes[InOutStartIndex++];
+				TextureIndex[2] = TriangulatedIndexes[InOutStartIndex];
+				InOutStartIndex += InTextureCount;
+				FaceIndex[2] = TriangulatedIndexes[InOutStartIndex++] / 3;
+
+				if (AddFace(FaceIndex, Tessellation, inOutLastVertexIndex))
+				{
+					AddNormals(TessellationNormals, NormalIndex, Tessellation.NormalArray);
+					AddTextureCoordinates(TessellationTexCoords, TextureIndex, Tessellation.TexCoordArray);
+				}
+			}
+		}
+
+		void AddFaceTriangleFan(CADLibrary::FTessellationData& Tessellation, const uint32 InTriangleCount, uint32& InOutLastTriangleIndex, int32& InOutLastVertexIndex)
+		{
+			int32 FaceIndex[3] = { 0, 0, 0 };
+			int32 NormalIndex[3] = { 0, 0, 0 };
+
+			NormalIndex[0] = TriangulatedIndexes[InOutLastTriangleIndex++];
+			FaceIndex[0] = TriangulatedIndexes[InOutLastTriangleIndex++] / 3;
+			NormalIndex[1] = TriangulatedIndexes[InOutLastTriangleIndex++];
+			FaceIndex[1] = TriangulatedIndexes[InOutLastTriangleIndex++] / 3;
+
+			for (unsigned long TriangleIndex = 2; TriangleIndex < InTriangleCount; TriangleIndex++)
+			{
+				NormalIndex[2] = TriangulatedIndexes[InOutLastTriangleIndex++];
+				FaceIndex[2] = TriangulatedIndexes[InOutLastTriangleIndex++] / 3;
+
+				if (AddFace(FaceIndex, Tessellation, InOutLastVertexIndex))
+				{
+					AddNormals(TessellationNormals, NormalIndex, Tessellation.NormalArray);
+				}
+
+				NormalIndex[1] = NormalIndex[2];
+				FaceIndex[1] = FaceIndex[2];
+			}
+		}
+
+		void AddFaceTriangleFanWithUniqueNormal(CADLibrary::FTessellationData& Tessellation, const uint32 InTriangleCount, uint32& InOutLastTriangleIndex, int32& /*LastVertexIndex*/)
+		{
+			int32 FaceIndex[3] = { 0, 0, 0 };
+			int32 NormalIndex[3] = { 0, 0, 0 };
+
+			NormalIndex[0] = TriangulatedIndexes[InOutLastTriangleIndex++];
+			NormalIndex[1] = NormalIndex[0];
+			NormalIndex[2] = NormalIndex[0];
+
+			FaceIndex[0] = TriangulatedIndexes[InOutLastTriangleIndex++] / 3;
+			FaceIndex[1] = TriangulatedIndexes[InOutLastTriangleIndex++] / 3;
+
+			// Get Triangles
+			for (uint32 TriangleIndex = 2; TriangleIndex < InTriangleCount; TriangleIndex++)
+			{
+				FaceIndex[2] = TriangulatedIndexes[InOutLastTriangleIndex++] / 3;
+
+				if (AddFace(FaceIndex, Tessellation, LastVertexIndex))
+				{
+					AddNormals(TessellationNormals, NormalIndex, Tessellation.NormalArray);
+				}
+
+				FaceIndex[1] = FaceIndex[2];
+			}
+		}
+
+		void AddFaceTriangleFanWithUniqueNormalAndTexture(CADLibrary::FTessellationData& Tessellation, const uint32 InTriangleCount, const uint32 TextureCount, uint32& InOutLastTriangleIndex, int32& InOutLastVertexIndex)
+		{
+			int32 FaceIndex[3] = { 0, 0, 0 };
+			int32 NormalIndex[3] = { 0, 0, 0 };
+			int32 TextureIndex[3] = { 0, 0, 0 };
+
+			NormalIndex[0] = TriangulatedIndexes[InOutLastTriangleIndex++];
+			NormalIndex[1] = NormalIndex[0];
+			NormalIndex[2] = NormalIndex[0];
+
+			TextureIndex[0] = TriangulatedIndexes[InOutLastTriangleIndex];
+			InOutLastTriangleIndex += TextureCount;
+			FaceIndex[0] = TriangulatedIndexes[InOutLastTriangleIndex++] / 3;
+
+			TextureIndex[1] = TriangulatedIndexes[InOutLastTriangleIndex];
+			InOutLastTriangleIndex += TextureCount;
+			FaceIndex[1] = TriangulatedIndexes[InOutLastTriangleIndex++] / 3;
+
+			for (uint32 TriangleIndex = 2; TriangleIndex < InTriangleCount; TriangleIndex++)
+			{
+				TextureIndex[2] = TriangulatedIndexes[InOutLastTriangleIndex];
+				InOutLastTriangleIndex += TextureCount;
+				FaceIndex[2] = TriangulatedIndexes[InOutLastTriangleIndex++] / 3;
+
+				if (AddFace(FaceIndex, Tessellation, InOutLastVertexIndex))
+				{
+					AddNormals(TessellationNormals, NormalIndex, Tessellation.NormalArray);
+					AddTextureCoordinates(TessellationTexCoords, TextureIndex, Tessellation.TexCoordArray);
+				}
+
+				FaceIndex[1] = FaceIndex[2];
+				TextureIndex[1] = TextureIndex[2];
+			}
+		}
+
+		void AddFaceTriangleFanWithTexture(CADLibrary::FTessellationData& Tessellation, const uint32 InTriangleCount, const uint32 TextureCount, uint32& InOutLastTriangleIndex, int32& InOutLastVertexIndex)
+		{
+			int32 FaceIndex[3] = { 0, 0, 0 };
+			int32 NormalIndex[3] = { 0, 0, 0 };
+			int32 TextureIndex[3] = { 0, 0, 0 };
+
+			NormalIndex[0] = TriangulatedIndexes[InOutLastTriangleIndex++];
+			TextureIndex[0] = TriangulatedIndexes[InOutLastTriangleIndex];
+			InOutLastTriangleIndex += TextureCount;
+			FaceIndex[0] = TriangulatedIndexes[InOutLastTriangleIndex++] / 3;
+
+			NormalIndex[1] = TriangulatedIndexes[InOutLastTriangleIndex++];
+			TextureIndex[1] = TriangulatedIndexes[InOutLastTriangleIndex];
+			InOutLastTriangleIndex += TextureCount;
+			FaceIndex[1] = TriangulatedIndexes[InOutLastTriangleIndex++] / 3;
+
+			for (uint32 TriangleIndex = 2; TriangleIndex < InTriangleCount; TriangleIndex++)
+			{
+				NormalIndex[2] = TriangulatedIndexes[InOutLastTriangleIndex++];
+				TextureIndex[2] = TriangulatedIndexes[InOutLastTriangleIndex];
+				InOutLastTriangleIndex += TextureCount;
+				FaceIndex[2] = TriangulatedIndexes[InOutLastTriangleIndex++] / 3;
+
+				if (AddFace(FaceIndex, Tessellation, InOutLastVertexIndex))
+				{
+					AddNormals(TessellationNormals, NormalIndex, Tessellation.NormalArray);
+					AddTextureCoordinates(TessellationTexCoords, TextureIndex, Tessellation.TexCoordArray);
+				}
+
+				NormalIndex[1] = NormalIndex[2];
+				TextureIndex[1] = TextureIndex[2];
+				FaceIndex[1] = FaceIndex[2];
+			}
+		}
+
+		void AddFaceTriangleStripe(CADLibrary::FTessellationData& Tessellation, const uint32 InTriangleCount, uint32& InOutLastTriangleIndex, int32& /*LastVertexIndex*/)
+		{
+			int32 FaceIndex[3] = { 0, 0, 0 };
+			int32 NormalIndex[3] = { 0, 0, 0 };
+
+			NormalIndex[0] = TriangulatedIndexes[InOutLastTriangleIndex++];
+			FaceIndex[0] = TriangulatedIndexes[InOutLastTriangleIndex++] / 3;
+			NormalIndex[1] = TriangulatedIndexes[InOutLastTriangleIndex++];
+			FaceIndex[1] = TriangulatedIndexes[InOutLastTriangleIndex++] / 3;
+
+			for (unsigned long TriangleIndex = 2; TriangleIndex < InTriangleCount; TriangleIndex++)
+			{
+				NormalIndex[2] = TriangulatedIndexes[InOutLastTriangleIndex++];
+				FaceIndex[2] = TriangulatedIndexes[InOutLastTriangleIndex++] / 3;
+
+				if (AddFace(FaceIndex, Tessellation, LastVertexIndex))
+				{
+					AddNormals(TessellationNormals, NormalIndex, Tessellation.NormalArray);
+				}
+
+				TriangleIndex++;
+				if (TriangleIndex == InTriangleCount)
+				{
+					break;
+				}
+
+				Swap(FaceIndex[1], FaceIndex[2]);
+				Swap(NormalIndex[1], NormalIndex[2]);
+
+				NormalIndex[0] = TriangulatedIndexes[InOutLastTriangleIndex++];
+				FaceIndex[0] = TriangulatedIndexes[InOutLastTriangleIndex++] / 3;
+
+				if (AddFace(FaceIndex, Tessellation, LastVertexIndex))
+				{
+					AddNormals(TessellationNormals, NormalIndex, Tessellation.NormalArray);
+				}
+
+				Swap(FaceIndex[0], FaceIndex[1]);
+				Swap(NormalIndex[0], NormalIndex[1]);
+			}
+		}
+
+		void AddFaceTriangleStripeWithTexture(CADLibrary::FTessellationData& Tessellation, const uint32 InTriangleCount, const uint32 TextureCount, uint32& InOutLastTriangleIndex, int32& InOutLastVertexIndex)
+		{
+			int32 FaceIndex[3] = { 0, 0, 0 };
+			int32 NormalIndex[3] = { 0, 0, 0 };
+			int32 TextureIndex[3] = { 0, 0, 0 };
+
+			NormalIndex[0] = TriangulatedIndexes[InOutLastTriangleIndex++];
+			TextureIndex[0] = TriangulatedIndexes[InOutLastTriangleIndex];
+			InOutLastTriangleIndex += TextureCount;
+			FaceIndex[0] = TriangulatedIndexes[InOutLastTriangleIndex++] / 3;
+			NormalIndex[1] = TriangulatedIndexes[InOutLastTriangleIndex++];
+			TextureIndex[1] = TriangulatedIndexes[InOutLastTriangleIndex];
+			InOutLastTriangleIndex += TextureCount;
+			FaceIndex[1] = TriangulatedIndexes[InOutLastTriangleIndex++] / 3;
+
+			for (unsigned long TriangleIndex = 0; TriangleIndex < InTriangleCount; TriangleIndex++)
+			{
+				NormalIndex[2] = TriangulatedIndexes[InOutLastTriangleIndex++];
+				TextureIndex[2] = TriangulatedIndexes[InOutLastTriangleIndex];
+				InOutLastTriangleIndex += TextureCount;
+				FaceIndex[2] = TriangulatedIndexes[InOutLastTriangleIndex++] / 3;
+
+				if (AddFace(FaceIndex, Tessellation, InOutLastVertexIndex))
+				{
+					AddNormals(TessellationNormals, NormalIndex, Tessellation.NormalArray);
+					AddTextureCoordinates(TessellationTexCoords, TextureIndex, Tessellation.TexCoordArray);
+				}
+
+				TriangleIndex++;
+				if (TriangleIndex == InTriangleCount)
+				{
+					break;
+				}
+
+				Swap(FaceIndex[1], FaceIndex[2]);
+				Swap(NormalIndex[1], NormalIndex[2]);
+				Swap(TextureIndex[1], TextureIndex[2]);
+
+				NormalIndex[0] = TriangulatedIndexes[InOutLastTriangleIndex++];
+				FaceIndex[0] = TriangulatedIndexes[InOutLastTriangleIndex++] / 3;
+
+				if (AddFace(FaceIndex, Tessellation, InOutLastVertexIndex))
+				{
+					AddNormals(TessellationNormals, NormalIndex, Tessellation.NormalArray);
+				}
+
+				Swap(FaceIndex[0], FaceIndex[1]);
+				Swap(NormalIndex[0], NormalIndex[1]);
+				Swap(TextureIndex[0], TextureIndex[1]);
+			}
+		}
+
+		void AddFaceTriangleStripeWithUniqueNormal(CADLibrary::FTessellationData& Tessellation, const uint32 InTriangleCount, uint32& InOutLastTriangleIndex, int32& InOutLastVertexIndex)
+		{
+			int32 FaceIndex[3] = { 0, 0, 0 };
+			int32 NormalIndex[3] = { 0, 0, 0 };
+
+			NormalIndex[0] = TriangulatedIndexes[InOutLastTriangleIndex++];
+			NormalIndex[1] = NormalIndex[0];
+			NormalIndex[2] = NormalIndex[0];
+
+			FaceIndex[0] = TriangulatedIndexes[InOutLastTriangleIndex++] / 3;
+			FaceIndex[1] = TriangulatedIndexes[InOutLastTriangleIndex++] / 3;
+
+			for (unsigned long TriangleIndex = 0; TriangleIndex < InTriangleCount; TriangleIndex++)
+			{
+				FaceIndex[2] = TriangulatedIndexes[InOutLastTriangleIndex++] / 3;
+
+				if (AddFace(FaceIndex, Tessellation, InOutLastVertexIndex))
+				{
+					AddNormals(TessellationNormals, NormalIndex, Tessellation.NormalArray);
+				}
+
+				TriangleIndex++;
+				if (TriangleIndex == InTriangleCount)
+				{
+					break;
+				}
+
+				Swap(FaceIndex[1], FaceIndex[2]);
+
+				NormalIndex[0] = TriangulatedIndexes[InOutLastTriangleIndex++];
+				FaceIndex[0] = TriangulatedIndexes[InOutLastTriangleIndex++] / 3;
+
+				if (AddFace(FaceIndex, Tessellation, InOutLastVertexIndex))
+				{
+					AddNormals(TessellationNormals, NormalIndex, Tessellation.NormalArray);
+				}
+
+				Swap(FaceIndex[0], FaceIndex[1]);
+			}
+		}
+
+		void AddFaceTriangleStripeWithUniqueNormalAndTexture(CADLibrary::FTessellationData& Tessellation, const uint32 InTriangleCount, const uint32 TextureCount, uint32& InOutLastTriangleIndex, int32& InOutLastVertexIndex)
+		{
+			int32 FaceIndex[3] = { 0, 0, 0 };
+			int32 NormalIndex[3] = { 0, 0, 0 };
+			int32 TextureIndex[3] = { 0, 0, 0 };
+
+			NormalIndex[0] = TriangulatedIndexes[InOutLastTriangleIndex++];
+			NormalIndex[1] = NormalIndex[0];
+			NormalIndex[2] = NormalIndex[0];
+
+			TextureIndex[0] = TriangulatedIndexes[InOutLastTriangleIndex];
+			InOutLastTriangleIndex += TextureCount;
+			FaceIndex[0] = TriangulatedIndexes[InOutLastTriangleIndex++] / 3;
+
+			TextureIndex[1] = TriangulatedIndexes[InOutLastTriangleIndex];
+			InOutLastTriangleIndex += TextureCount;
+			FaceIndex[1] = TriangulatedIndexes[InOutLastTriangleIndex++] / 3;
+
+			for (unsigned long TriangleIndex = 0; TriangleIndex < InTriangleCount; TriangleIndex++)
+			{
+				TextureIndex[2] = TriangulatedIndexes[InOutLastTriangleIndex];
+				InOutLastTriangleIndex += TextureCount;
+				FaceIndex[2] = TriangulatedIndexes[InOutLastTriangleIndex++] / 3;
+
+				if (AddFace(FaceIndex, Tessellation, InOutLastVertexIndex))
+				{
+					AddNormals(TessellationNormals, NormalIndex, Tessellation.NormalArray);
+					AddTextureCoordinates(TessellationTexCoords, TextureIndex, Tessellation.TexCoordArray);
+				}
+
+				TriangleIndex++;
+				if (TriangleIndex == InTriangleCount)
+				{
+					break;
+				}
+
+				Swap(FaceIndex[1], FaceIndex[2]);
+				Swap(TextureIndex[1], TextureIndex[2]);
+
+				FaceIndex[0] = TriangulatedIndexes[InOutLastTriangleIndex++] / 3;
+
+				if (AddFace(FaceIndex, Tessellation, InOutLastVertexIndex))
+				{
+					AddNormals(TessellationNormals, NormalIndex, Tessellation.NormalArray);
+				}
+
+				Swap(FaceIndex[0], FaceIndex[1]);
+				Swap(TextureIndex[0], TextureIndex[1]);
+			}
+		}
+
+	private:
+		const A3DTess3D* TessellationPtr;
+
+		uint32 LastTrianguleIndex = 0;
+		int32 LastVertexIndex = 0;
+		A3DUns32* TriangulatedIndexes = nullptr;
+		A3DDouble* TessellationNormals;
+		A3DDouble* TessellationTexCoords;				/*!< Array of \ref A3DDouble, as texture coordinates. */
+	};
+} // ns TechSoftUtils
+
+#endif
\ No newline at end of file
diff --git a/Engine/Plugins/Enterprise/DatasmithCADImporter/Source/CADInterfaces/Public/CADFileData.h b/Engine/Plugins/Enterprise/DatasmithCADImporter/Source/CADInterfaces/Public/CADFileData.h
--- a/Engine/Plugins/Enterprise/DatasmithCADImporter/Source/CADInterfaces/Public/CADFileData.h
+++ b/Engine/Plugins/Enterprise/DatasmithCADImporter/Source/CADInterfaces/Public/CADFileData.h
@@ -69,11 +69,7 @@ public:
 
 	const FString GetBodyCachePath(uint32 BodyHash) const
 	{
-		if (IsCacheDefined())
-		{
-			return CADLibrary::BuildCacheFilePath(*CachePath, TEXT("body"), BodyHash);
-		}
-		return FString();
+		return CADLibrary::BuildCacheFilePath(*CachePath, TEXT("body"), BodyHash);
 	}
 
 	/**
diff --git a/Engine/Plugins/Enterprise/DatasmithCADImporter/Source/CADInterfaces/Public/TechSoftInterface.h b/Engine/Plugins/Enterprise/DatasmithCADImporter/Source/CADInterfaces/Public/TechSoftInterface.h
--- a/Engine/Plugins/Enterprise/DatasmithCADImporter/Source/CADInterfaces/Public/TechSoftInterface.h
+++ b/Engine/Plugins/Enterprise/DatasmithCADImporter/Source/CADInterfaces/Public/TechSoftInterface.h
@@ -43,6 +43,11 @@ public:
 
 	bool InitializeKernel(const TCHAR* = TEXT(""));
 
+	void SaveBodyToHsfFile(void* BodyPtr, const FString& Filename);
+
+	bool GetBodyFromHsfFile(const FString& Filename, const class FImportParameters& ImportParameters, double FileUnit, class FBodyMesh& BodyMesh);
+
+	bool FillBodyMesh(void* BodyPtr, const class FImportParameters& ImportParameters, double FileUnit, class FBodyMesh& BodyMesh);
 
 #ifdef USE_TECHSOFT_SDK
 	A3DStatus Import(const A3DImport& Import);
diff --git a/Engine/Plugins/Enterprise/DatasmithCADImporter/Source/CADTools/Public/CADOptions.h b/Engine/Plugins/Enterprise/DatasmithCADImporter/Source/CADTools/Public/CADOptions.h
--- a/Engine/Plugins/Enterprise/DatasmithCADImporter/Source/CADTools/Public/CADOptions.h
+++ b/Engine/Plugins/Enterprise/DatasmithCADImporter/Source/CADTools/Public/CADOptions.h
@@ -212,7 +212,7 @@ namespace CADLibrary
 
 		if (FImportParameters::bGDisableCADKernelTessellation)
 		{
-			OutFileName += TEXT(".ct");
+			OutFileName += GCADLibrary.Equals("TechSoft") ? TEXT(".hsf") : TEXT(".ct");
 		}
 		else
 		{
diff --git a/Engine/Plugins/Enterprise/DatasmithCADImporter/Source/DatasmithCADTranslator/Private/DatasmithCADTranslator.cpp b/Engine/Plugins/Enterprise/DatasmithCADImporter/Source/DatasmithCADTranslator/Private/DatasmithCADTranslator.cpp
--- a/Engine/Plugins/Enterprise/DatasmithCADImporter/Source/DatasmithCADTranslator/Private/DatasmithCADTranslator.cpp
+++ b/Engine/Plugins/Enterprise/DatasmithCADImporter/Source/DatasmithCADTranslator/Private/DatasmithCADTranslator.cpp
@@ -172,14 +172,14 @@ bool FDatasmithCADTranslator::LoadScene(TSharedRef<IDatasmithScene> DatasmithSce
 		return true;
 	}
 
-	CADLibrary::FCADFileReader FileReader(ImportParameters, FileDescriptor, *FPaths::EnginePluginsDir());
+	CADLibrary::FCADFileReader FileReader(ImportParameters, FileDescriptor, *FPaths::EnginePluginsDir(), CachePath);
 	if (FileReader.ProcessFile() != CADLibrary::ECADParsingResult::ProcessOk)
 	{
 		return false;
 	}
 
 	CADLibrary::FCADFileData& CADFileData = FileReader.GetCADFileData();
-	FDatasmithSceneBaseGraphBuilder SceneGraphBuilder(&CADFileData.GetSceneGraphArchive(), DatasmithScene, GetSource(), ImportParameters);
+	FDatasmithSceneBaseGraphBuilder SceneGraphBuilder(&CADFileData.GetSceneGraphArchive(), CachePath, DatasmithScene, GetSource(), ImportParameters);
 	SceneGraphBuilder.Build();
 
 	MeshBuilderPtr = MakeUnique<FDatasmithMeshBuilder>(CADFileData.GetBodyMeshes(), ImportParameters);
diff --git a/Engine/Plugins/Enterprise/DatasmithCADImporter/Source/DatasmithCADTranslator/Private/DatasmithCADTranslatorModule.cpp b/Engine/Plugins/Enterprise/DatasmithCADImporter/Source/DatasmithCADTranslator/Private/DatasmithCADTranslatorModule.cpp
--- a/Engine/Plugins/Enterprise/DatasmithCADImporter/Source/DatasmithCADTranslator/Private/DatasmithCADTranslatorModule.cpp
+++ b/Engine/Plugins/Enterprise/DatasmithCADImporter/Source/DatasmithCADTranslator/Private/DatasmithCADTranslatorModule.cpp
@@ -33,6 +33,12 @@ void FDatasmithCADTranslatorModule::StartupModule()
 		CADLibrary::FImportParameters::bGEnableCADCache = false; // very weak protection: user could turn that on later, while the cache path is invalid
 	}
 
+	// Create body cache directory since this one is used even if bGEnableCADCache is false
+	if (!CacheDir.IsEmpty())
+	{
+		IFileManager::Get().MakeDirectory(*FPaths::Combine(CacheDir, TEXT("body")), true);
+	}
+
 	Datasmith::RegisterTranslator<FDatasmithCADTranslator>();
 }
 
diff --git a/Engine/Plugins/Enterprise/DatasmithCADImporter/Source/DatasmithCADTranslator/Private/DatasmithSceneGraphBuilder.cpp b/Engine/Plugins/Enterprise/DatasmithCADImporter/Source/DatasmithCADTranslator/Private/DatasmithSceneGraphBuilder.cpp
--- a/Engine/Plugins/Enterprise/DatasmithCADImporter/Source/DatasmithCADTranslator/Private/DatasmithSceneGraphBuilder.cpp
+++ b/Engine/Plugins/Enterprise/DatasmithCADImporter/Source/DatasmithCADTranslator/Private/DatasmithSceneGraphBuilder.cpp
@@ -77,9 +77,8 @@ FDatasmithSceneGraphBuilder::FDatasmithSceneGraphBuilder(
 	TSharedRef<IDatasmithScene> InScene, 
 	const FDatasmithSceneSource& InSource, 
 	const CADLibrary::FImportParameters& InImportParameters)
-		: FDatasmithSceneBaseGraphBuilder(nullptr, InScene, InSource, InImportParameters)
+		: FDatasmithSceneBaseGraphBuilder(nullptr, InCachePath, InScene, InSource, InImportParameters)
 		, CADFileToSceneGraphDescriptionFile(InCADFileToUnrealFileMap)
-		, CachePath(InCachePath)
 {
 }
 
@@ -163,20 +162,9 @@ void FDatasmithSceneGraphBuilder::FillAnchorActor(const TSharedRef<IDatasmithAct
 	ActorElement->SetLabel(*ActorLabel);
 }
 
-TSharedPtr<IDatasmithMeshElement> FDatasmithSceneGraphBuilder::FindOrAddMeshElement(CADLibrary::FArchiveBody & Body, FString & InLabel)
-{
-	if (TSharedPtr<IDatasmithMeshElement> MeshElement = FDatasmithSceneBaseGraphBuilder::FindOrAddMeshElement(Body, InLabel))
-	{
-		FString BodyCachePath = CADLibrary::BuildCacheFilePath(*CachePath, TEXT("body"), Body.MeshActorName);
-		MeshElement->SetFile(*BodyCachePath);
-		return MeshElement;
-	}
-
-	return TSharedPtr<IDatasmithMeshElement>();
-}
-
-FDatasmithSceneBaseGraphBuilder::FDatasmithSceneBaseGraphBuilder(CADLibrary::FArchiveSceneGraph* InSceneGraph, TSharedRef<IDatasmithScene> InScene, const FDatasmithSceneSource& InSource, const CADLibrary::FImportParameters& InImportParameters)
+FDatasmithSceneBaseGraphBuilder::FDatasmithSceneBaseGraphBuilder(CADLibrary::FArchiveSceneGraph* InSceneGraph, const FString& InCachePath,  TSharedRef<IDatasmithScene> InScene, const FDatasmithSceneSource& InSource, const CADLibrary::FImportParameters& InImportParameters)
 	: SceneGraph(InSceneGraph)
+	, CachePath(InCachePath)
 	, DatasmithScene(InScene)
 	, ImportParameters(InImportParameters)
 	, ImportParametersHash(ImportParameters.GetHash())
@@ -577,6 +565,9 @@ TSharedPtr< IDatasmithMeshElement > FDatasmithSceneBaseGraphBuilder::FindOrAddMe
 
 	BodyUuidToMeshElement.Add(Body.MeshActorName, MeshElement);
 
+	FString BodyCachePath = CADLibrary::BuildCacheFilePath(*CachePath, TEXT("body"), Body.MeshActorName);
+	MeshElement->SetFile(*BodyCachePath);
+
 	return MeshElement;
 }
 
diff --git a/Engine/Plugins/Enterprise/DatasmithCADImporter/Source/DatasmithCADTranslator/Public/DatasmithSceneGraphBuilder.h b/Engine/Plugins/Enterprise/DatasmithCADImporter/Source/DatasmithCADTranslator/Public/DatasmithSceneGraphBuilder.h
--- a/Engine/Plugins/Enterprise/DatasmithCADImporter/Source/DatasmithCADTranslator/Public/DatasmithSceneGraphBuilder.h
+++ b/Engine/Plugins/Enterprise/DatasmithCADImporter/Source/DatasmithCADTranslator/Public/DatasmithSceneGraphBuilder.h
@@ -56,6 +56,7 @@ class DATASMITHCADTRANSLATOR_API FDatasmithSceneBaseGraphBuilder
 public:
 	FDatasmithSceneBaseGraphBuilder(
 		CADLibrary::FArchiveSceneGraph* InSceneGraph,
+		const FString& InCachePath,
 		TSharedRef<IDatasmithScene> InScene, 
 		const FDatasmithSceneSource& InSource, 
 		const CADLibrary::FImportParameters& InImportParameters);
@@ -83,6 +84,7 @@ protected:
 
 protected:
 	CADLibrary::FArchiveSceneGraph* SceneGraph;
+	const FString& CachePath;
 	TSharedRef<IDatasmithScene> DatasmithScene;
 	const CADLibrary::FImportParameters& ImportParameters;
 	const uint32 ImportParametersHash;
@@ -124,10 +126,6 @@ public:
 
 	void FillAnchorActor(const TSharedRef< IDatasmithActorElement >& ActorElement, const FString& CleanFilenameOfCADFile);
 
-protected:
-	virtual TSharedPtr< IDatasmithMeshElement > FindOrAddMeshElement(CADLibrary::FArchiveBody& Body, FString& InLabel);
-
 protected:
 	TMap<uint32, FString>& CADFileToSceneGraphDescriptionFile;
-	const FString& CachePath;
 };
