Techsoft: fix the transparency of the materials

diff --git a/Engine/Plugins/Enterprise/DatasmithCADImporter/Source/CADInterfaces/Private/TechSoftFileParser.cpp b/Engine/Plugins/Enterprise/DatasmithCADImporter/Source/CADInterfaces/Private/TechSoftFileParser.cpp
--- a/Engine/Plugins/Enterprise/DatasmithCADImporter/Source/CADInterfaces/Private/TechSoftFileParser.cpp
+++ b/Engine/Plugins/Enterprise/DatasmithCADImporter/Source/CADInterfaces/Private/TechSoftFileParser.cpp
@@ -1448,7 +1448,7 @@ FArchiveColor& FTechSoftFileParser::FindOrAddColor(uint32 ColorIndex, uint8 Alph
 	return NewColor;
 }
 
-FArchiveMaterial& FTechSoftFileParser::AddMaterialAt(uint32 MaterialIndexToSave, uint32 GraphMaterialIndex)
+FArchiveMaterial& FTechSoftFileParser::AddMaterialAt(uint32 MaterialIndexToSave, uint32 GraphMaterialIndex, const A3DGraphStyleData& GraphStyleData)
 {
 	FArchiveMaterial& NewMaterial = CADFileData.AddMaterial(GraphMaterialIndex);
 	FCADMaterial& Material = NewMaterial.Material;
@@ -1460,7 +1460,10 @@ FArchiveMaterial& FTechSoftFileParser::AddMaterialAt(uint32 MaterialIndexToSave,
 		Material.Ambient = TechSoftFileParserImpl::GetColorAt(MaterialData->m_uiAmbient);
 		Material.Specular = TechSoftFileParserImpl::GetColorAt(MaterialData->m_uiSpecular);
 		Material.Shininess = MaterialData->m_dShininess;
-		Material.Transparency = 1. - MaterialData->m_dDiffuseAlpha;
+		if(GraphStyleData.m_bIsTransparencyDefined)
+		{
+			Material.Transparency = 1. - GraphStyleData.m_ucTransparency/255.;
+		}
 		// todo: find how to convert Emissive color into ? reflexion coef...
 		// Material.Emissive = GetColor(MaterialData->m_uiEmissive);
 		// Material.Reflexion;
@@ -1470,7 +1473,7 @@ FArchiveMaterial& FTechSoftFileParser::AddMaterialAt(uint32 MaterialIndexToSave,
 }
 
 
-FArchiveMaterial& FTechSoftFileParser::FindOrAddMaterial(uint32 MaterialIndex)
+FArchiveMaterial& FTechSoftFileParser::FindOrAddMaterial(uint32 MaterialIndex, const A3DGraphStyleData& GraphStyleData)
 {
 	if (FArchiveMaterial* MaterialArchive = CADFileData.FindMaterial(MaterialIndex))
 	{
@@ -1484,7 +1487,7 @@ FArchiveMaterial& FTechSoftFileParser::FindOrAddMaterial(uint32 MaterialIndex)
 		if (TextureData.IsValid())
 		{
 			TextureData->m_uiMaterialIndex;
-			return AddMaterialAt(TextureData->m_uiMaterialIndex, MaterialIndex);
+			return AddMaterialAt(TextureData->m_uiMaterialIndex, MaterialIndex, GraphStyleData);
 			
 #ifdef NOTYETDEFINE
 			TUniqueTSObj<A3DGraphTextureDefinitionData> TextureDefinitionData(TextureData->m_uiTextureDefinitionIndex);
@@ -1494,11 +1497,11 @@ FArchiveMaterial& FTechSoftFileParser::FindOrAddMaterial(uint32 MaterialIndex)
 			}
 #endif
 		}
-		return AddMaterialAt(MaterialIndex, 0);
+		return AddMaterialAt(MaterialIndex, 0, GraphStyleData);
 	}
 	else
 	{
-		return AddMaterial(MaterialIndex);
+		return AddMaterial(MaterialIndex, GraphStyleData);
 	}
 }
 
@@ -1560,7 +1563,7 @@ void FTechSoftFileParser::ExtractGraphStyleProperties(uint32 StyleIndex, FCADUUI
 	{
 		if (GraphStyleData->m_bMaterial)
 		{
-			FArchiveMaterial& MaterialArchive = FindOrAddMaterial(GraphStyleData->m_uiRgbColorIndex);
+			FArchiveMaterial& MaterialArchive = FindOrAddMaterial(GraphStyleData->m_uiRgbColorIndex, *GraphStyleData);
 			OutMaterialName = MaterialArchive.UEMaterialName;
 		}
 		else
@@ -1790,17 +1793,6 @@ void FTechSoftFileParser::ReadMaterialsAndColors()
 		return;
 	}
 
-	{
-		uint32 MaterialCount = GlobalData->m_uiMaterialsSize;
-		if (MaterialCount)
-		{
-			for (uint32 MaterialIndex = 0; MaterialIndex < MaterialCount; ++MaterialIndex)
-			{
-				FTechSoftFileParser::FindOrAddMaterial(MaterialIndex);
-			}
-		}
-	}
-
 	{
 		uint32 TextureDefinitionCount = GlobalData->m_uiTextureDefinitionsSize;
 		if (TextureDefinitionCount)
diff --git a/Engine/Plugins/Enterprise/DatasmithCADImporter/Source/CADInterfaces/Private/TechSoftInterfaceUtils.inl b/Engine/Plugins/Enterprise/DatasmithCADImporter/Source/CADInterfaces/Private/TechSoftInterfaceUtils.inl
--- a/Engine/Plugins/Enterprise/DatasmithCADImporter/Source/CADInterfaces/Private/TechSoftInterfaceUtils.inl
+++ b/Engine/Plugins/Enterprise/DatasmithCADImporter/Source/CADInterfaces/Private/TechSoftInterfaceUtils.inl
@@ -103,7 +103,6 @@ namespace TechSoftInterfaceUtils
 			return TriangleCount;
 		}
 
-#pragma optimize ("", off)
 		void FillFaceArray(CADLibrary::FBodyMesh& BodyMesh)
 		{
 			using namespace CADLibrary;
@@ -275,7 +274,6 @@ namespace TechSoftInterfaceUtils
 				ensure(!bMustProcess);
 			}
 		}
-#pragma optimize ("", on)
 
 		typedef double A3DDouble;
 		bool AddFace(int32 FaceIndex[3], CADLibrary::FTessellationData& Tessellation, int32& InOutVertexIndex)
