Fix TechsoftFileParser

diff --git a/Engine/Plugins/Enterprise/DatasmithCADImporter/Source/CADInterfaces/Private/TechSoftFileParser.h b/Engine/Plugins/Enterprise/DatasmithCADImporter/Source/CADInterfaces/Private/TechSoftFileParser.h
--- a/Engine/Plugins/Enterprise/DatasmithCADImporter/Source/CADInterfaces/Private/TechSoftFileParser.h
+++ b/Engine/Plugins/Enterprise/DatasmithCADImporter/Source/CADInterfaces/Private/TechSoftFileParser.h
@@ -187,16 +187,16 @@ private:
 	void ExtractGraphStyleProperties(uint32 StyleIndex, FCADUUID& ColorName, FCADUUID& MaterialName);
 	void ExtractMaterialProperties(const A3DEntity* Entity);
 	FArchiveColor& FindOrAddColor(uint32 ColorIndex, uint8 Alpha);
-	FArchiveMaterial& FindOrAddMaterial(uint32 MaterialId);
+	FArchiveMaterial& FindOrAddMaterial(uint32 MaterialId, const A3DGraphStyleData& GraphStyleData);
 
 	/**
 	 * @param GraphMaterialIndex is the techsoft index of the graphic data
 	 * @param MaterialIndexToSave is the index of the material really saved (i.e. for texture, at the texture index, with saved the material used by the texture)
 	 */
-	FArchiveMaterial& AddMaterialAt(uint32 MaterialIndexToSave, uint32 GraphMaterialIndex);
-	FArchiveMaterial& AddMaterial(uint32 MaterialIndex)
+	FArchiveMaterial& AddMaterialAt(uint32 MaterialIndexToSave, uint32 GraphMaterialIndex, const A3DGraphStyleData& GraphStyleData);
+	FArchiveMaterial& AddMaterial(uint32 MaterialIndex, const A3DGraphStyleData& GraphStyleData)
 	{
-		return AddMaterialAt(MaterialIndex, MaterialIndex);
+		return AddMaterialAt(MaterialIndex, MaterialIndex, GraphStyleData);
 	}
 
 	// Transform
