Fix change SDK with CVar

diff --git a/Engine/Plugins/Enterprise/DatasmithCADImporter/Source/CADInterfaces/Private/CADFileReader.cpp b/Engine/Plugins/Enterprise/DatasmithCADImporter/Source/CADInterfaces/Private/CADFileReader.cpp
--- a/Engine/Plugins/Enterprise/DatasmithCADImporter/Source/CADInterfaces/Private/CADFileReader.cpp
+++ b/Engine/Plugins/Enterprise/DatasmithCADImporter/Source/CADInterfaces/Private/CADFileReader.cpp
@@ -21,7 +21,7 @@ namespace CADLibrary
 		: CADFileData(ImportParams, InFile, InCachePath)
 	{
 #ifdef USE_KERNEL_IO_SDK
-		if (GCADLibrary == TEXT("KernelIO"))
+		if (FImportParameters::GCADLibrary == TEXT("KernelIO"))
 		{
 			CADParser = MakeUnique<FCoreTechFileParser>(CADFileData, EnginePluginsPath);
 		}
@@ -30,7 +30,7 @@ namespace CADLibrary
 		else
 #endif
 #ifdef USE_TECHSOFT_SDK
-		if (GCADLibrary == TEXT("TechSoft"))
+		if (FImportParameters::GCADLibrary == TEXT("TechSoft"))
 		{
 			CADParser = MakeUnique<FTechSoftFileParser>(CADFileData, EnginePluginsPath);
 		}
diff --git a/Engine/Plugins/Enterprise/DatasmithCADImporter/Source/CADInterfaces/Private/CADInterfacesModule.cpp b/Engine/Plugins/Enterprise/DatasmithCADImporter/Source/CADInterfaces/Private/CADInterfacesModule.cpp
--- a/Engine/Plugins/Enterprise/DatasmithCADImporter/Source/CADInterfaces/Private/CADInterfacesModule.cpp
+++ b/Engine/Plugins/Enterprise/DatasmithCADImporter/Source/CADInterfaces/Private/CADInterfacesModule.cpp
@@ -36,14 +36,14 @@ ECADInterfaceAvailability ICADInterfacesModule::GetAvailability()
 {
 	if (FModuleManager::Get().IsModuleLoaded(CADINTERFACES_MODULE_NAME))
 	{
-		if (CADLibrary::GCADLibrary == TEXT("KernelIO"))
+		if (CADLibrary::FImportParameters::GCADLibrary == TEXT("KernelIO"))
 		{
 			if (CADLibrary::CTKIO_InitializeKernel())
 			{
 				return ECADInterfaceAvailability::Available;
 			}
 		}
-		else if (CADLibrary::GCADLibrary == TEXT("TechSoft"))
+		else if (CADLibrary::FImportParameters::GCADLibrary == TEXT("TechSoft"))
 		{
 			if (CADLibrary::TechSoftUtils::TECHSOFT_InitializeKernel())
 			{
diff --git a/Engine/Plugins/Enterprise/DatasmithCADImporter/Source/CADInterfaces/Private/CoreTechFileParserImpl.cpp b/Engine/Plugins/Enterprise/DatasmithCADImporter/Source/CADInterfaces/Private/CoreTechFileParserImpl.cpp
--- a/Engine/Plugins/Enterprise/DatasmithCADImporter/Source/CADInterfaces/Private/CoreTechFileParserImpl.cpp
+++ b/Engine/Plugins/Enterprise/DatasmithCADImporter/Source/CADInterfaces/Private/CoreTechFileParserImpl.cpp
@@ -190,6 +190,7 @@ namespace CADLibrary
 		, FileDescription(InCADData.GetCADFileDescription())
 		, LastHostIdUsed(1<<30)
 	{
+		CTKIO_InitializeKernel(*EnginePluginsPath);
 	}
 
 	ECADParsingResult FCoreTechFileParser::Process()
diff --git a/Engine/Plugins/Enterprise/DatasmithCADImporter/Source/CADInterfaces/Private/TechSoftFileParser.cpp b/Engine/Plugins/Enterprise/DatasmithCADImporter/Source/CADInterfaces/Private/TechSoftFileParser.cpp
--- a/Engine/Plugins/Enterprise/DatasmithCADImporter/Source/CADInterfaces/Private/TechSoftFileParser.cpp
+++ b/Engine/Plugins/Enterprise/DatasmithCADImporter/Source/CADInterfaces/Private/TechSoftFileParser.cpp
@@ -278,6 +278,7 @@ FTechSoftFileParser::FTechSoftFileParser(FCADFileData& InCADData, const FString&
 	: CADFileData(InCADData)
 	, TechSoftInterface(TechSoftUtils::GetTechSoftInterface())
 {
+	TechSoftInterface.InitializeKernel(*EnginePluginsPath);
 }
 
 ECADParsingResult FTechSoftFileParser::Process()
diff --git a/Engine/Plugins/Enterprise/DatasmithCADImporter/Source/CADInterfaces/Public/CADFileData.h b/Engine/Plugins/Enterprise/DatasmithCADImporter/Source/CADInterfaces/Public/CADFileData.h
--- a/Engine/Plugins/Enterprise/DatasmithCADImporter/Source/CADInterfaces/Public/CADFileData.h
+++ b/Engine/Plugins/Enterprise/DatasmithCADImporter/Source/CADInterfaces/Public/CADFileData.h
@@ -28,7 +28,7 @@ public:
 		if (!SceneFileHash)
 		{
 			SceneFileHash = HashCombine(FileDescription.GetDescriptorHash(), ::GetTypeHash(ImportParameters.GetStitchingTechnique()));
-			SceneFileHash = HashCombine(SceneFileHash, ::GetTypeHash(CADLibrary::GCADLibrary));
+			SceneFileHash = HashCombine(SceneFileHash, ::GetTypeHash(CADLibrary::FImportParameters::GCADLibrary));
 		}
 		return SceneFileHash;
 	}
diff --git a/Engine/Plugins/Enterprise/DatasmithCADImporter/Source/CADTools/Private/CADOption.cpp b/Engine/Plugins/Enterprise/DatasmithCADImporter/Source/CADTools/Private/CADOption.cpp
--- a/Engine/Plugins/Enterprise/DatasmithCADImporter/Source/CADTools/Private/CADOption.cpp
+++ b/Engine/Plugins/Enterprise/DatasmithCADImporter/Source/CADTools/Private/CADOption.cpp
@@ -4,10 +4,10 @@
 namespace CADLibrary
 {
 
-FString GCADLibrary = TEXT("KernelIO");
+FString FImportParameters::GCADLibrary = TEXT("KernelIO");
 FAutoConsoleVariableRef GCADTranslatorLibrary(
 	TEXT("ds.CADTranslator.CADLibrary"),
-	GCADLibrary,
+	FImportParameters::GCADLibrary,
 	TEXT("Define the used library to import CAD file i.e. \"TechSoft\" or \"KernelIO\".\n"),
 	ECVF_Default);
 
diff --git a/Engine/Plugins/Enterprise/DatasmithCADImporter/Source/CADTools/Public/CADOptions.h b/Engine/Plugins/Enterprise/DatasmithCADImporter/Source/CADTools/Public/CADOptions.h
--- a/Engine/Plugins/Enterprise/DatasmithCADImporter/Source/CADTools/Public/CADOptions.h
+++ b/Engine/Plugins/Enterprise/DatasmithCADImporter/Source/CADTools/Public/CADOptions.h
@@ -11,7 +11,6 @@
 namespace CADLibrary
 {
 	CADTOOLS_API extern int32 GMaxImportThreads;
-	CADTOOLS_API extern FString GCADLibrary;
 
 	enum EStitchingTechnique
 	{
@@ -55,6 +54,8 @@ namespace CADLibrary
 		CADTOOLS_API static bool bGEnableCADCache;
 		CADTOOLS_API static bool bGOverwriteCache;
 		CADTOOLS_API static bool bGPreferJtFileEmbeddedTessellation;
+		CADTOOLS_API static FString GCADLibrary;
+
 
 	public:
 		FImportParameters(double InMetricUnit = 0.001, double InScaleFactor = 0.1, FDatasmithUtils::EModelCoordSystem NewCoordinateSystem = FDatasmithUtils::EModelCoordSystem::ZUp_RightHanded)
@@ -128,6 +129,7 @@ namespace CADLibrary
 			Ar << ImportParameters.bGEnableTimeControl;
 			Ar << ImportParameters.bGEnableCADCache;
 			Ar << ImportParameters.bGPreferJtFileEmbeddedTessellation;
+			Ar << ImportParameters.GCADLibrary;
 			return Ar;
 		}
 
@@ -212,7 +214,7 @@ namespace CADLibrary
 
 		if (FImportParameters::bGDisableCADKernelTessellation)
 		{
-			OutFileName += GCADLibrary.Equals("TechSoft") ? TEXT(".hsf") : TEXT(".ct");
+			OutFileName += FImportParameters::GCADLibrary.Equals("TechSoft") ? TEXT(".hsf") : TEXT(".ct");
 		}
 		else
 		{
diff --git a/Engine/Plugins/Enterprise/DatasmithCADImporter/Source/DatasmithCADTranslator/Private/DatasmithCADTranslator.cpp b/Engine/Plugins/Enterprise/DatasmithCADImporter/Source/DatasmithCADTranslator/Private/DatasmithCADTranslator.cpp
--- a/Engine/Plugins/Enterprise/DatasmithCADImporter/Source/DatasmithCADTranslator/Private/DatasmithCADTranslator.cpp
+++ b/Engine/Plugins/Enterprise/DatasmithCADImporter/Source/DatasmithCADTranslator/Private/DatasmithCADTranslator.cpp
@@ -74,7 +74,7 @@ void FDatasmithCADTranslator::Initialize(FDatasmithTranslatorCapabilities& OutCa
 	OutCapabilities.SupportedFileFormats.Add(FFileFormatInfo{ TEXT("dwg"), TEXT("AutoCAD, Model files") });
 	OutCapabilities.SupportedFileFormats.Add(FFileFormatInfo{ TEXT("dgn"), TEXT("MicroStation files") });
 
-	if (CADLibrary::GCADLibrary.Equals(TEXT("TechSoft")))
+	if (CADLibrary::FImportParameters::GCADLibrary.Equals(TEXT("TechSoft")))
 	{
 		OutCapabilities.SupportedFileFormats.Add(FFileFormatInfo{ TEXT("ifc"), TEXT("IFC files") });
 		OutCapabilities.SupportedFileFormats.Add(FFileFormatInfo{ TEXT("hsf"), TEXT("HOOPS stream files") });
@@ -101,9 +101,9 @@ bool FDatasmithCADTranslator::LoadScene(TSharedRef<IDatasmithScene> DatasmithSce
 	CADLibrary::FFileDescriptor FileDescriptor(*FPaths::ConvertRelativePathToFull(GetSource().GetSourceFile()));
 
 	UE_LOG(LogCADTranslator, Display, TEXT("CAD translation [%s]."), *FileDescriptor.GetSourcePath());
-	UE_LOG(LogCADTranslator, Display, TEXT(" - Parsing Library:     %s"), *CADLibrary::GCADLibrary);
+	UE_LOG(LogCADTranslator, Display, TEXT(" - Parsing Library:     %s"), *CADLibrary::FImportParameters::GCADLibrary);
 	UE_LOG(LogCADTranslator, Display, TEXT(" - Tesselation Library: %s")
-		, CADLibrary::FImportParameters::bGDisableCADKernelTessellation ? *CADLibrary::GCADLibrary : TEXT("CADKernel"));
+		, CADLibrary::FImportParameters::bGDisableCADKernelTessellation ? *CADLibrary::FImportParameters::GCADLibrary : TEXT("CADKernel"));
 	UE_LOG(LogCADTranslator, Display, TEXT(" - Cache mode:          %s")
 		, CADLibrary::FImportParameters::bGEnableCADCache ? (CADLibrary::FImportParameters::bGOverwriteCache ? TEXT("Override") : TEXT("Enabled")) : TEXT("Disabled"));
 	UE_LOG(LogCADTranslator, Display, TEXT(" - Processing:          %s")
