Implement retessellate for TechSoft
Fixed some compilation warning
Moved UCoreTechParametricSurfaceData to ParametricSurface module
Created UTempCoreTechParametricSurfaceData class temporarily as CoreTech is still used

diff --git a/Engine/Plugins/Enterprise/DatasmithCADImporter/Source/CADInterfaces/CADInterfaces.Build.cs b/Engine/Plugins/Enterprise/DatasmithCADImporter/Source/CADInterfaces/CADInterfaces.Build.cs
--- a/Engine/Plugins/Enterprise/DatasmithCADImporter/Source/CADInterfaces/CADInterfaces.Build.cs
+++ b/Engine/Plugins/Enterprise/DatasmithCADImporter/Source/CADInterfaces/CADInterfaces.Build.cs
@@ -18,7 +18,8 @@ namespace UnrealBuildTool.Rules
 					"CADKernel",
 					"CADLibrary",
 					"CADTools",
-					"DatasmithCore"
+					"DatasmithCore",
+					"Json",
 				}
 			);
 
diff --git a/Engine/Plugins/Enterprise/DatasmithCADImporter/Source/CADInterfaces/Private/TechSoftBridge.h b/Engine/Plugins/Enterprise/DatasmithCADImporter/Source/CADInterfaces/Private/TechSoftBridge.h
--- a/Engine/Plugins/Enterprise/DatasmithCADImporter/Source/CADInterfaces/Private/TechSoftBridge.h
+++ b/Engine/Plugins/Enterprise/DatasmithCADImporter/Source/CADInterfaces/Private/TechSoftBridge.h
@@ -10,12 +10,22 @@
 #include "CADKernel/Geo/GeoEnum.h"
 #include "CADKernel/Math/MatrixH.h"
 
+#if PLATFORM_WINDOWS
+#include "Windows/AllowWindowsPlatformTypes.h"
+#endif
+
 THIRD_PARTY_INCLUDES_START
+#pragma push_macro("TEXT")
 #include "A3DSDKEnums.h"
 #include "A3DSDKGeometry.h"
 #include "A3DSDKGeometrySrf.h"
+#pragma pop_macro("TEXT")
 THIRD_PARTY_INCLUDES_END
 
+#if PLATFORM_WINDOWS
+#include "Windows/HideWindowsPlatformTypes.h"
+#endif
+
 typedef void A3DAsmProductOccurrence;
 typedef void A3DCrvBase;
 typedef void A3DCrvBase;
diff --git a/Engine/Plugins/Enterprise/DatasmithCADImporter/Source/CADInterfaces/Private/TechSoftFileParser.cpp b/Engine/Plugins/Enterprise/DatasmithCADImporter/Source/CADInterfaces/Private/TechSoftFileParser.cpp
--- a/Engine/Plugins/Enterprise/DatasmithCADImporter/Source/CADInterfaces/Private/TechSoftFileParser.cpp
+++ b/Engine/Plugins/Enterprise/DatasmithCADImporter/Source/CADInterfaces/Private/TechSoftFileParser.cpp
@@ -8,6 +8,8 @@
 
 #include "TechSoftInterface.h"
 
+#include "Serialization/JsonSerializer.h"
+#include "Serialization/JsonWriter.h"
 #include "Templates/UnrealTemplate.h"
 
 namespace CADLibrary
@@ -375,7 +377,7 @@ void FTechSoftFileParser::GenerateBodyMeshes()
 			//	{
 			//		// Aggregate all the bodies into one?
 			//		MeshRepresentationsWithTechSoft(NewBRepCount, NewBReps, Body);
-			//		TechSoftInterface.FillBodyMesh(NewBRepCount, NewBReps, CADFileData.GetImportParameters(), FileUnit, BodyMesh);???
+			//		TechSoftInterface.FillBodyMeshes(NewBRepCount, NewBReps, CADFileData.GetImportParameters(), FileUnit, BodyMesh);???
 			//	}
 			//}
 			//else
@@ -396,7 +398,7 @@ void FTechSoftFileParser::GenerateBodyMeshes()
 				uint32 CachedStyleIndex = Tessellation.MaterialName;
 				Tessellation.MaterialName = 0;
 
-				if (CachedStyleIndex != 0xffffffff)
+				if (CachedStyleIndex != FTechSoftInterface::InvalidScriptIndex)
 				{
 					ExtractGraphStyleProperties(CachedStyleIndex, ColorName, MaterialName);
 				}
@@ -417,7 +419,7 @@ void FTechSoftFileParser::GenerateBodyMeshes()
 			Body.ColorFaceSet = BodyMesh.ColorSet;
 			Body.MaterialFaceSet = BodyMesh.MaterialSet;
 
-			// Write part's representation as hsf file
+			// Write part's representation as hsf file if it is a BRep
 			A3DEEntityType Type;
 			A3DEntityGetType(RepresentationItemPtr, &Type);
 
@@ -426,14 +428,22 @@ void FTechSoftFileParser::GenerateBodyMeshes()
 				FString FilePath = CADFileData.GetBodyCachePath(Body.MeshActorName);
 				if (!FilePath.IsEmpty())
 				{
-					this->TechSoftInterface.SaveBodyToHsfFile(RepresentationItemPtr, FilePath);
-				}
-			}
+					TSharedPtr<FJsonObject> JsonObject = MakeShared<FJsonObject>();
 
-			// #ueent_techsoft: Take care of materials stored in BodyMesh
+					// Save file unit and default color and material attributes in a json string
+					// This will be used when the file is reloaded
+					JsonObject->SetNumberField(JSON_ENTRY_FILE_UNIT, FileUnit);
+					JsonObject->SetNumberField(JSON_ENTRY_COLOR_NAME, DefaultColorName);
+					JsonObject->SetNumberField(JSON_ENTRY_MATERIAL_NAME, DefaultMaterialName);
 
-			Body.ColorFaceSet = BodyMesh.ColorSet;
-			Body.MaterialFaceSet = BodyMesh.MaterialSet;
+					FString JsonString;
+					TSharedRef< TJsonWriter< TCHAR, TPrettyJsonPrintPolicy<TCHAR> > > JsonWriter = TJsonWriterFactory< TCHAR, TPrettyJsonPrintPolicy<TCHAR> >::Create(&JsonString);
+
+					FJsonSerializer::Serialize(JsonObject.ToSharedRef(), JsonWriter);
+
+					TechSoftInterface.SaveBodyToHsfFile(RepresentationItemPtr, FilePath, JsonString);
+				}
+			}
 		}
 	}
 	else
@@ -1442,7 +1452,7 @@ FArchiveColor& FTechSoftFileParser::FindOrAddColor(uint32 ColorIndex, uint8 Alph
 	}
 
 	FArchiveColor& NewColor = CADFileData.AddColor(ColorHId);
-	NewColor.Color = TechSoftFileParserImpl::GetColorAt(ColorIndex);
+	NewColor.Color = TechSoftUtils::GetColorAt(ColorIndex);
 	NewColor.Color.A = Alpha;
 	
 	NewColor.UEMaterialName = BuildColorName(NewColor.Color);
@@ -1457,9 +1467,9 @@ FArchiveMaterial& FTechSoftFileParser::AddMaterialAt(uint32 MaterialIndexToSave,
 	TUniqueTSObjFromIndex<A3DGraphMaterialData> MaterialData(MaterialIndexToSave);
 	if(MaterialData.IsValid())
 	{
-		Material.Diffuse = TechSoftFileParserImpl::GetColorAt(MaterialData->m_uiDiffuse);
-		Material.Ambient = TechSoftFileParserImpl::GetColorAt(MaterialData->m_uiAmbient);
-		Material.Specular = TechSoftFileParserImpl::GetColorAt(MaterialData->m_uiSpecular);
+		Material.Diffuse = TechSoftUtils::GetColorAt(MaterialData->m_uiDiffuse);
+		Material.Ambient = TechSoftUtils::GetColorAt(MaterialData->m_uiAmbient);
+		Material.Specular = TechSoftUtils::GetColorAt(MaterialData->m_uiSpecular);
 		Material.Shininess = MaterialData->m_dShininess;
 		if(GraphStyleData.m_bIsTransparencyDefined)
 		{
@@ -1474,6 +1484,8 @@ FArchiveMaterial& FTechSoftFileParser::AddMaterialAt(uint32 MaterialIndexToSave,
 }
 
 
+// Look at TechSoftUtils::BuildCADMaterial if any loigc changes in this method
+// or any of the methos it calls
 FArchiveMaterial& FTechSoftFileParser::FindOrAddMaterial(uint32 MaterialIndex, const A3DGraphStyleData& GraphStyleData)
 {
 	if (FArchiveMaterial* MaterialArchive = CADFileData.FindMaterial(MaterialIndex))
@@ -1556,6 +1568,8 @@ void FTechSoftFileParser::ExtractGraphicProperties(const A3DGraphics* Graphics,
 	}
 }
 
+// Please review TechSoftUtils::GetMaterialValues if anything changes
+// in this method or the methods it calls
 void FTechSoftFileParser::ExtractGraphStyleProperties(uint32 StyleIndex, FCADUUID& OutColorName, FCADUUID& OutMaterialName)
 {
 	TUniqueTSObjFromIndex<A3DGraphStyleData> GraphStyleData(StyleIndex);
diff --git a/Engine/Plugins/Enterprise/DatasmithCADImporter/Source/CADInterfaces/Private/TechSoftInterface.cpp b/Engine/Plugins/Enterprise/DatasmithCADImporter/Source/CADInterfaces/Private/TechSoftInterface.cpp
--- a/Engine/Plugins/Enterprise/DatasmithCADImporter/Source/CADInterfaces/Private/TechSoftInterface.cpp
+++ b/Engine/Plugins/Enterprise/DatasmithCADImporter/Source/CADInterfaces/Private/TechSoftInterface.cpp
@@ -4,20 +4,34 @@
 
 #include "TechSoftInterface.h"
 
+#include "TUniqueTechSoftObj.h"
+
 #include "CADOptions.h"
 #include "CADInterfacesModule.h"
 
 #include "HAL/PlatformProcess.h"
+#include "Math/Color.h"
 #include "Misc/Paths.h"
-#include "TUniqueTechSoftObj.h"
+#include "Serialization/JsonSerializer.h"
+#include "Serialization/JsonWriter.h"
 
 #ifndef CADKERNEL_DEV
 #include "TechSoftInterfaceUtils.inl"
 #endif
 
+#include <string>
+
 namespace CADLibrary
 {
-FTechSoftInterface& GetTechSoftInterface()
+namespace TechSoftUtils
 {
+#if defined USE_TECHSOFT_SDK && !defined CADKERNEL_DEV
+	TSharedPtr<FJsonObject> GetJsonObject(A3DAsmProductOccurrence* ProductOcccurence);
+	void RestoreMaterials(const TSharedPtr<FJsonObject>& DefaultValues, CADLibrary::FBodyMesh& BodyMesh);
+#endif
+}
+
+	FTechSoftInterface& GetTechSoftInterface()
+{
 	static FTechSoftInterface TechSoftInterface;
 	return TechSoftInterface;
@@ -70,14 +84,41 @@
 #endif
 }
 
-void FTechSoftInterface::SaveBodyToHsfFile(void* BodyPtr, const FString& Filename)
+#ifdef USE_TECHSOFT_SDK
+
+A3DUns32 FTechSoftInterface::InvalidScriptIndex = -1;
+
+A3DStatus FTechSoftInterface::Import(const A3DImport& Import)
 {
+	return ExchangeLoader->Import(Import);
+}
+
+A3DStatus FTechSoftInterface::UnloadModel()
+{
+	A3DStatus Ret = A3DStatus::A3D_ERROR;
+	if (ExchangeLoader)
+	{
+		Ret = A3DAsmModelFileDelete(ExchangeLoader->m_psModelFile);
+	}
+	ExchangeLoader->m_psModelFile = nullptr;
+	return Ret;
+}
+
+A3DAsmModelFile* FTechSoftInterface::GetModelFile()
+{
+	return ExchangeLoader->m_psModelFile;
+}
+#endif
+
+void FTechSoftInterface::SaveBodyToHsfFile(void* BodyPtr, const FString& Filename, const FString& JsonString)
+{
+#if defined USE_TECHSOFT_SDK && !defined CADKERNEL_DEV
 	if (!BodyPtr)
 	{
 		return;
 	}
 
-#if defined USE_TECHSOFT_SDK && !defined CADKERNEL_DEV
+	// Create PartDefinition
 	A3DRiRepresentationItem* RepresentationItem = (A3DRiRepresentationItem*)BodyPtr;
 	A3DAsmPartDefinition* PartDefinition = nullptr;
 
@@ -89,6 +130,7 @@
 
 	A3DAsmPartDefinitionCreate(&PartDefinitionData, &PartDefinition);
 
+	// Create ProductOccurrence
 	A3DAsmProductOccurrence* ProductOccurrence = nullptr;
 
 	A3DAsmProductOccurrenceData ProductOccurrenceData;
@@ -98,6 +140,35 @@
 
 	A3DInt32 iRet = A3DAsmProductOccurrenceCreate(&ProductOccurrenceData, &ProductOccurrence);
 
+	// Add MaterialTable as attribute to ProductOccurrence
+	std::string StringAnsi(TCHAR_TO_UTF8(*JsonString));
+
+	A3DMiscSingleAttributeData SingleAttributeData;
+	A3D_INITIALIZE_DATA(A3DMiscSingleAttributeData, SingleAttributeData);
+
+	SingleAttributeData.m_eType = kA3DModellerAttributeTypeString;
+	SingleAttributeData.m_pcTitle = "MaterialTable";
+	SingleAttributeData.m_pcData = (char*)StringAnsi.c_str();
+
+	A3DMiscAttributeData AttributesData;
+	A3D_INITIALIZE_DATA(A3DMiscAttributeData, AttributesData);
+
+	A3DMiscAttribute* Attributes = nullptr;
+
+	AttributesData.m_pcTitle = SingleAttributeData.m_pcTitle;
+	AttributesData.m_asSingleAttributesData = &SingleAttributeData;
+	AttributesData.m_uiSize = 1;
+	A3DMiscAttributeCreate(&AttributesData, &Attributes);
+
+	A3DRootBaseData RootBaseData;
+	A3D_INITIALIZE_DATA(A3DRootBaseData, RootBaseData);
+
+	RootBaseData.m_pcName = SingleAttributeData.m_pcTitle;
+	RootBaseData.m_ppAttributes = &Attributes;
+	RootBaseData.m_uiSize = 1;
+	A3DRootBaseSet(ProductOccurrence, &RootBaseData);
+
+	// Create ModelFile
 	A3DAsmModelFile* ModelFile = nullptr;
 
 	A3DAsmModelFileData ModelFileData;
@@ -109,6 +180,7 @@
 
 	A3DAsmModelFileCreate(&ModelFileData, &ModelFile);
 
+	// Save ModelFile to hsf file
 	A3DRWParamsExportPrcData sParamsExportData;
 	A3D_INITIALIZE_DATA(A3DRWParamsExportPrcData, sParamsExportData);
 
@@ -121,12 +193,15 @@
 	iRet = A3DAsmModelFileExportToPrcFile(ModelFile, &sParamsExportData, HsfFileName, nullptr);
 
 	// #ueent_techsoft: Deleting the model seems to delete the entire content. To be double-checked
+	A3DEntityDelete(Attributes);
 	A3DEntityDelete(ModelFile);
 #endif
 }
 
-bool FTechSoftInterface::GetBodyFromHsfFile(const FString& Filename, const FImportParameters& ImportParameters, double FileUnit, FBodyMesh& BodyMesh)
+bool FTechSoftInterface::GetBodyFromHsfFile(const FString& Filename, const FImportParameters& ImportParameters, FBodyMesh& BodyMesh)
 {
+	bool bExtractionSuccessful = false;
+
 #if defined USE_TECHSOFT_SDK && !defined CADKERNEL_DEV
 	A3DRWParamsPrcReadHelper* ReadHelper = nullptr;
 	A3DAsmModelFile* ModelFile = nullptr;
@@ -156,10 +231,25 @@
 		return false;
 	}
 
-	return FillBodyMesh(PartDefinitionData->m_ppRepItems[0], ImportParameters, FileUnit, BodyMesh);
-#else
-	return false;
+	TSharedPtr<FJsonObject> JsonObject = TechSoftUtils::GetJsonObject(ModelFileData->m_ppPOccurrences[0]);
+	if (JsonObject.IsValid())
+	{
+		double FileUnit = 1.0;
+
+		JsonObject->TryGetNumberField(JSON_ENTRY_FILE_UNIT, FileUnit);
+
+		if (FillBodyMesh(PartDefinitionData->m_ppRepItems[0], ImportParameters, FileUnit, BodyMesh))
+		{
+			bExtractionSuccessful = true;
+
+			TechSoftUtils::RestoreMaterials(JsonObject, BodyMesh);
+		}
+	}
+
+	A3DEntityDelete(ModelFile);
 #endif
+
+	return bExtractionSuccessful;
 }
 
 bool FTechSoftInterface::FillBodyMesh(void* BodyPtr, const FImportParameters& ImportParameters, double FileUnit, FBodyMesh& BodyMesh)
@@ -206,33 +296,254 @@
 	TechSoftInterfaceUtils::FTechSoftTessellationExtractor Extractor(RepresentationItemData->m_pTessBase);
 	return Extractor.FillBodyMesh(BodyMesh, FileUnit);
 #else
-return false;
+	return false;
 #endif
+}
+
+namespace TechSoftUtils
+{
+
+#if defined USE_TECHSOFT_SDK && !defined CADKERNEL_DEV
+TSharedPtr<FJsonObject> GetJsonObject(A3DAsmProductOccurrence* ProductOcccurence)
+{
+	TUniqueTSObj<A3DRootBaseData> RootBaseData(ProductOcccurence);
+
+	if (RootBaseData.IsValid() && RootBaseData->m_uiSize > 0)
+	{
+		TUniqueTSObj<A3DMiscAttributeData> AttributeData(RootBaseData->m_ppAttributes[0]);
+		if (AttributeData->m_uiSize > 0 && AttributeData->m_asSingleAttributesData[0].m_eType == kA3DModellerAttributeTypeString)
+		{
+			FString JsonString = UTF8_TO_TCHAR(AttributeData->m_asSingleAttributesData[0].m_pcData);
+
+			TSharedRef<TJsonReader<>> JsonReader = TJsonReaderFactory<>::Create(JsonString);
+			TSharedPtr<FJsonObject> JsonObject;
+
+			if (FJsonSerializer::Deserialize(JsonReader, JsonObject))
+			{
+				return JsonObject;
+			}
+		}
+	}
+
+	return {};
+}
+
+FColor GetColorAt(uint32 ColorIndex)
+{
+	TUniqueTSObjFromIndex<A3DGraphRgbColorData> ColorData(ColorIndex);
+	if (ColorData.IsValid())
+	{
+		return FColor((uint8)(ColorData->m_dRed * 255)
+			, (uint8)(ColorData->m_dGreen * 255)
+			, (uint8)(ColorData->m_dBlue * 255));
+	}
+	else
+	{
+		return FColor(200, 200, 200);
+	}
+}
+
+// Replicate logic in FTechSoftFileParser::FindOrAddMaterial and the methods it calls
+void BuildCADMaterial(uint32 MaterialIndex, const A3DGraphStyleData& GraphStyleData, FCADMaterial& Material)
+{
+	if (TechSoftUtils::IsMaterialTexture(MaterialIndex))
+	{
+		TUniqueTSObjFromIndex<A3DGraphTextureApplicationData> TextureData(MaterialIndex);
+		if (TextureData.IsValid())
+		{
+			MaterialIndex = TextureData->m_uiMaterialIndex;
+		}
+	}
+
+	TUniqueTSObjFromIndex<A3DGraphMaterialData> MaterialData(MaterialIndex);
+	if (MaterialData.IsValid())
+	{
+		Material.Diffuse = GetColorAt(MaterialData->m_uiDiffuse);
+		Material.Ambient = GetColorAt(MaterialData->m_uiAmbient);
+		Material.Specular = GetColorAt(MaterialData->m_uiSpecular);
+		Material.Shininess = MaterialData->m_dShininess;
+		if (GraphStyleData.m_bIsTransparencyDefined)
+		{
+			Material.Transparency = 1. - GraphStyleData.m_ucTransparency / 255.;
+		}
+		// todo: find how to convert Emissive color into ? reflexion coef...
+		// Material.Emissive = GetColor(MaterialData->m_uiEmissive);
+		// Material.Reflexion;
+	}
+}
+
+// Replicates the logic in FTechSoftFileParser::ExtractGraphStyleProperties
+void GetMaterialValues(uint32 StyleIndex, FCADUUID& OutColorName, FCADUUID& OutMaterialName)
+{
+	TUniqueTSObjFromIndex<A3DGraphStyleData> GraphStyleData(StyleIndex);
+
+	if (GraphStyleData.IsValid())
+	{
+		if (GraphStyleData->m_bMaterial)
+		{
+			FCADMaterial Material;
+
+			BuildCADMaterial(GraphStyleData->m_uiRgbColorIndex, *GraphStyleData, Material);
+
+			OutMaterialName = BuildMaterialName(Material);
+		}
+		else
+		{
+			TUniqueTSObjFromIndex<A3DGraphRgbColorData> ColorData(GraphStyleData->m_uiRgbColorIndex);
+			if (ColorData.IsValid())
+			{
+				const uint8 Alpha = GraphStyleData->m_bIsTransparencyDefined ? GraphStyleData->m_ucTransparency : 255;
+				const FColor ColorValue((uint8)(ColorData->m_dRed * 255), (uint8)(ColorData->m_dGreen * 255), (uint8)(ColorData->m_dBlue * 255), Alpha);
+
+				OutColorName = BuildColorName(ColorValue);
+			}
+		}
+	}
+}
+
+void RestoreMaterials(const TSharedPtr<FJsonObject>& DefaultValues, FBodyMesh& BodyMesh)
+{
+	FCADUUID DefaultColorName = 0;
+	FCADUUID DefaultMaterialName = 0;
+
+	DefaultValues->TryGetNumberField(JSON_ENTRY_COLOR_NAME, DefaultColorName);
+	DefaultValues->TryGetNumberField(JSON_ENTRY_MATERIAL_NAME, DefaultMaterialName);
+
+	BodyMesh.MaterialSet.Empty();
+	BodyMesh.ColorSet.Empty();
+
+	for (FTessellationData& Tessellation : BodyMesh.Faces)
+	{
+		// Extract proper color or material based on style index
+		uint32 CachedStyleIndex = Tessellation.MaterialName;
+		Tessellation.MaterialName = 0;
+
+		FCADUUID ColorName = DefaultColorName;
+		FCADUUID MaterialName = DefaultMaterialName;
+
+		GetMaterialValues(CachedStyleIndex, ColorName, MaterialName);
+
+		if (ColorName)
+		{
+			Tessellation.ColorName = ColorName;
+			BodyMesh.ColorSet.Add(ColorName);
+		}
+
+		if (MaterialName)
+		{
+			Tessellation.MaterialName = MaterialName;
+			BodyMesh.MaterialSet.Add(MaterialName);
+		}
+	}
+}
+#endif
+
+FTechSoftInterface& GetTechSoftInterface()
+{
+	static FTechSoftInterface TechSoftInterface;
+	return TechSoftInterface;
+}
+
+FString GetTechSoftVersion()
+{
+#ifdef USE_TECHSOFT_SDK
+	A3DInt32 MajorVersion = 0, MinorVersion = 0;
+	A3DDllGetVersion(&MajorVersion, &MinorVersion);
+	return FString::Printf(TEXT("Techsoft %d.%d"), MajorVersion, MinorVersion);
+#endif
+	return FString();
+}
+
+
+bool TECHSOFT_InitializeKernel(const TCHAR* InEnginePluginsPath)
+{
+	return GetTechSoftInterface().InitializeKernel(InEnginePluginsPath);
 }
 
 #ifdef USE_TECHSOFT_SDK
 
-A3DStatus FTechSoftInterface::Import(const A3DImport& Import)
+A3DStatus GetGlobalPointer(A3DGlobal** GlobalPtr)
 {
-	return ExchangeLoader->Import(Import);
+	return A3DGlobalGetPointer(GlobalPtr);
 }
 
-A3DStatus FTechSoftInterface::UnloadModel()
+A3DStatus GetSurfaceAsNurbs(const A3DSurfBase* SurfacePtr, A3DSurfNurbsData* DataPtr, A3DDouble Tolerance, A3DBool bUseSameParameterization)
 {
-	A3DStatus Ret = A3DStatus::A3D_ERROR;
-	if (ExchangeLoader)
-	{
-		Ret = A3DAsmModelFileDelete(ExchangeLoader->m_psModelFile);
-	}
-	ExchangeLoader->m_psModelFile = nullptr;
-	return Ret;
+	return A3DSurfBaseGetAsNurbs(SurfacePtr, Tolerance, bUseSameParameterization, DataPtr);
 }
 
-A3DAsmModelFile* FTechSoftInterface::GetModelFile()
+A3DStatus GetCurveAsNurbs(const A3DCrvBase* CurvePtr, A3DCrvNurbsData* DataPtr, A3DDouble Tolerance, A3DBool bUseSameParameterization)
 {
-	return ExchangeLoader->m_psModelFile;
+	return A3DCrvBaseGetAsNurbs(CurvePtr, Tolerance, bUseSameParameterization, DataPtr);
 }
 
+A3DStatus GetOriginalFilePathName(const A3DAsmProductOccurrence* A3DOccurrencePtr, A3DUTF8Char** FilePathUTF8Ptr)
+{
+	return A3DAsmProductOccurrenceGetOriginalFilePathName(A3DOccurrencePtr, FilePathUTF8Ptr);
+}
+
+A3DStatus GetFilePathName(const A3DAsmProductOccurrence* A3DOccurrencePtr, A3DUTF8Char** FilePathUTF8Ptr)
+{
+	return A3DAsmProductOccurrenceGetFilePathName(A3DOccurrencePtr, FilePathUTF8Ptr);
+}
+
+A3DStatus GetEntityType(const A3DEntity* EntityPtr, A3DEEntityType* EntityTypePtr)
+{
+	return A3DEntityGetType(EntityPtr, EntityTypePtr);
+}
+
+bool IsEntityBaseWithGraphicsType(const A3DEntity* EntityPtr)
+{
+	return (bool)A3DEntityIsBaseWithGraphicsType(EntityPtr);
+}
+
+bool IsEntityBaseType(const A3DEntity* EntityPtr)
+{
+	return (bool)A3DEntityIsBaseType(EntityPtr);
+}
+
+bool IsMaterialTexture(const uint32 MaterialIndex)
+{
+	A3DBool bIsTexture = false;
+	return A3DGlobalIsMaterialTexture(MaterialIndex, &bIsTexture) == A3DStatus::A3D_SUCCESS ? bool(bIsTexture) : false;
+}
+
+A3DEntity* GetPointerFromIndex(const uint32 Index, const A3DEEntityType Type)
+{
+	A3DEntity* EntityPtr;
+	if (A3DMiscPointerFromIndexGet(Index, Type, &EntityPtr) != A3DStatus::A3D_SUCCESS)
+	{
+		return nullptr;
+	}
+	return EntityPtr;
+}
+
+A3DStatus HealBRep(A3DRiBrepModel** BRepToHeal, double Tolerance, A3DSewOptionsData const* SewOptions, A3DRiBrepModel*** OutNewBReps, uint32& OutNewBRepCount)
+{
+	A3DUns32 NewBRepCount;
+	A3DStatus Status = A3DSewBrep(&BRepToHeal, 1, Tolerance, SewOptions, OutNewBReps, &NewBRepCount);
+	OutNewBRepCount = NewBRepCount;
+	return Status;
+}
+
+A3DStatus SewBReps(A3DRiBrepModel*** BRepsToSew, uint32 const BRepCount, double Tolerance, A3DSewOptionsData const* SewOptions, A3DRiBrepModel*** OutNewBReps, uint32& OutNewBRepCount)
+{
+	A3DUns32 NewBRepCount;
+	A3DStatus Status = A3DSewBrep(BRepsToSew, BRepCount, Tolerance, SewOptions, OutNewBReps, &NewBRepCount);
+	OutNewBRepCount = NewBRepCount;
+	return Status;
+}
+
+A3DStatus SewModel(A3DAsmModelFile** ModelPtr, double Tolerance, A3DSewOptionsData const* SewOptions)
+{
+	return A3DAsmModelFileSew(ModelPtr, Tolerance, SewOptions);
+}
+
+#endif
+
+}
+
+#ifdef USE_TECHSOFT_SDK
 void TUniqueTSObj<A3DAsmModelFileData>::InitializeData()
 {
 	A3D_INITIALIZE_DATA(A3DAsmModelFileData, Data);
@@ -1013,120 +1324,6 @@
 {
 	return A3DAsmProductOccurrenceGetFilePathName(InOccurrencePtr, &Data);
 }
-
-#endif
-
-namespace TechSoftUtils
-{
-
-FTechSoftInterface& GetTechSoftInterface()
-{
-	static FTechSoftInterface TechSoftInterface;
-	return TechSoftInterface;
-}
-
-FString GetTechSoftVersion()
-{
-#ifdef USE_TECHSOFT_SDK
-	A3DInt32 MajorVersion = 0, MinorVersion = 0;
-	A3DDllGetVersion(&MajorVersion, &MinorVersion);
-	return FString::Printf(TEXT("Techsoft %d.%d"), MajorVersion, MinorVersion);
-#endif
-	return FString();
-}
-
-
-bool TECHSOFT_InitializeKernel(const TCHAR* InEnginePluginsPath)
-{
-	return GetTechSoftInterface().InitializeKernel(InEnginePluginsPath);
-}
-
-#ifdef USE_TECHSOFT_SDK
-
-A3DStatus GetGlobalPointer(A3DGlobal** GlobalPtr)
-{
-	return A3DGlobalGetPointer(GlobalPtr);
-}
-
-A3DStatus GetSurfaceAsNurbs(const A3DSurfBase* SurfacePtr, A3DSurfNurbsData* DataPtr, A3DDouble Tolerance, A3DBool bUseSameParameterization)
-{
-	return A3DSurfBaseGetAsNurbs(SurfacePtr, Tolerance, bUseSameParameterization, DataPtr);
-}
-
-A3DStatus GetCurveAsNurbs(const A3DCrvBase* CurvePtr, A3DCrvNurbsData* DataPtr, A3DDouble Tolerance, A3DBool bUseSameParameterization)
-{
-	return A3DCrvBaseGetAsNurbs(CurvePtr, Tolerance, bUseSameParameterization, DataPtr);
-}
-
-A3DStatus GetOriginalFilePathName(const A3DAsmProductOccurrence* A3DOccurrencePtr, A3DUTF8Char** FilePathUTF8Ptr)
-{
-	return A3DAsmProductOccurrenceGetOriginalFilePathName(A3DOccurrencePtr, FilePathUTF8Ptr);
-}
-
-A3DStatus GetFilePathName(const A3DAsmProductOccurrence* A3DOccurrencePtr, A3DUTF8Char** FilePathUTF8Ptr)
-{
-	return A3DAsmProductOccurrenceGetFilePathName(A3DOccurrencePtr, FilePathUTF8Ptr);
-}
-
-A3DStatus GetEntityType(const A3DEntity* EntityPtr, A3DEEntityType* EntityTypePtr)
-{
-	return A3DEntityGetType(EntityPtr, EntityTypePtr);
-}
-
-bool IsEntityBaseWithGraphicsType(const A3DEntity* EntityPtr)
-{
-	return (bool)A3DEntityIsBaseWithGraphicsType(EntityPtr);
-}
-
-bool IsEntityBaseType(const A3DEntity* EntityPtr)
-{
-	return (bool)A3DEntityIsBaseType(EntityPtr);
-}
-
-bool IsMaterialTexture(const uint32 MaterialIndex)
-{
-	A3DBool bIsTexture = false;
-	if (A3DGlobalIsMaterialTexture(MaterialIndex, &bIsTexture) != A3DStatus::A3D_SUCCESS)
-	{
-		return false;
-	}
-	return bIsTexture != false;
-}
-
-A3DEntity* GetPointerFromIndex(const uint32 Index, const A3DEEntityType Type)
-{
-	A3DEntity* EntityPtr;
-	if (A3DMiscPointerFromIndexGet(Index, Type, &EntityPtr) != A3DStatus::A3D_SUCCESS)
-	{
-		return nullptr;
-	}
-	return EntityPtr;
-}
-
-A3DStatus HealBRep(A3DRiBrepModel** BRepToHeal, double Tolerance, A3DSewOptionsData const* SewOptions, A3DRiBrepModel*** OutNewBReps, uint32& OutNewBRepCount)
-{
-	A3DUns32 NewBRepCount;
-	A3DStatus Status = A3DSewBrep(&BRepToHeal, 1, Tolerance, SewOptions, OutNewBReps, &NewBRepCount);
-	OutNewBRepCount = NewBRepCount;
-	return Status;
-}
-
-A3DStatus SewBReps(A3DRiBrepModel*** BRepsToSew, uint32 const BRepCount, double Tolerance, A3DSewOptionsData const* SewOptions, A3DRiBrepModel*** OutNewBReps, uint32& OutNewBRepCount)
-{
-	A3DUns32 NewBRepCount;
-	A3DStatus Status = A3DSewBrep(BRepsToSew, BRepCount, Tolerance, SewOptions, OutNewBReps, &NewBRepCount);
-	OutNewBRepCount = NewBRepCount;
-	return Status;
-}
-
-A3DStatus SewModel(A3DAsmModelFile** ModelPtr, double Tolerance, A3DSewOptionsData const* SewOptions)
-{
-	return A3DAsmModelFileSew(ModelPtr, Tolerance, SewOptions);
-}
-
 #endif
 
 }
-
-
-}
diff --git a/Engine/Plugins/Enterprise/DatasmithCADImporter/Source/CADInterfaces/Private/TechSoftInterfaceUtils.inl b/Engine/Plugins/Enterprise/DatasmithCADImporter/Source/CADInterfaces/Private/TechSoftInterfaceUtils.inl
--- a/Engine/Plugins/Enterprise/DatasmithCADImporter/Source/CADInterfaces/Private/TechSoftInterfaceUtils.inl
+++ b/Engine/Plugins/Enterprise/DatasmithCADImporter/Source/CADInterfaces/Private/TechSoftInterfaceUtils.inl
@@ -20,14 +20,15 @@ namespace TechSoftInterfaceUtils
 
 		bool  FillBodyMesh(CADLibrary::FBodyMesh& BodyMesh, double FileUnit)
 		{
+			int32 VertexOffset = BodyMesh.VertexArray.Num();
 			FillVertexArray(FileUnit, BodyMesh.VertexArray);
 
-			if (BodyMesh.VertexArray.Num() == 0)
+			if (BodyMesh.VertexArray.Num() == VertexOffset)
 			{
 				return false;
 			}
 
-			FillFaceArray(BodyMesh);
+			FillFaceArray(BodyMesh, VertexOffset);
 
 			return BodyMesh.Faces.Num() > 0 ? true : false;
 		}
@@ -45,7 +46,7 @@ namespace TechSoftInterfaceUtils
 			}
 
 			int32 VertexCount = TessellationBaseData->m_uiCoordSize / 3;
-			VertexArray.Reserve(VertexCount);
+			VertexArray.Reserve(VertexArray.Num() + VertexCount);
 
 			double* Coordinates = TessellationBaseData->m_pdCoords;
 			for (unsigned int Index = 0; Index < TessellationBaseData->m_uiCoordSize; ++Index)
@@ -103,7 +104,7 @@ namespace TechSoftInterfaceUtils
 			return TriangleCount;
 		}
 
-		void FillFaceArray(CADLibrary::FBodyMesh& BodyMesh)
+		void FillFaceArray(CADLibrary::FBodyMesh& BodyMesh, int32 VertexOffset = 0)
 		{
 			using namespace CADLibrary;
 
@@ -125,8 +126,8 @@ namespace TechSoftInterfaceUtils
 				const A3DTessFaceData& FaceTessData = TessellationData->m_psFaceTessData[Index];
 				FTessellationData& Tessellation = Faces.Emplace_GetRef();
 
-				Tessellation.MaterialName = 0xffffffff; // Assumes -1 is an invalid value for m_puiStyleIndexes
-				if (FaceTessData.m_uiStyleIndexesSize == 1)
+				Tessellation.MaterialName = FTechSoftInterface::InvalidScriptIndex;
+				if (FaceTessData.m_uiStyleIndexesSize > 0)
 				{
 					// Store the StyleIndex on the MaterialName. It will be processed after tessellation
 					Tessellation.MaterialName = FaceTessData.m_puiStyleIndexes[0];
@@ -270,6 +271,14 @@ namespace TechSoftInterfaceUtils
 					}
 					bMustProcess = FaceTessData.m_uiSizesTriangulatedSize > FaceSetIndex;
 				}
+				
+				if (VertexOffset > 0)
+				{
+					for (int32& VertexIndex : Tessellation.VertexIndices)
+					{
+						VertexIndex += VertexOffset;
+					}
+				}
 
 				ensure(!bMustProcess);
 			}
diff --git a/Engine/Plugins/Enterprise/DatasmithCADImporter/Source/CADInterfaces/Public/TechSoftInterface.h b/Engine/Plugins/Enterprise/DatasmithCADImporter/Source/CADInterfaces/Public/TechSoftInterface.h
--- a/Engine/Plugins/Enterprise/DatasmithCADImporter/Source/CADInterfaces/Public/TechSoftInterface.h
+++ b/Engine/Plugins/Enterprise/DatasmithCADImporter/Source/CADInterfaces/Public/TechSoftInterface.h
@@ -13,7 +13,9 @@
 
 THIRD_PARTY_INCLUDES_START
 #ifdef USE_TECHSOFT_SDK
+#pragma push_macro("TEXT")
 #include "A3DSDKIncludes.h"
+#pragma pop_macro("TEXT")
 #endif
 THIRD_PARTY_INCLUDES_END
 
@@ -21,8 +23,14 @@ THIRD_PARTY_INCLUDES_END
 #include "Windows/HideWindowsPlatformTypes.h"
 #endif
 
+#define JSON_ENTRY_FILE_UNIT		TEXT("FileUnit")
+#define JSON_ENTRY_COLOR_NAME		TEXT("ColorName")
+#define JSON_ENTRY_MATERIAL_NAME	TEXT("MaterialName")
+
 namespace CADLibrary
 {
+	class FImportParameters;
+	class FBodyMesh;
 
 CADINTERFACES_API class FTechSoftInterface
 {
@@ -31,30 +39,32 @@ public:
 	* Returns true if the object has been created outside of the memory pool of the running process
 	* This is the case when the object has been created by the DatasmithRuntime plugin
 	*/
-	bool IsExternal()
+	CADINTERFACES_API bool IsExternal()
 	{
 		return bIsExternal;
 	}
 
-	void SetExternal(bool Value)
+	CADINTERFACES_API void SetExternal(bool Value)
 	{
 		bIsExternal = Value;
 	}
 
-	bool InitializeKernel(const TCHAR* = TEXT(""));
-
-	void SaveBodyToHsfFile(void* BodyPtr, const FString& Filename);
-
-	bool GetBodyFromHsfFile(const FString& Filename, const class FImportParameters& ImportParameters, double FileUnit, class FBodyMesh& BodyMesh);
-
-	bool FillBodyMesh(void* BodyPtr, const class FImportParameters& ImportParameters, double FileUnit, class FBodyMesh& BodyMesh);
+	CADINTERFACES_API bool InitializeKernel(const TCHAR* = TEXT(""));
 
 #ifdef USE_TECHSOFT_SDK
 	A3DStatus Import(const A3DImport& Import);
 	A3DAsmModelFile* GetModelFile();
 	A3DStatus UnloadModel();
+
+	static A3DUns32 InvalidScriptIndex;
 #endif
 
+	CADINTERFACES_API bool GetBodyFromHsfFile(const FString& Filename, const FImportParameters& ImportParameters, FBodyMesh& BodyMesh);
+
+	bool FillBodyMesh(void* BodyPtr, const FImportParameters& ImportParameters, double FileUnit, FBodyMesh& BodyMesh);
+
+	void SaveBodyToHsfFile(void* BodyPtr, const FString& Filename, const FString& JsonString);
+
 private:
 
 	bool bIsExternal = false;
@@ -67,12 +77,12 @@ private:
 
 namespace TechSoftUtils
 {
-
 CADINTERFACES_API bool TECHSOFT_InitializeKernel(const TCHAR* = TEXT(""));
 CADINTERFACES_API FTechSoftInterface& GetTechSoftInterface();
 CADINTERFACES_API FString GetTechSoftVersion();
 
 #ifdef USE_TECHSOFT_SDK
+FColor GetColorAt(uint32 ColorIndex);
 
 CADINTERFACES_API A3DStatus GetGlobalPointer(A3DGlobal** GlobalPtr);
 
diff --git a/Engine/Plugins/Enterprise/DatasmithCADImporter/Source/CADKernelSurface/Private/CADKernelSurfaceExtension.cpp b/Engine/Plugins/Enterprise/DatasmithCADImporter/Source/CADKernelSurface/Private/CADKernelSurfaceExtension.cpp
--- a/Engine/Plugins/Enterprise/DatasmithCADImporter/Source/CADKernelSurface/Private/CADKernelSurfaceExtension.cpp
+++ b/Engine/Plugins/Enterprise/DatasmithCADImporter/Source/CADKernelSurface/Private/CADKernelSurfaceExtension.cpp
@@ -97,25 +97,14 @@ namespace CADKernelSurface
 	void AddSurfaceDataForMesh(const TCHAR* CADKernelArchive, const CADLibrary::FImportParameters& InSceneParameters, const CADLibrary::FMeshParameters& InMeshParameters, const FDatasmithTessellationOptions& InTessellationOptions, FDatasmithMeshElementPayload& OutMeshPayload)
 	{
 		// Store CADKernel archive if provided
-		if (FPaths::FileExists(CADKernelArchive))
+		UCADKernelParametricSurfaceData* CADKernelData = Datasmith::MakeAdditionalData<UCADKernelParametricSurfaceData>();
+		if (CADKernelData->SetFile(CADKernelArchive))
 		{
-			TArray<uint8> ByteArray;
-			if (FFileHelper::LoadFileToArray(ByteArray, CADKernelArchive))
-			{
-				UCADKernelParametricSurfaceData* CADKernelData = Datasmith::MakeAdditionalData<UCADKernelParametricSurfaceData>();
-				CADKernelData->RawData = MoveTemp(ByteArray);
-				CADKernelData->SceneParameters.ModelCoordSys = uint8(InSceneParameters.GetModelCoordSys());
-				CADKernelData->SceneParameters.MetricUnit = InSceneParameters.GetMetricUnit();
-				CADKernelData->SceneParameters.ScaleFactor = InSceneParameters.GetScaleFactor();
-
-				CADKernelData->MeshParameters.bNeedSwapOrientation = InMeshParameters.bNeedSwapOrientation;
-				CADKernelData->MeshParameters.bIsSymmetric = InMeshParameters.bIsSymmetric;
-				CADKernelData->MeshParameters.SymmetricNormal = InMeshParameters.SymmetricNormal;
-				CADKernelData->MeshParameters.SymmetricOrigin = InMeshParameters.SymmetricOrigin;
-
-				CADKernelData->LastTessellationOptions = InTessellationOptions;
-				OutMeshPayload.AdditionalData.Add(CADKernelData);
-			}
+			CADKernelData->SetImportParameters(InSceneParameters);
+			CADKernelData->SetMeshParameters(InMeshParameters);
+			CADKernelData->SetLastTessellationOptions(InTessellationOptions);
+
+			OutMeshPayload.AdditionalData.Add(CADKernelData);
 		}
 	}
 
diff --git a/Engine/Plugins/Enterprise/DatasmithCADImporter/Source/CADKernelSurface/Public/CADKernelSurfaceExtension.h b/Engine/Plugins/Enterprise/DatasmithCADImporter/Source/CADKernelSurface/Public/CADKernelSurfaceExtension.h
--- a/Engine/Plugins/Enterprise/DatasmithCADImporter/Source/CADKernelSurface/Public/CADKernelSurfaceExtension.h
+++ b/Engine/Plugins/Enterprise/DatasmithCADImporter/Source/CADKernelSurface/Public/CADKernelSurfaceExtension.h
@@ -4,7 +4,7 @@
 
 #include "CoreMinimal.h"
 
-#include "ParametricSurfaceExtension.h"
+#include "ParametricSurfaceData.h"
 #include "DatasmithAdditionalData.h"
 #include "DatasmithCustomAction.h"
 #include "DatasmithImportOptions.h"
@@ -18,14 +18,6 @@ class CADKERNELSURFACE_API UCADKernelParametricSurfaceData : public UParametricS
 	GENERATED_BODY()
 
 public:
-	// Too costly to serialize as a UPROPERTY, will use custom serialization.
-	TArray<uint8> RawData;
-
-	virtual bool IsValid() override
-	{
-		return RawData.Num() > 0;
-	}
-	
 	virtual bool Tessellate(UStaticMesh& StaticMesh, const FDatasmithRetessellationOptions& RetessellateOptions) override;
 
 private:
diff --git a/Engine/Plugins/Enterprise/DatasmithCADImporter/Source/CoreTechSurface/Private/CoreTechSurfaceExtension.cpp b/Engine/Plugins/Enterprise/DatasmithCADImporter/Source/CoreTechSurface/Private/CoreTechSurfaceExtension.cpp
--- a/Engine/Plugins/Enterprise/DatasmithCADImporter/Source/CoreTechSurface/Private/CoreTechSurfaceExtension.cpp
+++ b/Engine/Plugins/Enterprise/DatasmithCADImporter/Source/CoreTechSurface/Private/CoreTechSurfaceExtension.cpp
@@ -16,24 +16,25 @@
 #include "StaticMeshAttributes.h"
 #include "UObject/EnterpriseObjectVersion.h"
 
-void UCoreTechParametricSurfaceData::Serialize(FArchive& Ar)
+bool UTempCoreTechParametricSurfaceData::SetFile(const TCHAR* FilePath)
 {
-	Ar.UsingCustomVersion(FEnterpriseObjectVersion::GUID);
-
-	Super::Serialize(Ar);
-
-	if (Ar.IsSaving() || (Ar.IsLoading() && Ar.CustomVer(FEnterpriseObjectVersion::GUID) >= FEnterpriseObjectVersion::CoreTechParametricSurfaceOptim))
+	if (UParametricSurfaceData::SetFile(FilePath))
 	{
-		Ar << RawData;
+		SourceFile = FilePath;
+		return true;
 	}
 
-	if (RawData_DEPRECATED.Num() && RawData.Num() == 0)
-	{
-		RawData = MoveTemp(RawData_DEPRECATED);
-	}
+	return false;
+}
+
+void UTempCoreTechParametricSurfaceData::Serialize(FArchive& Ar)
+{
+	Super::Serialize(Ar);
+
+	Ar << RawData;
 }
 
-bool UCoreTechParametricSurfaceData::Tessellate(UStaticMesh& StaticMesh, const FDatasmithRetessellationOptions& RetessellateOptions)
+bool UTempCoreTechParametricSurfaceData::Tessellate(UStaticMesh& StaticMesh, const FDatasmithRetessellationOptions& RetessellateOptions)
 {
 	bool bSuccessfulTessellation = false;
 
diff --git a/Engine/Plugins/Enterprise/DatasmithCADImporter/Source/CoreTechSurface/Private/CoreTechSurfaceHelper.cpp b/Engine/Plugins/Enterprise/DatasmithCADImporter/Source/CoreTechSurface/Private/CoreTechSurfaceHelper.cpp
--- a/Engine/Plugins/Enterprise/DatasmithCADImporter/Source/CoreTechSurface/Private/CoreTechSurfaceHelper.cpp
+++ b/Engine/Plugins/Enterprise/DatasmithCADImporter/Source/CoreTechSurface/Private/CoreTechSurfaceHelper.cpp
@@ -67,26 +67,15 @@ namespace CoreTechSurface
 	{
 		if (ICADInterfacesModule::GetAvailability() == ECADInterfaceAvailability::Available)
 		{
-			if (FPaths::FileExists(InFilePath))
+			UTempCoreTechParametricSurfaceData* CoreTechData = Datasmith::MakeAdditionalData<UTempCoreTechParametricSurfaceData>();
+
+			if (CoreTechData->SetFile(InFilePath))
 			{
-				TArray<uint8> ByteArray;
-				if (FFileHelper::LoadFileToArray(ByteArray, InFilePath))
-				{
-					UCoreTechParametricSurfaceData* CoreTechData = Datasmith::MakeAdditionalData<UCoreTechParametricSurfaceData>();
-					CoreTechData->SourceFile = InFilePath;
-					CoreTechData->RawData = MoveTemp(ByteArray);
-					CoreTechData->SceneParameters.ModelCoordSys = uint8(InSceneParameters.GetModelCoordSys());
-					CoreTechData->SceneParameters.MetricUnit = InSceneParameters.GetMetricUnit();
-					CoreTechData->SceneParameters.ScaleFactor = InSceneParameters.GetScaleFactor();
-
-					CoreTechData->MeshParameters.bNeedSwapOrientation = InMeshParameters.bNeedSwapOrientation;
-					CoreTechData->MeshParameters.bIsSymmetric = InMeshParameters.bIsSymmetric;
-					CoreTechData->MeshParameters.SymmetricNormal = InMeshParameters.SymmetricNormal;
-					CoreTechData->MeshParameters.SymmetricOrigin = InMeshParameters.SymmetricOrigin;
-
-					CoreTechData->LastTessellationOptions = InTessellationOptions;
-					OutMeshPayload.AdditionalData.Add(CoreTechData);
-				}
+				CoreTechData->SetImportParameters(InSceneParameters);
+				CoreTechData->SetMeshParameters(InMeshParameters);
+				CoreTechData->SetLastTessellationOptions(InTessellationOptions);
+
+				OutMeshPayload.AdditionalData.Add(CoreTechData);
 			}
 		}
 	}
diff --git a/Engine/Plugins/Enterprise/DatasmithCADImporter/Source/CoreTechSurface/Private/CoreTechSurfaceModule.cpp b/Engine/Plugins/Enterprise/DatasmithCADImporter/Source/CoreTechSurface/Private/CoreTechSurfaceModule.cpp
--- a/Engine/Plugins/Enterprise/DatasmithCADImporter/Source/CoreTechSurface/Private/CoreTechSurfaceModule.cpp
+++ b/Engine/Plugins/Enterprise/DatasmithCADImporter/Source/CoreTechSurface/Private/CoreTechSurfaceModule.cpp
@@ -10,9 +10,6 @@
 
 void FCoreTechSurfaceModule::StartupModule()
 {
-	TArray<FCoreRedirect> Redirects;
-	Redirects.Emplace(ECoreRedirectFlags::Type_Property, TEXT("UCoreTechParametricSurfaceData.RawData"), TEXT("RawData_DEPRECATED"));
-	FCoreRedirects::AddRedirectList(Redirects, CORETECHSURFACE_MODULE_NAME);
 }
 
 FCoreTechSurfaceModule& FCoreTechSurfaceModule::Get()
diff --git a/Engine/Plugins/Enterprise/DatasmithCADImporter/Source/CoreTechSurface/Public/CoreTechSurfaceExtension.h b/Engine/Plugins/Enterprise/DatasmithCADImporter/Source/CoreTechSurface/Public/CoreTechSurfaceExtension.h
--- a/Engine/Plugins/Enterprise/DatasmithCADImporter/Source/CoreTechSurface/Public/CoreTechSurfaceExtension.h
+++ b/Engine/Plugins/Enterprise/DatasmithCADImporter/Source/CoreTechSurface/Public/CoreTechSurfaceExtension.h
@@ -4,7 +4,7 @@
 
 #include "CoreMinimal.h"
 
-#include "ParametricSurfaceExtension.h"
+#include "ParametricSurfaceData.h"
 #include "DatasmithAdditionalData.h"
 #include "DatasmithCustomAction.h"
 #include "DatasmithImportOptions.h"
@@ -13,27 +13,17 @@
 #include "CoreTechSurfaceExtension.generated.h"
 
 UCLASS(meta = (DisplayName = "Kernel IO Parametric Surface Data"))
-class CORETECHSURFACE_API UCoreTechParametricSurfaceData : public UParametricSurfaceData
+class CORETECHSURFACE_API UTempCoreTechParametricSurfaceData : public UParametricSurfaceData
 {
 	GENERATED_BODY()
 
 public:
 	UPROPERTY()
 	FString SourceFile;
-
-	// Too costly to serialize as a UPROPERTY, will use custom serialization.
-	TArray<uint8> RawData;
-
-	virtual bool IsValid() override
-	{
-		return RawData.Num() > 0;
-	}
 	
+	virtual bool SetFile(const TCHAR* SourceFile) override;
 	virtual bool Tessellate(UStaticMesh& StaticMesh, const FDatasmithRetessellationOptions& RetessellateOptions) override;
 
-private:
-	UPROPERTY()
-	TArray<uint8> RawData_DEPRECATED;
-
+protected:
 	virtual void Serialize(FArchive& Ar) override;
 };
diff --git a/Engine/Plugins/Enterprise/DatasmithCADImporter/Source/DatasmithCADTranslator/Private/DatasmithCADTranslator.cpp b/Engine/Plugins/Enterprise/DatasmithCADImporter/Source/DatasmithCADTranslator/Private/DatasmithCADTranslator.cpp
--- a/Engine/Plugins/Enterprise/DatasmithCADImporter/Source/DatasmithCADTranslator/Private/DatasmithCADTranslator.cpp
+++ b/Engine/Plugins/Enterprise/DatasmithCADImporter/Source/DatasmithCADTranslator/Private/DatasmithCADTranslator.cpp
@@ -2,10 +2,10 @@
 
 #include "DatasmithCADTranslator.h"
 
-#include "CADInterfacesModule.h"
 #include "CADFileReader.h"
+#include "CADInterfacesModule.h"
 #include "CADKernelSurfaceExtension.h"
-#include "CoreTechSurfaceHelper.h"
+
 #include "DatasmithCADTranslatorModule.h"
 #include "DatasmithDispatcher.h"
 #include "DatasmithMeshBuilder.h"
@@ -206,15 +206,17 @@ bool FDatasmithCADTranslator::LoadStaticMesh(const TSharedRef<IDatasmithMeshElem
 	if (TOptional< FMeshDescription > Mesh = MeshBuilderPtr->GetMeshDescription(MeshElement, MeshParameters))
 	{
 		OutMeshPayload.LodMeshes.Add(MoveTemp(Mesh.GetValue()));
+
 		if (CADLibrary::FImportParameters::bGDisableCADKernelTessellation)
 		{
-			CoreTechSurface::AddSurfaceDataForMesh(MeshElement->GetFile(), ImportParameters, MeshParameters, GetCommonTessellationOptions(), OutMeshPayload);
+			AddSurfaceData(MeshElement->GetFile(), ImportParameters, MeshParameters, OutMeshPayload);
 		}
 		else
 		{
 			CADKernelSurface::AddSurfaceDataForMesh(MeshElement->GetFile(), ImportParameters, MeshParameters, GetCommonTessellationOptions(), OutMeshPayload);
 		}
 	}
+
 	return OutMeshPayload.LodMeshes.Num() > 0;
 }
 
diff --git a/Engine/Plugins/Enterprise/DatasmithCADImporter/Source/ParametricSurface/ParametricSurface.Build.cs b/Engine/Plugins/Enterprise/DatasmithCADImporter/Source/ParametricSurface/ParametricSurface.Build.cs
--- a/Engine/Plugins/Enterprise/DatasmithCADImporter/Source/ParametricSurface/ParametricSurface.Build.cs
+++ b/Engine/Plugins/Enterprise/DatasmithCADImporter/Source/ParametricSurface/ParametricSurface.Build.cs
@@ -20,9 +20,14 @@ namespace UnrealBuildTool.Rules
 				{
 					"Core",
 					"CoreUObject",
+					"CADInterfaces",
+					"CADLibrary",
 					"CADTools",
-                    "DatasmithContent",
+					"DatasmithContent",
 					"DatasmithTranslator",
+					"Engine",
+					"MeshDescription",
+					"StaticMeshDescription",
 				}
 			);
 		}
diff --git a/Engine/Plugins/Enterprise/DatasmithCADImporter/Source/ParametricSurface/Private/CoreTech/CoreTechSurfaceData.cpp b/Engine/Plugins/Enterprise/DatasmithCADImporter/Source/ParametricSurface/Private/CoreTech/CoreTechSurfaceData.cpp
new file mode 100644
--- /dev/null
+++ b/Engine/Plugins/Enterprise/DatasmithCADImporter/Source/ParametricSurface/Private/CoreTech/CoreTechSurfaceData.cpp
@@ -0,0 +1,144 @@
+// Copyright Epic Games, Inc. All Rights Reserved.
+
+#include "CoreTechSurfaceData.h"
+
+#include "CADInterfacesModule.h"
+#include "CoreTechTypes.h"
+#include "DatasmithAdditionalData.h"
+#include "DatasmithPayload.h"
+#include "IDatasmithSceneElements.h"
+
+#include "Engine/StaticMesh.h"
+#include "HAL/PlatformFileManager.h"
+#include "MeshDescriptionHelper.h"
+#include "Misc/FileHelper.h"
+#include "Misc/Paths.h"
+#include "StaticMeshAttributes.h"
+#include "UObject/EnterpriseObjectVersion.h"
+
+namespace CoreTechParametricSurfaceDataUtils
+{
+	bool Tessellate(uint64 MainObjectId, const CADLibrary::FImportParameters& ImportParams, const CADLibrary::FMeshParameters& MeshParameters, FMeshDescription& MeshDesc)
+	{
+		CADLibrary::CTKIO_SetCoreTechTessellationState(ImportParams);
+
+		CADLibrary::FBodyMesh BodyMesh;
+		BodyMesh.BodyID = 1;
+
+		CADLibrary::CTKIO_GetTessellation(MainObjectId, BodyMesh, false);
+
+		if (BodyMesh.Faces.Num() == 0)
+		{
+			return false;
+		}
+
+		if (!CADLibrary::ConvertBodyMeshToMeshDescription(ImportParams, MeshParameters, BodyMesh, MeshDesc))
+		{
+			ensureMsgf(false, TEXT("Error during mesh conversion"));
+			return false;
+		}
+
+		return true;
+	}
+
+	bool LoadFile(const FString& FileName, const CADLibrary::FImportParameters& ImportParameters, const CADLibrary::FMeshParameters& MeshParameters, FMeshDescription& MeshDescription)
+	{
+		CADLibrary::FCoreTechSessionBase Session(TEXT("CoreTechMeshLoader::LoadFile"));
+		if (!Session.IsCoreTechSessionValid())
+		{
+			return false;
+		}
+
+		CADLibrary::CTKIO_ChangeUnit(ImportParameters.GetMetricUnit());
+		uint64 MainObjectID;
+		if (!CADLibrary::CTKIO_LoadModel(*FileName, MainObjectID, 0x00020000 /* CT_LOAD_FLAGS_READ_META_DATA */))
+		{
+			// Something wrong happened during the load, abort
+			return false;
+		}
+
+		if (ImportParameters.GetStitchingTechnique() != CADLibrary::EStitchingTechnique::StitchingNone)
+		{
+			CADLibrary::CTKIO_Repair(MainObjectID, CADLibrary::EStitchingTechnique::StitchingSew);
+		}
+
+		return Tessellate(MainObjectID, ImportParameters, MeshParameters, MeshDescription);
+	}
+}
+
+bool UCoreTechParametricSurfaceData::SetFile(const TCHAR* FilePath)
+{
+	if (UParametricSurfaceData::SetFile(FilePath))
+	{
+		SourceFile = FilePath;
+		return true;
+	}
+
+	return false;
+}
+
+bool UCoreTechParametricSurfaceData::Tessellate(UStaticMesh& StaticMesh, const FDatasmithRetessellationOptions& RetessellateOptions)
+{
+	bool bSuccessfulTessellation = false;
+
+#if WITH_EDITOR
+	// make a temporary file as CoreTech can only deal with files.
+	int32 Hash = GetTypeHash(StaticMesh.GetPathName());
+	FString ResourceFile = FPaths::ConvertRelativePathToFull(FPaths::ProjectIntermediateDir() / FString::Printf(TEXT("0x%08x.ct"), Hash));
+
+	FFileHelper::SaveArrayToFile(RawData, *ResourceFile);
+
+	CADLibrary::FCTMesh Mesh;
+
+	CADLibrary::FImportParameters ImportParameters(SceneParameters.MetricUnit, SceneParameters.ScaleFactor, (FDatasmithUtils::EModelCoordSystem) SceneParameters.ModelCoordSys);
+	ImportParameters.SetTesselationParameters(RetessellateOptions.ChordTolerance, RetessellateOptions.MaxEdgeLength, RetessellateOptions.NormalTolerance, (CADLibrary::EStitchingTechnique)RetessellateOptions.StitchingTechnique);
+
+	CADLibrary::FMeshParameters CadMeshParameters;
+	CadMeshParameters.bNeedSwapOrientation = MeshParameters.bNeedSwapOrientation;
+	CadMeshParameters.bIsSymmetric = MeshParameters.bIsSymmetric;
+	CadMeshParameters.SymmetricNormal = MeshParameters.SymmetricNormal;
+	CadMeshParameters.SymmetricOrigin = MeshParameters.SymmetricOrigin;
+
+	// Previous MeshDescription is get to be able to create a new one with the same order of PolygonGroup (the matching of color and partition is currently based on their order)
+	if (FMeshDescription* DestinationMeshDescription = StaticMesh.GetMeshDescription(0))
+	{
+		FMeshDescription MeshDescription;
+		FStaticMeshAttributes MeshDescriptionAttributes(MeshDescription);
+		MeshDescriptionAttributes.Register();
+
+		if (RetessellateOptions.RetessellationRule == EDatasmithCADRetessellationRule::SkipDeletedSurfaces)
+		{
+			CADLibrary::CopyPatchGroups(*DestinationMeshDescription, MeshDescription);
+		}
+
+		if (CoreTechParametricSurfaceDataUtils::LoadFile(ResourceFile, ImportParameters, CadMeshParameters, MeshDescription))
+		{
+			// To update the SectionInfoMap 
+			{
+				TPolygonGroupAttributesConstRef<FName> MaterialSlotNames = MeshDescriptionAttributes.GetPolygonGroupMaterialSlotNames();
+				FMeshSectionInfoMap& SectionInfoMap = StaticMesh.GetSectionInfoMap();
+
+				for (FPolygonGroupID PolygonGroupID : MeshDescription.PolygonGroups().GetElementIDs())
+				{
+					FMeshSectionInfo Section = SectionInfoMap.Get(0, PolygonGroupID.GetValue());
+					int32 MaterialIndex = StaticMesh.GetMaterialIndex(MaterialSlotNames[PolygonGroupID]);
+					if (MaterialIndex < 0)
+					{
+						MaterialIndex = 0;
+					}
+					Section.MaterialIndex = MaterialIndex;
+					SectionInfoMap.Set(0, PolygonGroupID.GetValue(), Section);
+				}
+			}
+			*DestinationMeshDescription = MoveTemp(MeshDescription);
+			bSuccessfulTessellation = true;
+		}
+	}
+
+	// Remove temporary file
+	FPlatformFileManager::Get().GetPlatformFile().DeleteFile(*ResourceFile);
+#endif
+	return bSuccessfulTessellation;
+}
+
+
diff --git a/Engine/Plugins/Enterprise/DatasmithCADImporter/Source/ParametricSurface/Private/CoreTech/CoreTechSurfaceData.h b/Engine/Plugins/Enterprise/DatasmithCADImporter/Source/ParametricSurface/Private/CoreTech/CoreTechSurfaceData.h
new file mode 100644
--- /dev/null
+++ b/Engine/Plugins/Enterprise/DatasmithCADImporter/Source/ParametricSurface/Private/CoreTech/CoreTechSurfaceData.h
@@ -0,0 +1,27 @@
+// Copyright Epic Games, Inc. All Rights Reserved.
+
+#pragma once
+
+#include "CoreMinimal.h"
+
+#include "ParametricSurfaceData.h"
+#include "DatasmithAdditionalData.h"
+#include "DatasmithCustomAction.h"
+#include "DatasmithImportOptions.h"
+#include "DatasmithUtils.h"
+
+#include "CoreTechSurfaceData.generated.h"
+
+UCLASS(meta = (DisplayName = "Kernel IO Parametric Surface Data"))
+class PARAMETRICSURFACE_API UCoreTechParametricSurfaceData : public UParametricSurfaceData
+{
+	GENERATED_BODY()
+
+public:
+	UPROPERTY()
+	FString SourceFile;
+
+	virtual bool SetFile(const TCHAR* FilePath) override;
+
+	virtual bool Tessellate(UStaticMesh& StaticMesh, const FDatasmithRetessellationOptions& RetessellateOptions) override;
+};
diff --git a/Engine/Plugins/Enterprise/DatasmithCADImporter/Source/ParametricSurface/Private/ParametricSurfaceData.cpp b/Engine/Plugins/Enterprise/DatasmithCADImporter/Source/ParametricSurface/Private/ParametricSurfaceData.cpp
new file mode 100644
--- /dev/null
+++ b/Engine/Plugins/Enterprise/DatasmithCADImporter/Source/ParametricSurface/Private/ParametricSurfaceData.cpp
@@ -0,0 +1,59 @@
+// Copyright Epic Games, Inc. All Rights Reserved.
+
+#include "ParametricSurfaceData.h"
+
+#include "Misc/FileHelper.h"
+#include "UObject/EnterpriseObjectVersion.h"
+
+#define LOCTEXT_NAMESPACE "ParametricSurfaceData"
+
+bool UParametricSurfaceData::SetFile(const TCHAR* FilePath)
+{
+	if (FPaths::FileExists(FilePath))
+	{
+		TArray<uint8> ByteArray;
+
+		if (FFileHelper::LoadFileToArray(ByteArray, FilePath))
+		{
+			RawData = MoveTemp(ByteArray);
+			return true;
+		}
+	}
+
+	return false;
+}
+
+void UParametricSurfaceData::SetImportParameters(const CADLibrary::FImportParameters& InSceneParameters)
+{
+	SceneParameters.ModelCoordSys = uint8(InSceneParameters.GetModelCoordSys());
+	SceneParameters.MetricUnit = InSceneParameters.GetMetricUnit();
+	SceneParameters.ScaleFactor = InSceneParameters.GetScaleFactor();
+}
+
+void UParametricSurfaceData::SetMeshParameters(const CADLibrary::FMeshParameters& InMeshParameters)
+{
+	MeshParameters.bNeedSwapOrientation = InMeshParameters.bNeedSwapOrientation;
+	MeshParameters.bIsSymmetric = InMeshParameters.bIsSymmetric;
+	MeshParameters.SymmetricNormal = InMeshParameters.SymmetricNormal;
+	MeshParameters.SymmetricOrigin = InMeshParameters.SymmetricOrigin;
+}
+
+void UParametricSurfaceData::Serialize(FArchive& Ar)
+{
+	Ar.UsingCustomVersion(FEnterpriseObjectVersion::GUID);
+
+	Super::Serialize(Ar);
+
+	if (Ar.IsSaving() || (Ar.IsLoading() && Ar.CustomVer(FEnterpriseObjectVersion::GUID) >= FEnterpriseObjectVersion::CoreTechParametricSurfaceOptim))
+	{
+		Ar << RawData;
+	}
+
+	if (RawData_DEPRECATED.Num() && RawData.Num() == 0)
+	{
+		RawData = MoveTemp(RawData_DEPRECATED);
+	}
+}
+
+#undef LOCTEXT_NAMESPACE // "ParametricSurfaceData"
+
diff --git a/Engine/Plugins/Enterprise/DatasmithCADImporter/Source/ParametricSurface/Private/ParametricSurfaceModule.cpp b/Engine/Plugins/Enterprise/DatasmithCADImporter/Source/ParametricSurface/Private/ParametricSurfaceModule.cpp
--- a/Engine/Plugins/Enterprise/DatasmithCADImporter/Source/ParametricSurface/Private/ParametricSurfaceModule.cpp
+++ b/Engine/Plugins/Enterprise/DatasmithCADImporter/Source/ParametricSurface/Private/ParametricSurfaceModule.cpp
@@ -2,6 +2,9 @@
 
 #include "ParametricSurfaceModule.h"
 
+#include "CoreTech/CoreTechSurfaceData.h"
+#include "TechSoft/TechSoftParametricSurface.h"
+
 #include "Modules/ModuleInterface.h"
 #include "Modules/ModuleManager.h"
 #include "UObject/CoreRedirects.h"
@@ -10,9 +13,15 @@
 
 void FParametricSurfaceModule::StartupModule()
 {
-	TArray<FCoreRedirect> Redirects;
-	Redirects.Emplace(ECoreRedirectFlags::Type_Property, TEXT("UParametricSurfaceData.RawData"), TEXT("RawData_DEPRECATED"));
-	FCoreRedirects::AddRedirectList(Redirects, PARAMETRICSURFACE_MODULE_NAME);
+	{
+		TArray<FCoreRedirect> Redirects;
+		Redirects.Emplace(ECoreRedirectFlags::Type_Package, TEXT("/Script/DatasmithCoreTechParametricSurfaceData"), TEXT("/Script/ParametricSurface"));
+		Redirects.Emplace(ECoreRedirectFlags::Type_Property, TEXT("UParametricSurfaceData.RawData"), TEXT("RawData_DEPRECATED"));
+		Redirects.Emplace(ECoreRedirectFlags::Type_Property, TEXT("UCoreTechParametricSurfaceData.RawData"), TEXT("RawData_DEPRECATED"));
+		Redirects.Emplace(ECoreRedirectFlags::Type_Struct, TEXT("CoreTechSceneParameters"), TEXT("ParametricSceneParameters"));
+		Redirects.Emplace(ECoreRedirectFlags::Type_Struct, TEXT("CoreTechMeshParameters"), TEXT("ParametricMeshParameters"));
+		FCoreRedirects::AddRedirectList(Redirects, PARAMETRICSURFACE_MODULE_NAME);
+	}
 }
 
 FParametricSurfaceModule& FParametricSurfaceModule::Get()
@@ -25,7 +34,22 @@ bool FParametricSurfaceModule::IsAvailable()
 	return FModuleManager::Get().IsModuleLoaded(PARAMETRICSURFACE_MODULE_NAME);
 }
 
-IMPLEMENT_MODULE(FParametricSurfaceModule, ParametricSurface);
+UParametricSurfaceData* FParametricSurfaceModule::CreateParametricSurface(const TCHAR* CADLibraryName)
+{
+	if (!FCString::Strcmp(TEXT("TechSoft"), CADLibraryName))
+	{
+		return Datasmith::MakeAdditionalData<UTechSoftParametricSurfaceData>();
+	}
+
+	if (!FCString::Strcmp(TEXT("KernelIO"), CADLibraryName))
+	{
+		return Datasmith::MakeAdditionalData<UCoreTechParametricSurfaceData>();
+	}
+
+	return nullptr;
+}
+
+IMPLEMENT_MODULE(FParametricSurfaceModule, ParametricSurface)
 
 #undef LOCTEXT_NAMESPACE // "ParametricSurfaceModule"
 
diff --git a/Engine/Plugins/Enterprise/DatasmithCADImporter/Source/ParametricSurface/Private/ParametricSurfaceTranslator.cpp b/Engine/Plugins/Enterprise/DatasmithCADImporter/Source/ParametricSurface/Private/ParametricSurfaceTranslator.cpp
--- a/Engine/Plugins/Enterprise/DatasmithCADImporter/Source/ParametricSurface/Private/ParametricSurfaceTranslator.cpp
+++ b/Engine/Plugins/Enterprise/DatasmithCADImporter/Source/ParametricSurface/Private/ParametricSurfaceTranslator.cpp
@@ -2,7 +2,13 @@
 
 #include "ParametricSurfaceTranslator.h"
 
+#include "ParametricSurfaceData.h"
+#include "ParametricSurfaceModule.h"
+
 #include "DatasmithImportOptions.h"
+#include "IDatasmithSceneElements.h"
+
+#include "Misc/FileHelper.h"
 
 void FParametricSurfaceTranslator::GetSceneImportOptions(TArray<TStrongObjectPtr<UDatasmithOptionsBase>>& Options)
 {
@@ -30,3 +36,31 @@ void FParametricSurfaceTranslator::SetSceneImportOptions(TArray<TStrongObjectPtr
 	}
 }
 
+bool FParametricSurfaceTranslator::AddSurfaceData(const TCHAR* MeshFilePath, const CADLibrary::FImportParameters& ImportParameters, const CADLibrary::FMeshParameters& MeshParameters, FDatasmithMeshElementPayload& OutMeshPayload)
+{
+	if (MeshFilePath && IFileManager::Get().FileExists(MeshFilePath))
+	{
+		UParametricSurfaceData* ParametricSurfaceData = FParametricSurfaceModule::CreateParametricSurface(*CADLibrary::FImportParameters::GCADLibrary);
+
+		if (!ParametricSurfaceData || !ParametricSurfaceData->SetFile(MeshFilePath))
+		{
+			return false;
+		}
+
+		ParametricSurfaceData->SetImportParameters(ImportParameters);
+		ParametricSurfaceData->SetMeshParameters(MeshParameters);
+		ParametricSurfaceData->SetLastTessellationOptions(CommonTessellationOptions);
+
+		OutMeshPayload.AdditionalData.Add(ParametricSurfaceData);
+
+		// Remove the file because it is temporary since caching is disabled.
+		if (!CADLibrary::FImportParameters::bGEnableCADCache)
+		{
+			IFileManager::Get().Delete(MeshFilePath);
+		}
+
+		return true;
+	}
+
+	return false;
+}
diff --git a/Engine/Plugins/Enterprise/DatasmithCADImporter/Source/ParametricSurface/Private/TechSoft/TechSoftParametricSurface.cpp b/Engine/Plugins/Enterprise/DatasmithCADImporter/Source/ParametricSurface/Private/TechSoft/TechSoftParametricSurface.cpp
new file mode 100644
--- /dev/null
+++ b/Engine/Plugins/Enterprise/DatasmithCADImporter/Source/ParametricSurface/Private/TechSoft/TechSoftParametricSurface.cpp
@@ -0,0 +1,95 @@
+// Copyright Epic Games, Inc. All Rights Reserved.
+
+#include "TechSoftParametricSurface.h"
+
+#include "TechSoftInterface.h"
+
+#include "CADInterfacesModule.h"
+#include "DatasmithAdditionalData.h"
+#include "DatasmithPayload.h"
+
+#include "Engine/StaticMesh.h"
+#include "HAL/PlatformFileManager.h"
+#include "IDatasmithSceneElements.h"
+#include "MeshDescriptionHelper.h"
+#include "Misc/FileHelper.h"
+#include "Misc/Paths.h"
+#include "StaticMeshAttributes.h"
+
+bool UTechSoftParametricSurfaceData::Tessellate(UStaticMesh& StaticMesh, const FDatasmithRetessellationOptions& RetessellateOptions)
+{
+	using namespace CADLibrary;
+
+	bool bSuccessfulTessellation = false;
+
+#if WITH_EDITOR
+	// make a temporary file as TechSoft can only deal with files.
+	int32 Hash = GetTypeHash(StaticMesh.GetPathName());
+	FString ResourceFile = FPaths::ConvertRelativePathToFull(FPaths::ProjectIntermediateDir() / FString::Printf(TEXT("0x%08x.hsf"), Hash));
+
+	FFileHelper::SaveArrayToFile(RawData, *ResourceFile);
+
+	// Previous MeshDescription is get to be able to create a new one with the same order of PolygonGroup (the matching of color and partition is currently based on their order)
+	if (FMeshDescription* DestinationMeshDescription = StaticMesh.GetMeshDescription(0))
+	{
+		FMeshDescription MeshDescription;
+		FStaticMeshAttributes MeshDescriptionAttributes(MeshDescription);
+		MeshDescriptionAttributes.Register();
+
+		if (RetessellateOptions.RetessellationRule == EDatasmithCADRetessellationRule::SkipDeletedSurfaces)
+		{
+			CADLibrary::CopyPatchGroups(*DestinationMeshDescription, MeshDescription);
+		}
+
+		CADLibrary::FImportParameters ImportParameters(SceneParameters.MetricUnit, SceneParameters.ScaleFactor, (FDatasmithUtils::EModelCoordSystem)SceneParameters.ModelCoordSys);
+		ImportParameters.SetTesselationParameters(RetessellateOptions.ChordTolerance, RetessellateOptions.MaxEdgeLength, RetessellateOptions.NormalTolerance, (CADLibrary::EStitchingTechnique)RetessellateOptions.StitchingTechnique);
+
+		FTechSoftInterface& TechSoftInterface = TechSoftUtils::GetTechSoftInterface();
+
+		bSuccessfulTessellation = TechSoftInterface.InitializeKernel(*FPaths::EnginePluginsDir());
+		if (bSuccessfulTessellation)
+		{
+			CADLibrary::FBodyMesh BodyMesh;
+
+			bSuccessfulTessellation = TechSoftInterface.GetBodyFromHsfFile(ResourceFile, ImportParameters, BodyMesh);
+
+			if (bSuccessfulTessellation)
+			{
+				bSuccessfulTessellation = CADLibrary::ConvertBodyMeshToMeshDescription(ImportParameters, MeshParameters, BodyMesh, MeshDescription);
+
+				if (bSuccessfulTessellation)
+				{
+					TPolygonGroupAttributesConstRef<FName> MaterialSlotNames = MeshDescriptionAttributes.GetPolygonGroupMaterialSlotNames();
+					FMeshSectionInfoMap& SectionInfoMap = StaticMesh.GetSectionInfoMap();
+
+					for (FPolygonGroupID PolygonGroupID : MeshDescription.PolygonGroups().GetElementIDs())
+					{
+						FMeshSectionInfo Section = SectionInfoMap.Get(0, PolygonGroupID.GetValue());
+						int32 MaterialIndex = StaticMesh.GetMaterialIndex(MaterialSlotNames[PolygonGroupID]);
+						if (MaterialIndex < 0)
+						{
+							MaterialIndex = 0;
+						}
+						Section.MaterialIndex = MaterialIndex;
+						SectionInfoMap.Set(0, PolygonGroupID.GetValue(), Section);
+					}
+
+					*DestinationMeshDescription = MoveTemp(MeshDescription);
+				}
+			}
+		}
+	}
+
+	// Remove temporary file
+	FPlatformFileManager::Get().GetPlatformFile().DeleteFile(*ResourceFile);
+
+	if (!bSuccessfulTessellation)
+	{
+		ensureMsgf(false, TEXT("Error during mesh conversion"));
+		return false;
+	}
+#endif
+	return bSuccessfulTessellation;
+}
+
+
diff --git a/Engine/Plugins/Enterprise/DatasmithCADImporter/Source/ParametricSurface/Private/TechSoft/TechSoftParametricSurface.h b/Engine/Plugins/Enterprise/DatasmithCADImporter/Source/ParametricSurface/Private/TechSoft/TechSoftParametricSurface.h
new file mode 100644
--- /dev/null
+++ b/Engine/Plugins/Enterprise/DatasmithCADImporter/Source/ParametricSurface/Private/TechSoft/TechSoftParametricSurface.h
@@ -0,0 +1,22 @@
+// Copyright Epic Games, Inc. All Rights Reserved.
+
+#pragma once
+
+#include "CoreMinimal.h"
+
+#include "ParametricSurfaceData.h"
+#include "DatasmithAdditionalData.h"
+#include "DatasmithCustomAction.h"
+#include "DatasmithImportOptions.h"
+#include "DatasmithUtils.h"
+
+#include "TechSoftParametricSurface.generated.h"
+
+UCLASS(meta = (DisplayName = "TechSoft Parametric Surface Data"))
+class PARAMETRICSURFACE_API UTechSoftParametricSurfaceData : public UParametricSurfaceData
+{
+	GENERATED_BODY()
+
+public:
+	virtual bool Tessellate(UStaticMesh& StaticMesh, const FDatasmithRetessellationOptions& RetessellateOptions) override;
+};
diff --git a/Engine/Plugins/Enterprise/DatasmithCADImporter/Source/ParametricSurface/Public/ParametricSurfaceData.h b/Engine/Plugins/Enterprise/DatasmithCADImporter/Source/ParametricSurface/Public/ParametricSurfaceData.h
new file mode 100644
--- /dev/null
+++ b/Engine/Plugins/Enterprise/DatasmithCADImporter/Source/ParametricSurface/Public/ParametricSurfaceData.h
@@ -0,0 +1,136 @@
+// Copyright Epic Games, Inc. All Rights Reserved.
+
+#pragma once
+
+#include "CoreMinimal.h"
+
+#include "CADOptions.h"
+
+#include "DatasmithAdditionalData.h"
+#include "DatasmithCustomAction.h"
+#include "DatasmithImportOptions.h"
+#include "DatasmithUtils.h"
+
+#include "ParametricSurfaceData.generated.h"
+
+
+USTRUCT(BlueprintType)
+struct PARAMETRICSURFACE_API FParametricSceneParameters
+{
+	GENERATED_BODY()
+
+	// value from FDatasmithUtils::EModelCoordSystem
+	UPROPERTY()
+	uint8 ModelCoordSys = (uint8)FDatasmithUtils::EModelCoordSystem::ZUp_LeftHanded;
+
+	UPROPERTY()
+	float MetricUnit = 0.01f;
+
+	UPROPERTY()
+	float ScaleFactor = 1.0f;
+};
+
+USTRUCT()
+struct PARAMETRICSURFACE_API FParametricMeshParameters
+{
+	GENERATED_BODY()
+
+	UPROPERTY()
+	bool bNeedSwapOrientation = false;
+
+	UPROPERTY()
+	bool bIsSymmetric = false;
+
+	UPROPERTY()
+	FVector SymmetricOrigin = FVector::ZeroVector;
+
+	UPROPERTY()
+	FVector SymmetricNormal = FVector::ZeroVector;
+
+	operator CADLibrary::FMeshParameters()
+	{
+		CADLibrary::FMeshParameters Parameters;
+
+		Parameters.bNeedSwapOrientation = bNeedSwapOrientation;
+		Parameters.bIsSymmetric = bIsSymmetric;
+		Parameters.SymmetricNormal = SymmetricNormal;
+		Parameters.SymmetricOrigin = SymmetricOrigin;
+
+		return Parameters;
+	}
+};
+
+UCLASS(meta = (DisplayName = "Parametric Surface Data"))
+class PARAMETRICSURFACE_API UParametricSurfaceData : public UDatasmithAdditionalData
+{
+	GENERATED_BODY()
+
+public:
+
+	virtual bool IsValid()
+	{
+		return RawData.Num() > 0;
+	}
+
+	virtual bool SetFile(const TCHAR* FilePath);
+
+	virtual void SetImportParameters(const CADLibrary::FImportParameters& InSceneParameters);
+
+	virtual void SetMeshParameters(const CADLibrary::FMeshParameters& InMeshParameters);
+
+	virtual const FDatasmithTessellationOptions& GetLastTessellationOptions() const
+	{
+		return LastTessellationOptions;
+	}
+
+	virtual void SetLastTessellationOptions(const FDatasmithTessellationOptions& InTessellationOptions)
+	{
+		LastTessellationOptions = InTessellationOptions;
+	}
+
+	virtual bool Tessellate(UStaticMesh& StaticMesh, const FDatasmithRetessellationOptions& RetessellateOptions)
+	{
+		return false;
+	}
+
+protected:
+	virtual void Serialize(FArchive& Ar) override;
+
+protected:
+	UPROPERTY()
+	FParametricSceneParameters SceneParameters;
+
+	UPROPERTY()
+	FParametricMeshParameters MeshParameters;
+
+	UPROPERTY(EditAnywhere, Category = NURBS)
+	FDatasmithTessellationOptions LastTessellationOptions;
+
+	UPROPERTY()
+	TArray<uint8> RawData_DEPRECATED;
+
+	// Too costly to serialize as a UPROPERTY, will use custom serialization.
+	TArray<uint8> RawData;
+};
+
+//void AddSurfaceDataForMesh(const TCHAR* InFilePath, const CADLibrary::FImportParameters& InSceneParameters, const CADLibrary::FMeshParameters& InMeshParameters, , FDatasmithMeshElementPayload& OutMeshPayload)
+//{
+//	if (ICADInterfacesModule::GetAvailability() == ECADInterfaceAvailability::Available)
+//	{
+//		if (FPaths::FileExists(InFilePath))
+//		{
+//			TArray<uint8> ByteArray;
+//			if (FFileHelper::LoadFileToArray(ByteArray, InFilePath))
+//			{
+//				UCoreTechParametricSurfaceData* CoreTechData = Datasmith::MakeAdditionalData<UCoreTechParametricSurfaceData>();
+//				CoreTechData->SourceFile = InFilePath;
+//				CoreTechData->RawData = MoveTemp(ByteArray);
+
+
+//				CoreTechData->
+//				OutMeshPayload.AdditionalData.Add(CoreTechData);
+//			}
+//		}
+//	}
+//}
+
diff --git a/Engine/Plugins/Enterprise/DatasmithCADImporter/Source/ParametricSurface/Public/ParametricSurfaceExtension.h b/Engine/Plugins/Enterprise/DatasmithCADImporter/Source/ParametricSurface/Public/ParametricSurfaceExtension.h
deleted file mode 100644
diff --git a/Engine/Plugins/Enterprise/DatasmithCADImporter/Source/ParametricSurface/Public/ParametricSurfaceModule.h b/Engine/Plugins/Enterprise/DatasmithCADImporter/Source/ParametricSurface/Public/ParametricSurfaceModule.h
--- a/Engine/Plugins/Enterprise/DatasmithCADImporter/Source/ParametricSurface/Public/ParametricSurfaceModule.h
+++ b/Engine/Plugins/Enterprise/DatasmithCADImporter/Source/ParametricSurface/Public/ParametricSurfaceModule.h
@@ -16,4 +16,6 @@ public:
 
 	static FParametricSurfaceModule& Get();
 	static bool IsAvailable();
+
+	static PARAMETRICSURFACE_API class UParametricSurfaceData* CreateParametricSurface(const TCHAR* CADLibraryName);
 };
diff --git a/Engine/Plugins/Enterprise/DatasmithCADImporter/Source/ParametricSurface/Public/ParametricSurfaceTranslator.h b/Engine/Plugins/Enterprise/DatasmithCADImporter/Source/ParametricSurface/Public/ParametricSurfaceTranslator.h
--- a/Engine/Plugins/Enterprise/DatasmithCADImporter/Source/ParametricSurface/Public/ParametricSurfaceTranslator.h
+++ b/Engine/Plugins/Enterprise/DatasmithCADImporter/Source/ParametricSurface/Public/ParametricSurfaceTranslator.h
@@ -10,6 +10,11 @@
 
 class UDatasmithCommonTessellationOptions;
 
+namespace CADLibrary
+{
+	class FImportParameters;
+	struct FMeshParameters;
+}
 
 class PARAMETRICSURFACE_API FParametricSurfaceTranslator : public IDatasmithTranslator
 {
@@ -20,6 +25,8 @@ public:
 	virtual void SetSceneImportOptions(TArray<TStrongObjectPtr<UDatasmithOptionsBase>>& Options) override;
 	// End IDatasmithTranslator overrides
 
+	bool AddSurfaceData(const TCHAR* MeshFilePath, const CADLibrary::FImportParameters& InSceneParameters, const CADLibrary::FMeshParameters& MeshParameters, FDatasmithMeshElementPayload& OutMeshPayload);
+
 protected:
 	const FDatasmithTessellationOptions& GetCommonTessellationOptions()
 	{
diff --git a/Engine/Plugins/Enterprise/DatasmithCADImporter/Source/ParametricSurfaceExtension/Private/ParametricRetessellateAction.cpp b/Engine/Plugins/Enterprise/DatasmithCADImporter/Source/ParametricSurfaceExtension/Private/ParametricRetessellateAction.cpp
--- a/Engine/Plugins/Enterprise/DatasmithCADImporter/Source/ParametricSurfaceExtension/Private/ParametricRetessellateAction.cpp
+++ b/Engine/Plugins/Enterprise/DatasmithCADImporter/Source/ParametricSurfaceExtension/Private/ParametricRetessellateAction.cpp
@@ -103,7 +103,7 @@ void FParametricRetessellateAction_Impl::ApplyOnAssets(const TArray<FAssetData>&
 						Parameters.CancelButtonTooltip = LOCTEXT("OptionWindow_CancelButtonTooltip", "Cancel the retessellation operation");
 
 						bAskForSameOption = false; // ask only the fist time
-						RetessellateOptions->Options = ParametricSurfaceData->LastTessellationOptions;
+						RetessellateOptions->Options = ParametricSurfaceData->GetLastTessellationOptions();
 						Datasmith::FDisplayResult Result = Datasmith::DisplayOptions(RetessellateOptions, Parameters);
 						if (!Result.bValidated)
 						{
@@ -111,7 +111,7 @@ void FParametricRetessellateAction_Impl::ApplyOnAssets(const TArray<FAssetData>&
 						}
 						bSameOptionsForAll |= Result.bUseSameOption;
 					}
-					ParametricSurfaceData->LastTessellationOptions = RetessellateOptions->Options;
+					ParametricSurfaceData->SetLastTessellationOptions(RetessellateOptions->Options);
 
 					int32 RemainingAssetsToProcess = NumAssetsToProcess - AssetIndex;
 					if (bSameOptionsForAll && !Progress.IsValid() && RemainingAssetsToProcess > 1)
diff --git a/Engine/Plugins/Enterprise/DatasmithCADImporter/Source/ParametricSurfaceExtension/Private/ParametricSurfaceBlueprintLibrary.cpp b/Engine/Plugins/Enterprise/DatasmithCADImporter/Source/ParametricSurfaceExtension/Private/ParametricSurfaceBlueprintLibrary.cpp
--- a/Engine/Plugins/Enterprise/DatasmithCADImporter/Source/ParametricSurfaceExtension/Private/ParametricSurfaceBlueprintLibrary.cpp
+++ b/Engine/Plugins/Enterprise/DatasmithCADImporter/Source/ParametricSurfaceExtension/Private/ParametricSurfaceBlueprintLibrary.cpp
@@ -81,7 +81,7 @@ bool UParametricSurfaceBlueprintLibrary::RetessellateStaticMeshWithNotification(
 				}
 
 				// Save last tessellation settings
-				ParametricSurfaceData->LastTessellationOptions = TessellationSettings;
+				ParametricSurfaceData->SetLastTessellationOptions(TessellationSettings);
 
 				bTessellationOutcome = true;
 			}
diff --git a/Engine/Plugins/Enterprise/DatasmithCADImporter/Source/ParametricSurfaceExtension/Public/ParametricRetessellateAction.h b/Engine/Plugins/Enterprise/DatasmithCADImporter/Source/ParametricSurfaceExtension/Public/ParametricRetessellateAction.h
--- a/Engine/Plugins/Enterprise/DatasmithCADImporter/Source/ParametricSurfaceExtension/Public/ParametricRetessellateAction.h
+++ b/Engine/Plugins/Enterprise/DatasmithCADImporter/Source/ParametricSurfaceExtension/Public/ParametricRetessellateAction.h
@@ -4,7 +4,7 @@
 
 #include "CoreMinimal.h"
 
-#include "ParametricSurfaceExtension.h"
+#include "ParametricSurfaceData.h"
 #include "DatasmithCustomAction.h"
 #include "DatasmithAdditionalData.h"
 #include "DatasmithImportOptions.h"
