Add CADLibrary in FImportParameters hash

diff --git a/Engine/Plugins/Enterprise/DatasmithCADImporter/Source/CADTools/Private/CADOption.cpp b/Engine/Plugins/Enterprise/DatasmithCADImporter/Source/CADTools/Private/CADOption.cpp
--- a/Engine/Plugins/Enterprise/DatasmithCADImporter/Source/CADTools/Private/CADOption.cpp
+++ b/Engine/Plugins/Enterprise/DatasmithCADImporter/Source/CADTools/Private/CADOption.cpp
@@ -64,6 +64,7 @@ FAutoConsoleVariableRef GCADTranslatorJtFileEmbeddedTessellation(
 uint32 GetTypeHash(const FImportParameters& ImportParameters)
 {
 	uint32 ParametersHash = ::GetTypeHash(ImportParameters.bGDisableCADKernelTessellation);
+	ParametersHash = HashCombine(ParametersHash, ::GetTypeHash(ImportParameters.GCADLibrary));
 	ParametersHash = HashCombine(ParametersHash, ::GetTypeHash(ImportParameters.ChordTolerance));
 	ParametersHash = HashCombine(ParametersHash, ::GetTypeHash(ImportParameters.MaxEdgeLength));
 	ParametersHash = HashCombine(ParametersHash, ::GetTypeHash(ImportParameters.MaxNormalAngle));
