Fix MeshOrientation
A triangle is defined by its 3 vertices: v0 v1 v2 and 3 edges: e0(v0v1), e1 (v1v2), e2(v2v0)
Swapping triangle orientation is done be swapping:
- 2 vertex v0 and v1 => new triangle is  (v1 v0 v2)
- 2 edge but thy are e1 and e2 that need to be swapped and not e0 and e1, so  (e0(v0v1), e2(v2v0), e1 (v1v2)

diff --git a/Engine/Plugins/Enterprise/DatasmithCADImporter/Source/CADLibrary/Private/MeshOperator.cpp b/Engine/Plugins/Enterprise/DatasmithCADImporter/Source/CADLibrary/Private/MeshOperator.cpp
--- a/Engine/Plugins/Enterprise/DatasmithCADImporter/Source/CADLibrary/Private/MeshOperator.cpp
+++ b/Engine/Plugins/Enterprise/DatasmithCADImporter/Source/CADLibrary/Private/MeshOperator.cpp
@@ -80,14 +80,16 @@ bool MeshOperator::OrientMesh(FMeshDescription& MeshDescription)
 					NbBorderEdges++;
 					continue;
 				}
-				NbSurfaceEdges++;
 
 				AdjacentTriangle = MeshWrapper.GetOtherTriangleAtEdge(Edge, Triangle);
 				if (MeshWrapper.IsTriangleMarked(AdjacentTriangle))
 				{
 					continue;
 				}
+
+				NbSurfaceEdges++;
 				NbConnectedFaces++;
+
 				ConnectedTriangles.Add(AdjacentTriangle);
 
 				if(MeshWrapper.GetEdgeDirectionInTriangle(Edge, 0) == MeshWrapper.GetEdgeDirectionInTriangle(Edge, 1))
@@ -105,7 +107,7 @@ bool MeshOperator::OrientMesh(FMeshDescription& MeshDescription)
 		// Check if the mesh orientation need to be swapped
 		int NbInverted = 0;
 		int NbNotInverted = 0;
-		if (NbBorderEdges == 0 || NbBorderEdges * 40 < NbSurfaceEdges)  // NbBorderEdges * 40 < nbSurfaceEdges => basic rule to define if a mesh is a surface mesh or if the mesh is a volume mesh with gaps 
+		if (NbBorderEdges == 0 || NbBorderEdges * 20 < NbSurfaceEdges)  // NbBorderEdges * 20 < nbSurfaceEdges => basic rule to define if a mesh is a surface mesh or if the mesh is a volume mesh with gaps 
 		{
 			// case of volume mesh
 			// A vertex can have many normal (one by vertex instance) e.g. the corner of a box that have 3 normal and could be the highest 
diff --git a/Engine/Source/Runtime/MeshDescription/Private/MeshDescription.cpp b/Engine/Source/Runtime/MeshDescription/Private/MeshDescription.cpp
--- a/Engine/Source/Runtime/MeshDescription/Private/MeshDescription.cpp
+++ b/Engine/Source/Runtime/MeshDescription/Private/MeshDescription.cpp
@@ -1886,7 +1886,7 @@ void FMeshDescription::ReverseTriangleFacing(const FTriangleID TriangleID)
 	TArrayView<FEdgeID> TriEdges = TriangleEdges[TriangleID];
 	Swap(TriVertexInstances[0], TriVertexInstances[1]);
 	Swap(TriVertices[0], TriVertices[1]);
-	Swap(TriEdges[0], TriEdges[1]);
+	Swap(TriEdges[1], TriEdges[2]);
 }
 
 
