Use TechSoft to tessellate BReps

diff --git a/Engine/Plugins/Enterprise/DatasmithCADImporter/Source/CADInterfaces/Private/CADInterfacesModule.cpp b/Engine/Plugins/Enterprise/DatasmithCADImporter/Source/CADInterfaces/Private/CADInterfacesModule.cpp
--- a/Engine/Plugins/Enterprise/DatasmithCADImporter/Source/CADInterfaces/Private/CADInterfacesModule.cpp
+++ b/Engine/Plugins/Enterprise/DatasmithCADImporter/Source/CADInterfaces/Private/CADInterfacesModule.cpp
@@ -45,7 +45,7 @@ ECADInterfaceAvailability ICADInterfacesModule::GetAvailability()
 		}
 		else if (CADLibrary::FImportParameters::GCADLibrary == TEXT("TechSoft"))
 		{
-			if (CADLibrary::TechSoftUtils::TECHSOFT_InitializeKernel())
+			if (CADLibrary::TechSoftInterface::TECHSOFT_InitializeKernel())
 			{
 				return ECADInterfaceAvailability::Available;
 			}
diff --git a/Engine/Plugins/Enterprise/DatasmithCADImporter/Source/CADInterfaces/Private/TUniqueTechSoftObj.h b/Engine/Plugins/Enterprise/DatasmithCADImporter/Source/CADInterfaces/Private/TUniqueTechSoftObj.h
deleted file mode 100644
diff --git a/Engine/Plugins/Enterprise/DatasmithCADImporter/Source/CADInterfaces/Private/TechSoftBridge.cpp b/Engine/Plugins/Enterprise/DatasmithCADImporter/Source/CADInterfaces/Private/TechSoftBridge.cpp
--- a/Engine/Plugins/Enterprise/DatasmithCADImporter/Source/CADInterfaces/Private/TechSoftBridge.cpp
+++ b/Engine/Plugins/Enterprise/DatasmithCADImporter/Source/CADInterfaces/Private/TechSoftBridge.cpp
@@ -60,13 +60,13 @@ const FUVReparameterization FUVReparameterization::Identity(1., 0, 1., 0);
 template<typename... InArgTypes>
 A3DStatus GetCurveAsNurbs(const A3DCrvBase* A3DCurve, A3DCrvNurbsData* DataPtr, InArgTypes&&... Args)
 {
-	return CADLibrary::TechSoftUtils::GetCurveAsNurbs(A3DCurve, DataPtr, Forward<InArgTypes>(Args)...);
+	return CADLibrary::TechSoftInterface::GetCurveAsNurbs(A3DCurve, DataPtr, Forward<InArgTypes>(Args)...);
 };
 
 template<typename... InArgTypes>
 A3DStatus GetSurfaceAsNurbs(const A3DSurfBase* A3DSurface, A3DSurfNurbsData* DataPtr, InArgTypes&&... Args)
 {
-	return CADLibrary::TechSoftUtils::GetSurfaceAsNurbs(A3DSurface, DataPtr, Forward<InArgTypes>(Args)...);
+	return CADLibrary::TechSoftInterface::GetSurfaceAsNurbs(A3DSurface, DataPtr, Forward<InArgTypes>(Args)...);
 }; 
 
 FString GetFileName(const FString& InFilePath)
@@ -366,11 +366,11 @@ int TraverseBaseWithGraphics(const ECADType FileType, const A3DEntity* Entity, F
 
 int TraverseSource(const A3DEntity* Entity, FEntityData& OutEntityData, bool bIsOccurrence, const ECADType FileType)
 {
-	if (CADLibrary::TechSoftUtils::IsEntityBaseWithGraphicsType(Entity))
+	if (CADLibrary::TechSoftInterface::IsEntityBaseWithGraphicsType(Entity))
 	{
 		return TraverseBaseWithGraphics(FileType, Entity, OutEntityData, bIsOccurrence);
 	}
-	else if (CADLibrary::TechSoftUtils::IsEntityBaseType(Entity))
+	else if (CADLibrary::TechSoftInterface::IsEntityBaseType(Entity))
 	{
 		return TraverseBase(FileType, Entity, OutEntityData, bIsOccurrence);
 	}
@@ -588,7 +588,7 @@ FMatrixH TraverseTransformation(const A3DMiscCartesianTransformation* Transforma
 	TraverseSource(Transformation3D, MetaData, false);
 
 	A3DEEntityType Type = kA3DTypeUnknown;
-	CADLibrary::TechSoftUtils::GetEntityType(Transformation3D, &Type);
+	CADLibrary::TechSoftInterface::GetEntityType(Transformation3D, &Type);
 
 	if (Type == kA3DTypeMiscCartesianTransformation)
 	{
@@ -814,8 +814,8 @@ TSharedPtr<FCurve> FTechSoftBridge::AddCurve(const A3DCrvBase* A3DCurve, const T
 {
 	TSharedPtr<FCurve> Curve = TSharedPtr<FCurve>();
 	A3DEEntityType eType;
-	A3DInt32 iRet = CADLibrary::TechSoftUtils::GetEntityType(A3DCurve, &eType);
-	if (iRet == A3D_SUCCESS)
+	A3DInt32 Ret = CADLibrary::TechSoftInterface::GetEntityType(A3DCurve, &eType);
+	if (Ret == A3D_SUCCESS)
 	{
 		Report.CurveCount++;
 
@@ -1391,7 +1391,7 @@ TSharedPtr<FSurface> FTechSoftBridge::AddSurface(const A3DSurfBase* A3DSurface,
 	TechSoftUtils::TraverseSource(A3DSurface, MetaData);
 
 	A3DEEntityType Type;
-	int32 Ret = CADLibrary::TechSoftUtils::GetEntityType(A3DSurface, &Type);
+	int32 Ret = CADLibrary::TechSoftInterface::GetEntityType(A3DSurface, &Type);
 	if (Ret == A3D_SUCCESS)
 	{
 		switch (Type)
diff --git a/Engine/Plugins/Enterprise/DatasmithCADImporter/Source/CADInterfaces/Private/TechSoftBridge.h b/Engine/Plugins/Enterprise/DatasmithCADImporter/Source/CADInterfaces/Private/TechSoftBridge.h
--- a/Engine/Plugins/Enterprise/DatasmithCADImporter/Source/CADInterfaces/Private/TechSoftBridge.h
+++ b/Engine/Plugins/Enterprise/DatasmithCADImporter/Source/CADInterfaces/Private/TechSoftBridge.h
@@ -10,21 +10,7 @@
 #include "CADKernel/Geo/GeoEnum.h"
 #include "CADKernel/Math/MatrixH.h"
 
-#if PLATFORM_WINDOWS
-#include "Windows/AllowWindowsPlatformTypes.h"
-#endif
-
-THIRD_PARTY_INCLUDES_START
-#pragma push_macro("TEXT")
-#include "A3DSDKEnums.h"
-#include "A3DSDKGeometry.h"
-#include "A3DSDKGeometrySrf.h"
-#pragma pop_macro("TEXT")
-THIRD_PARTY_INCLUDES_END
-
-#if PLATFORM_WINDOWS
-#include "Windows/HideWindowsPlatformTypes.h"
-#endif
+#include "TechSoftInterface.h"
 
 typedef void A3DAsmProductOccurrence;
 typedef void A3DCrvBase;
diff --git a/Engine/Plugins/Enterprise/DatasmithCADImporter/Source/CADInterfaces/Private/TechSoftFileParser.cpp b/Engine/Plugins/Enterprise/DatasmithCADImporter/Source/CADInterfaces/Private/TechSoftFileParser.cpp
--- a/Engine/Plugins/Enterprise/DatasmithCADImporter/Source/CADInterfaces/Private/TechSoftFileParser.cpp
+++ b/Engine/Plugins/Enterprise/DatasmithCADImporter/Source/CADInterfaces/Private/TechSoftFileParser.cpp
@@ -7,6 +7,7 @@
 #ifdef USE_TECHSOFT_SDK
 
 #include "TechSoftInterface.h"
+#include "TechSoftUtils.h"
 
 #include "Serialization/JsonSerializer.h"
 #include "Serialization/JsonWriter.h"
@@ -278,7 +279,7 @@ void UpdateIOOptionAccordingToFormat(const CADLibrary::ECADFormat Format, A3DImp
 
 FTechSoftFileParser::FTechSoftFileParser(FCADFileData& InCADData, const FString& EnginePluginsPath)
 	: CADFileData(InCADData)
-	, TechSoftInterface(TechSoftUtils::GetTechSoftInterface())
+	, TechSoftInterface(FTechSoftInterface::Get())
 {
 	TechSoftInterface.InitializeKernel(*EnginePluginsPath);
 }
@@ -300,8 +301,9 @@ ECADParsingResult FTechSoftFileParser::Process()
 	Format = File.GetFileFormat();
 	TechSoftFileParserImpl::UpdateIOOptionAccordingToFormat(Format, Import);
 
-	A3DStatus IRet = TechSoftInterface.Import(Import);
-	if (IRet != A3D_SUCCESS && IRet != A3D_LOAD_MULTI_MODELS_CADFILE && IRet != A3D_LOAD_MISSING_COMPONENTS)
+	ModelFile = TechSoftInterface::LoadModelFileFromFile(Import);
+
+	if (!ModelFile.IsValid())
 	{
 		return ECADParsingResult::ProcessFailed;
 	}
@@ -312,18 +314,17 @@ ECADParsingResult FTechSoftFileParser::Process()
 		FString CacheFilePath = CADFileData.GetCADCachePath();
 		if (CacheFilePath != File.GetPathOfFileToLoad())
 		{
-			// todo
+			TechSoftUtils::SaveModelFileToPcrFile(ModelFile.Get(), CacheFilePath);
 		}
 	}
 
-	A3DAsmModelFile* ModelFile = TechSoftInterface.GetModelFile();
-
 	if (CADFileData.GetImportParameters().GetStitchingTechnique() == StitchingSew && FImportParameters::bGDisableCADKernelTessellation)
 	{
-		TUniqueTSObj<A3DSewOptionsData> SewData;
+		CADLibrary::TUniqueTSObj<A3DSewOptionsData> SewData;
 		SewData->m_bComputePreferredOpenShellOrientation = false;
 		double ToleranceMM = 0.01 / FileUnit;
-		A3DStatus Status = TechSoftUtils::SewModel(&ModelFile, ToleranceMM, &*SewData);
+
+		A3DStatus Status = TechSoftInterface::SewModel(ModelFile.GetPtr(), ToleranceMM, SewData.GetPtr());
 		if(Status != A3DStatus::A3D_SUCCESS)
 		{
 			// To do but what ?
@@ -334,20 +335,20 @@ ECADParsingResult FTechSoftFileParser::Process()
 
 	ReadMaterialsAndColors();
 
-	ECADParsingResult Result = TraverseModel(ModelFile);
+	ECADParsingResult Result = TraverseModel();
 
 	if (Result == ECADParsingResult::ProcessOk)
 	{
 		GenerateBodyMeshes();
 	}
 
-	FString TechsoftVersion = TechSoftUtils::GetTechSoftVersion();
-	if (!TechsoftVersion.IsEmpty() && CADFileData.ComponentCount() > 0)
+	FString TechSoftVersion = TechSoftInterface::GetTechSoftVersion();
+	if (!TechSoftVersion.IsEmpty() && CADFileData.ComponentCount() > 0)
 	{
-		CADFileData.GetComponentAt(0).MetaData.Add(TEXT("TechsoftVersion"), TechsoftVersion);
+		CADFileData.GetComponentAt(0).MetaData.Add(TEXT("TechsoftVersion"), TechSoftVersion);
 	}
-
-	TechSoftInterface.UnloadModel();
+	
+	ModelFile.Reset();
 
 	return Result;
 }
@@ -382,7 +383,7 @@ void FTechSoftFileParser::GenerateBodyMeshes()
 			//}
 			//else
 			{
-				TechSoftInterface.FillBodyMesh(RepresentationItemPtr, CADFileData.GetImportParameters(), FileUnit, BodyMesh);
+				TechSoftUtils::FillBodyMesh(RepresentationItemPtr, CADFileData.GetImportParameters(), FileUnit, BodyMesh);
 			}
 
 			// Convert material
@@ -419,7 +420,7 @@ void FTechSoftFileParser::GenerateBodyMeshes()
 			Body.ColorFaceSet = BodyMesh.ColorSet;
 			Body.MaterialFaceSet = BodyMesh.MaterialSet;
 
-			// Write part's representation as hsf file if it is a BRep
+			// Write part's representation as Prc file if it is a BRep
 			A3DEEntityType Type;
 			A3DEntityGetType(RepresentationItemPtr, &Type);
 
@@ -441,7 +442,7 @@ void FTechSoftFileParser::GenerateBodyMeshes()
 
 					FJsonSerializer::Serialize(JsonObject.ToSharedRef(), JsonWriter);
 
-					TechSoftInterface.SaveBodyToHsfFile(RepresentationItemPtr, FilePath, JsonString);
+					TechSoftUtils::SaveBodiesToPcrFile(&RepresentationItemPtr, 1, FilePath, JsonString);
 				}
 			}
 		}
@@ -455,7 +456,7 @@ void FTechSoftFileParser::GenerateBodyMeshes()
 void FTechSoftFileParser::ReserveCADFileData()
 {
 	// TODO: Could be more accurate
-	CountUnderModel(TechSoftInterface.GetModelFile());
+	CountUnderModel();
 
 	CADFileData.ReserveBodyMeshes(ComponentCount[EComponentType::Body]);
 
@@ -466,9 +467,9 @@ void FTechSoftFileParser::ReserveCADFileData()
 	SceneGraphArchive.MaterialHIdToMaterial.Reserve(MaterialNum);
 }
 
-void FTechSoftFileParser::CountUnderModel(const A3DAsmModelFile* AsmModel)
+void FTechSoftFileParser::CountUnderModel()
 {
-	TUniqueTSObj<A3DAsmModelFileData> ModelFileData(AsmModel);
+	TUniqueTSObj<A3DAsmModelFileData> ModelFileData(ModelFile.Get());
 	if (!ModelFileData.IsValid())
 	{
 		return;
@@ -489,9 +490,9 @@ void FTechSoftFileParser::CountUnderModel(const A3DAsmModelFile* AsmModel)
 	}
 }
 
-ECADParsingResult FTechSoftFileParser::TraverseModel(const A3DAsmModelFile* ModelFile)
+ECADParsingResult FTechSoftFileParser::TraverseModel()
 {
-	TUniqueTSObj<A3DAsmModelFileData> ModelFileData(ModelFile);
+	TUniqueTSObj<A3DAsmModelFileData> ModelFileData(ModelFile.Get());
 	if (!ModelFileData.IsValid())
 	{
 		return ECADParsingResult::ProcessFailed;
@@ -501,8 +502,8 @@ ECADParsingResult FTechSoftFileParser::TraverseModel(const A3DAsmModelFile* Mode
 	FileUnit = ModelFileData->m_dUnit;
 
 	FEntityMetaData MetaData;
-	ExtractMetaData(ModelFile, MetaData);
-	ExtractSpecificMetaData(ModelFile, MetaData);
+	ExtractMetaData(ModelFile.Get(), MetaData);
+	ExtractSpecificMetaData(ModelFile.Get(), MetaData);
 
 	for (uint32 Index = 0; Index < ModelFileData->m_uiPOccurrencesSize; ++Index)
 	{
@@ -915,10 +916,10 @@ void FTechSoftFileParser::ProcessPrototype(const A3DAsmProductOccurrence* InProt
 			ExtractMaterialProperties(PrototypePtr);
 
 			TUniqueTSObj<A3DUTF8Char*> FilePathUTF8Ptr;
-			FilePathUTF8Ptr.FillWith(&TechSoftUtils::GetFilePathName, PrototypePtr);
+			FilePathUTF8Ptr.FillWith(&TechSoftInterface::GetFilePathName, PrototypePtr);
 			if (!FilePathUTF8Ptr.IsValid())
 			{
-				FilePathUTF8Ptr.FillWith(&TechSoftUtils::GetOriginalFilePathName, PrototypePtr);
+				FilePathUTF8Ptr.FillWith(&TechSoftInterface::GetOriginalFilePathName, PrototypePtr);
 			}
 			if (FilePathUTF8Ptr.IsValid())
 			{
@@ -1493,7 +1494,7 @@ FArchiveMaterial& FTechSoftFileParser::FindOrAddMaterial(uint32 MaterialIndex, c
 		return *MaterialArchive;
 	}
 
-	bool bIsTexture = TechSoftUtils::IsMaterialTexture(MaterialIndex);
+	bool bIsTexture = TechSoftInterface::IsMaterialTexture(MaterialIndex);
 	if (bIsTexture)
 	{
 		TUniqueTSObjFromIndex<A3DGraphTextureApplicationData> TextureData(MaterialIndex);
@@ -1737,8 +1738,8 @@ bool FTechSoftFileParser::IsConfigurationSet(const A3DAsmProductOccurrence* Occu
 
 uint32 FTechSoftFileParser::CountColorAndMaterial()
 {
-	A3DGlobal* GlobalPtr = nullptr;
-	if (TechSoftUtils::GetGlobalPointer(&GlobalPtr) != A3D_SUCCESS)
+	A3DGlobal* GlobalPtr = TechSoftInterface::GetGlobalPointer();
+	if (GlobalPtr == nullptr)
 	{
 		return 0;
 	}
@@ -1795,8 +1796,8 @@ void ExtractTextureDefinition(const A3DGraphTextureDefinitionData& TextureDefini
 
 void FTechSoftFileParser::ReadMaterialsAndColors()
 {
-	A3DGlobal* GlobalPtr = nullptr;
-	if(TechSoftUtils::GetGlobalPointer(&GlobalPtr) != A3D_SUCCESS)
+	A3DGlobal* GlobalPtr = TechSoftInterface::GetGlobalPointer();
+	if(GlobalPtr == nullptr)
 	{
 		return;
 	}
@@ -1828,7 +1829,7 @@ void FTechSoftFileParser::ReadMaterialsAndColors()
 			TUniqueTSObjFromIndex<A3DGraphPictureData> PictureData;
 			for (uint32 PictureIndex = 0; PictureIndex < PictureCount; ++PictureIndex)
 			{
-				A3DEntity* PicturePtr = TechSoftUtils::GetPointerFromIndex(PictureIndex, kA3DTypeGraphPicture);
+				A3DEntity* PicturePtr = TechSoftInterface::GetPointerFromIndex(PictureIndex, kA3DTypeGraphPicture);
 				if (PicturePtr)
 				{
 					FEntityMetaData PictureMetaData;
diff --git a/Engine/Plugins/Enterprise/DatasmithCADImporter/Source/CADInterfaces/Private/TechSoftFileParser.h b/Engine/Plugins/Enterprise/DatasmithCADImporter/Source/CADInterfaces/Private/TechSoftFileParser.h
--- a/Engine/Plugins/Enterprise/DatasmithCADImporter/Source/CADInterfaces/Private/TechSoftFileParser.h
+++ b/Engine/Plugins/Enterprise/DatasmithCADImporter/Source/CADInterfaces/Private/TechSoftFileParser.h
@@ -9,6 +9,7 @@
 #include "CADOptions.h"
 #include "CADSceneGraph.h"
 
+#include "TechSoftInterface.h"
 #include "TUniqueTechSoftObj.h"
 
 typedef void A3DAsmModelFile;
@@ -137,7 +138,7 @@ private:
 	// Methods to parse a first time the file to count the number of components
 	// Needed to reserve CADFileData
 	// Start with CountUnderModel
-	void CountUnderModel(const A3DAsmModelFile* AsmModel);
+	void CountUnderModel();
 	void CountUnderConfigurationSet(const A3DAsmProductOccurrence* Occurrence);
 	void CountUnderOccurrence(const A3DAsmProductOccurrence* Occurrence);
 	void CountUnderPrototype(const A3DAsmProductOccurrence* Prototype);
@@ -153,7 +154,7 @@ private:
 	void ReadMaterialsAndColors();
 
 	// Traverse ASM tree by starting from the model
-	ECADParsingResult TraverseModel(const A3DAsmModelFile* AsmModel);
+	ECADParsingResult TraverseModel();
 	void TraverseReference(const A3DAsmProductOccurrence* Reference);
 	bool IsConfigurationSet(const A3DAsmProductOccurrence* Occurrence);
 	void TraverseConfigurationSet(const A3DAsmProductOccurrence* ConfigurationSet);
@@ -205,7 +206,7 @@ private:
 	FMatrix TraverseGeneralTransformation(const A3DMiscTransformation* GeneralTransformation);
 	FMatrix TraverseTransformation3D(const A3DMiscTransformation* CartesianTransformation);
 
-	// Archive methodes
+	// Archive methods
 	FArchiveInstance& AddInstance(FEntityMetaData& InstanceMetaData);
 	FArchiveComponent& AddComponent(FEntityMetaData& ComponentMetaData, FArchiveInstance& Instance);
 	FArchiveUnloadedComponent& AddUnloadedComponent(FEntityMetaData& ComponentMetaData, FArchiveInstance& Instance);
@@ -216,6 +217,8 @@ private:
 
 private:
 
+	FUniqueTechSoftModelFile ModelFile;
+
 	TSet<const A3DAsmProductOccurrence*> PrototypeCounted;
 	uint32 ComponentCount[EComponentType::LastType] = { 0 };
 
@@ -232,24 +235,4 @@ private:
 	TMap<A3DRiRepresentationItem*, int32> RepresentationItemsCache;
 };
 
-namespace TechSoftFileParserImpl
-{
-#ifdef USE_TECHSOFT_SDK
-	inline FColor GetColorAt(uint32 ColorIndex)
-	{
-		TUniqueTSObjFromIndex<A3DGraphRgbColorData> ColorData(ColorIndex);
-		if (ColorData.IsValid())
-		{
-			return FColor((uint8)(ColorData->m_dRed * 255)
-				, (uint8)(ColorData->m_dGreen * 255)
-				, (uint8)(ColorData->m_dBlue * 255));
-		}
-		else
-		{
-			return FColor(200, 200, 200);
-		}
-	}
-#endif
-} // ns TechSoftFileParserImpl
-
 } // ns CADLibrary
\ No newline at end of file
diff --git a/Engine/Plugins/Enterprise/DatasmithCADImporter/Source/CADInterfaces/Private/TechSoftInterface.cpp b/Engine/Plugins/Enterprise/DatasmithCADImporter/Source/CADInterfaces/Private/TechSoftInterface.cpp
--- a/Engine/Plugins/Enterprise/DatasmithCADImporter/Source/CADInterfaces/Private/TechSoftInterface.cpp
+++ b/Engine/Plugins/Enterprise/DatasmithCADImporter/Source/CADInterfaces/Private/TechSoftInterface.cpp
@@ -6,42 +6,21 @@
 
 #include "TUniqueTechSoftObj.h"
 
-#include "CADOptions.h"
 #include "CADInterfacesModule.h"
-
 #include "HAL/PlatformProcess.h"
-#include "Math/Color.h"
 #include "Misc/Paths.h"
-#include "Serialization/JsonSerializer.h"
-#include "Serialization/JsonWriter.h"
-
-#ifndef CADKERNEL_DEV
-#include "TechSoftInterfaceUtils.inl"
-#endif
-
-#include <string>
 
 namespace CADLibrary
 {
-namespace TechSoftUtils
-{
-#if defined USE_TECHSOFT_SDK && !defined CADKERNEL_DEV
-	TSharedPtr<FJsonObject> GetJsonObject(A3DAsmProductOccurrence* ProductOcccurence);
-	void RestoreMaterials(const TSharedPtr<FJsonObject>& DefaultValues, CADLibrary::FBodyMesh& BodyMesh);
-#endif
-}
 
-	FTechSoftInterface& GetTechSoftInterface()
+uint32 FTechSoftInterface::InvalidScriptIndex = -1;
+
+FTechSoftInterface& FTechSoftInterface::Get()
 {
 	static FTechSoftInterface TechSoftInterface;
 	return TechSoftInterface;
 }
 
-bool TECHSOFT_InitializeKernel(const TCHAR* InEnginePluginsPath)
-{
-	return GetTechSoftInterface().InitializeKernel(InEnginePluginsPath);
-}
-
 bool FTechSoftInterface::InitializeKernel(const TCHAR* InEnginePluginsPath)
 {
 #ifdef USE_TECHSOFT_SDK
@@ -65,1265 +44,1235 @@ bool FTechSoftInterface::InitializeKernel(const TCHAR* InEnginePluginsPath)
 #endif
 
 	TechSoftDllPath = FPaths::ConvertRelativePathToFull(TechSoftDllPath);
-	ExchangeLoader = MakeUnique<A3DSDKHOOPSExchangeLoader>(*TechSoftDllPath);
 
-	const A3DStatus IRet = ExchangeLoader->m_eSDKStatus;
-	if (IRet != A3D_SUCCESS)
-	{
-#ifndef CADKERNEL_DEV
-		UE_LOG(LogCADInterfaces, Warning, TEXT("Failed to load required library in %s. Plug-in will not be functional."), *TechSoftDllPath);
-#endif
-	}
-	else
+	if (A3DSDKLoadLibrary(*TechSoftDllPath))
 	{
-		bIsInitialize = true;
+		A3DLicPutUnifiedLicense(HOOPS_LICENSE);
+
+		A3DInt32 iMajorVersion = 0, iMinorVersion = 0;
+		if (A3DDllGetVersion(&iMajorVersion, &iMinorVersion) == A3D_SUCCESS)
+		{
+			if (A3DDllInitialize(A3D_DLL_MAJORVERSION, A3D_DLL_MINORVERSION) == A3D_SUCCESS)
+			{
+				bIsInitialize = true;
+				return true;
+			}
+		}
 	}
-	return bIsInitialize;
-#else
-	return false;
+	UE_LOG(LogCADInterfaces, Warning, TEXT("Failed to load required library in %s. Plug-in will not be functional."), *TechSoftDllPath);
+
 #endif
+	return false;
 }
 
+namespace TechSoftInterface
+{
+
+bool TECHSOFT_InitializeKernel(const TCHAR* InEnginePluginsPath)
+{
+	return FTechSoftInterface::Get().InitializeKernel(InEnginePluginsPath);
+}
+
+FString GetTechSoftVersion()
+{
 #ifdef USE_TECHSOFT_SDK
+	A3DInt32 MajorVersion = 0, MinorVersion = 0;
+	A3DDllGetVersion(&MajorVersion, &MinorVersion);
+	return FString::Printf(TEXT("Techsoft %d.%d"), MajorVersion, MinorVersion);
+#endif
+	return FString();
+}
 
-A3DUns32 FTechSoftInterface::InvalidScriptIndex = -1;
+#ifdef USE_TECHSOFT_SDK
 
-A3DStatus FTechSoftInterface::Import(const A3DImport& Import)
+FUniqueTechSoftModelFile LoadModelFileFromFile(const A3DImport& Import)
 {
-	return ExchangeLoader->Import(Import);
+	A3DAsmModelFile* ModelFile = nullptr;
+	if (A3DAsmModelFileLoadFromFile(Import.GetFilePath(), &Import.m_sLoadData, &ModelFile) != A3DStatus::A3D_SUCCESS)
+	{
+		return FUniqueTechSoftModelFile();
+	}
+	return FUniqueTechSoftModelFile(ModelFile);
 }
 
-A3DStatus FTechSoftInterface::UnloadModel()
+FUniqueTechSoftModelFile LoadModelFileFromPrcFile(const A3DUTF8Char* CADFileName, A3DRWParamsPrcReadHelper** ReadHelper)
 {
-	A3DStatus Ret = A3DStatus::A3D_ERROR;
-	if (ExchangeLoader)
+	A3DAsmModelFile* ModelFile = nullptr;
+	if (A3DAsmModelFileLoadFromPrcFile(CADFileName, ReadHelper, &ModelFile) != A3DStatus::A3D_SUCCESS)
 	{
-		Ret = A3DAsmModelFileDelete(ExchangeLoader->m_psModelFile);
+		return FUniqueTechSoftModelFile();
 	}
-	ExchangeLoader->m_psModelFile = nullptr;
-	return Ret;
+	return FUniqueTechSoftModelFile(ModelFile);
 }
 
-A3DAsmModelFile* FTechSoftInterface::GetModelFile()
+A3DStatus DeleteModelFile(A3DAsmModelFile* ModelFile)
 {
-	return ExchangeLoader->m_psModelFile;
+	return A3DAsmModelFileDelete(ModelFile);
+}
+
+A3DStatus DeleteEntity(A3DEntity* EntityPtr)
+{
+	return A3DEntityDelete(EntityPtr);
 }
-#endif
 
-void FTechSoftInterface::SaveBodyToHsfFile(void* BodyPtr, const FString& Filename, const FString& JsonString)
+A3DGlobal* GetGlobalPointer()
 {
-#if defined USE_TECHSOFT_SDK && !defined CADKERNEL_DEV
-	if (!BodyPtr)
+	A3DGlobal* GlobalPtr;
+	if(A3DGlobalGetPointer(&GlobalPtr) == A3DStatus::A3D_SUCCESS)
 	{
-		return;
+		return GlobalPtr;
 	}
+	return nullptr;
+}
 
-	// Create PartDefinition
-	A3DRiRepresentationItem* RepresentationItem = (A3DRiRepresentationItem*)BodyPtr;
-	A3DAsmPartDefinition* PartDefinition = nullptr;
-
-	A3DAsmPartDefinitionData PartDefinitionData;
-	A3D_INITIALIZE_DATA(A3DAsmPartDefinitionData, PartDefinitionData);
-
-	PartDefinitionData.m_uiRepItemsSize = 1;
-	PartDefinitionData.m_ppRepItems = &RepresentationItem;
-
-	A3DAsmPartDefinitionCreate(&PartDefinitionData, &PartDefinition);
-
-	// Create ProductOccurrence
-	A3DAsmProductOccurrence* ProductOccurrence = nullptr;
-
-	A3DAsmProductOccurrenceData ProductOccurrenceData;
-	A3D_INITIALIZE_DATA(A3DAsmProductOccurrenceData, ProductOccurrenceData);
-
-	ProductOccurrenceData.m_pPart = PartDefinition;
-
-	A3DInt32 iRet = A3DAsmProductOccurrenceCreate(&ProductOccurrenceData, &ProductOccurrence);
-
-	// Add MaterialTable as attribute to ProductOccurrence
-	std::string StringAnsi(TCHAR_TO_UTF8(*JsonString));
-
-	A3DMiscSingleAttributeData SingleAttributeData;
-	A3D_INITIALIZE_DATA(A3DMiscSingleAttributeData, SingleAttributeData);
-
-	SingleAttributeData.m_eType = kA3DModellerAttributeTypeString;
-	SingleAttributeData.m_pcTitle = "MaterialTable";
-	SingleAttributeData.m_pcData = (char*)StringAnsi.c_str();
-
-	A3DMiscAttributeData AttributesData;
-	A3D_INITIALIZE_DATA(A3DMiscAttributeData, AttributesData);
-
-	A3DMiscAttribute* Attributes = nullptr;
-
-	AttributesData.m_pcTitle = SingleAttributeData.m_pcTitle;
-	AttributesData.m_asSingleAttributesData = &SingleAttributeData;
-	AttributesData.m_uiSize = 1;
-	A3DMiscAttributeCreate(&AttributesData, &Attributes);
-
-	A3DRootBaseData RootBaseData;
-	A3D_INITIALIZE_DATA(A3DRootBaseData, RootBaseData);
-
-	RootBaseData.m_pcName = SingleAttributeData.m_pcTitle;
-	RootBaseData.m_ppAttributes = &Attributes;
-	RootBaseData.m_uiSize = 1;
-	A3DRootBaseSet(ProductOccurrence, &RootBaseData);
-
-	// Create ModelFile
-	A3DAsmModelFile* ModelFile = nullptr;
-
-	A3DAsmModelFileData ModelFileData;
-	A3D_INITIALIZE_DATA(A3DAsmModelFileData, ModelFileData);
+A3DEntity* GetPointerFromIndex(const uint32 Index, const A3DEEntityType Type)
+{
+	A3DEntity* EntityPtr = nullptr;
+	if (A3DMiscPointerFromIndexGet(Index, Type, &EntityPtr) != A3DStatus::A3D_SUCCESS)
+	{
+		return nullptr;
+	}
+	return EntityPtr;
+}
 
-	ModelFileData.m_uiPOccurrencesSize = 1;
-	ModelFileData.m_dUnit = 1.0;
-	ModelFileData.m_ppPOccurrences = &ProductOccurrence;
+A3DStatus GetSurfaceAsNurbs(const A3DSurfBase* SurfacePtr, A3DSurfNurbsData* OutDataPtr, A3DDouble Tolerance, A3DBool bUseSameParameterization)
+{
+	return A3DSurfBaseGetAsNurbs(SurfacePtr, Tolerance, bUseSameParameterization, OutDataPtr);
+}
 
-	A3DAsmModelFileCreate(&ModelFileData, &ModelFile);
+A3DStatus GetCurveAsNurbs(const A3DCrvBase* CurvePtr, A3DCrvNurbsData* OutDataPtr, A3DDouble Tolerance, A3DBool bUseSameParameterization)
+{
+	return A3DCrvBaseGetAsNurbs(CurvePtr, Tolerance, bUseSameParameterization, OutDataPtr);
+}
 
-	// Save ModelFile to hsf file
-	A3DRWParamsExportPrcData sParamsExportData;
-	A3D_INITIALIZE_DATA(A3DRWParamsExportPrcData, sParamsExportData);
+A3DStatus GetOriginalFilePathName(const A3DAsmProductOccurrence* A3DOccurrencePtr, A3DUTF8Char** FilePathUTF8Ptr)
+{
+	return A3DAsmProductOccurrenceGetOriginalFilePathName(A3DOccurrencePtr, FilePathUTF8Ptr);
+}
 
-	sParamsExportData.m_bCompressBrep = false;
-	sParamsExportData.m_bCompressTessellation = false;
+A3DStatus GetFilePathName(const A3DAsmProductOccurrence* A3DOccurrencePtr, A3DUTF8Char** FilePathUTF8Ptr)
+{
+	return A3DAsmProductOccurrenceGetFilePathName(A3DOccurrencePtr, FilePathUTF8Ptr);
+}
 
-	A3DUTF8Char HsfFileName[_MAX_PATH];
-	FCStringAnsi::Strncpy(HsfFileName, TCHAR_TO_UTF8(*Filename), _MAX_PATH);
+A3DStatus GetEntityType(const A3DEntity* EntityPtr, A3DEEntityType* EntityTypePtr)
+{
+	return A3DEntityGetType(EntityPtr, EntityTypePtr);
+}
 
-	iRet = A3DAsmModelFileExportToPrcFile(ModelFile, &sParamsExportData, HsfFileName, nullptr);
+bool IsEntityBaseWithGraphicsType(const A3DEntity* EntityPtr)
+{
+	return (bool)A3DEntityIsBaseWithGraphicsType(EntityPtr);
+}
 
-	// #ueent_techsoft: Deleting the model seems to delete the entire content. To be double-checked
-	A3DEntityDelete(Attributes);
-	A3DEntityDelete(ModelFile);
-#endif
+bool IsEntityBaseType(const A3DEntity* EntityPtr)
+{
+	return (bool)A3DEntityIsBaseType(EntityPtr);
 }
 
-bool FTechSoftInterface::GetBodyFromHsfFile(const FString& Filename, const FImportParameters& ImportParameters, FBodyMesh& BodyMesh)
+bool IsMaterialTexture(const uint32 MaterialIndex)
 {
-	bool bExtractionSuccessful = false;
+	A3DBool bIsTexture = false;
+	return A3DGlobalIsMaterialTexture(MaterialIndex, &bIsTexture) == A3DStatus::A3D_SUCCESS ? bool(bIsTexture) : false;
+}
 
-#if defined USE_TECHSOFT_SDK && !defined CADKERNEL_DEV
-	A3DRWParamsPrcReadHelper* ReadHelper = nullptr;
+FUniqueTechSoftModelFile CreateAsmModelFile(const A3DAsmModelFileData& ModelFileData)
+{
 	A3DAsmModelFile* ModelFile = nullptr;
-
-	A3DInt32 iRet = A3DAsmModelFileLoadFromPrcFile(TCHAR_TO_UTF8(*Filename), &ReadHelper, &ModelFile);
-
-	if (iRet != A3D_SUCCESS || !ModelFile)
+	if (A3DAsmModelFileCreate(&ModelFileData, &ModelFile) != A3DStatus::A3D_SUCCESS)
 	{
-		return false;
-	}
-
-	TUniqueTSObj<A3DAsmModelFileData> ModelFileData(ModelFile);
-	if (!ModelFileData.IsValid() || ModelFileData->m_uiPOccurrencesSize == 0)
-	{
-		return false;
-	}
-
-	TUniqueTSObj<A3DAsmProductOccurrenceData> OccurenceData(ModelFileData->m_ppPOccurrences[0]);
-	if (!OccurenceData.IsValid() || !OccurenceData->m_pPart)
-	{
-		return false;
+		return FUniqueTechSoftModelFile();
 	}
+	return FUniqueTechSoftModelFile(ModelFile);
+}
 
-	TUniqueTSObj<A3DAsmPartDefinitionData> PartDefinitionData(OccurenceData->m_pPart);
-	if (!PartDefinitionData.IsValid() || PartDefinitionData->m_uiRepItemsSize == 0)
+A3DRiBrepModel* CreateRiBRepModel(const A3DRiBrepModelData& RiBRepModelData)
+{
+	A3DRiBrepModel* RiBrepModelPtr = nullptr;
+	if (A3DRiBrepModelCreate(&RiBRepModelData, &RiBrepModelPtr) != A3DStatus::A3D_SUCCESS)
 	{
-		return false;
+		return nullptr;
 	}
+	return RiBrepModelPtr;
+}
 
-	TSharedPtr<FJsonObject> JsonObject = TechSoftUtils::GetJsonObject(ModelFileData->m_ppPOccurrences[0]);
-	if (JsonObject.IsValid())
+A3DAsmPartDefinition* CreateAsmPartDefinition(const A3DAsmPartDefinitionData& PartDefinitionData)
+{
+	A3DAsmPartDefinition* AsmPartDefinitionPtr = nullptr;
+	if (A3DAsmPartDefinitionCreate(&PartDefinitionData, &AsmPartDefinitionPtr) != A3DStatus::A3D_SUCCESS)
 	{
-		double FileUnit = 1.0;
-
-		JsonObject->TryGetNumberField(JSON_ENTRY_FILE_UNIT, FileUnit);
-
-		if (FillBodyMesh(PartDefinitionData->m_ppRepItems[0], ImportParameters, FileUnit, BodyMesh))
-		{
-			bExtractionSuccessful = true;
-
-			TechSoftUtils::RestoreMaterials(JsonObject, BodyMesh);
-		}
+		return nullptr;
 	}
-
-	A3DEntityDelete(ModelFile);
-#endif
-
-	return bExtractionSuccessful;
+	return AsmPartDefinitionPtr;
 }
 
-bool FTechSoftInterface::FillBodyMesh(void* BodyPtr, const FImportParameters& ImportParameters, double FileUnit, FBodyMesh& BodyMesh)
+A3DAsmProductOccurrence* CreateAsmProductOccurrence(const A3DAsmProductOccurrenceData& ProductOccurrenceData)
 {
-#if defined USE_TECHSOFT_SDK && !defined CADKERNEL_DEV
-	A3DRiRepresentationItem* RepresentationItemPtr = (A3DRiRepresentationItem*)BodyPtr;
-
-	A3DEEntityType Type;
-	A3DEntityGetType(RepresentationItemPtr, &Type);
-	if (Type == kA3DTypeRiPolyBrepModel)
+	A3DAsmProductOccurrence* ProductOccurrencePtr = nullptr;
+	if (A3DAsmProductOccurrenceCreate(&ProductOccurrenceData, &ProductOccurrencePtr) != A3DStatus::A3D_SUCCESS)
 	{
-		TUniqueTSObj<A3DRiRepresentationItemData> RepresentationItemData(RepresentationItemPtr);
-		TechSoftInterfaceUtils::FTechSoftTessellationExtractor Extractor(RepresentationItemData->m_pTessBase);
-		return Extractor.FillBodyMesh(BodyMesh, FileUnit);
+		return nullptr;
 	}
+	return ProductOccurrencePtr;
+}
 
-	TUniqueTSObj<A3DRiRepresentationItemData> RepresentationItemData;
-
-	// TUniqueTechSoftObj does not work in this case
-	A3DRWParamsTessellationData TessellationParameters;
-	A3D_INITIALIZE_DATA(A3DRWParamsTessellationData, TessellationParameters);
-
-	TessellationParameters.m_eTessellationLevelOfDetail = kA3DTessLODUserDefined; // Enum to specify predefined values for some following members.
-	TessellationParameters.m_bUseHeightInsteadOfRatio = A3D_TRUE;
-	TessellationParameters.m_dMaxChordHeight = ImportParameters.GetChordTolerance() * 10. / FileUnit;
-	TessellationParameters.m_dAngleToleranceDeg = ImportParameters.GetMaxNormalAngle();
-	TessellationParameters.m_dMaximalTriangleEdgeLength = 0; //ImportParameters.MaxEdgeLength;
-
-	TessellationParameters.m_bAccurateTessellation = A3D_FALSE;  // A3D_FALSE' indicates the tessellation is set for visualization
-	TessellationParameters.m_bAccurateTessellationWithGrid = A3D_FALSE; // Enable accurate tessellation with faces inner points on a grid.
-	TessellationParameters.m_dAccurateTessellationWithGridMaximumStitchLength = 0; 	// Maximal grid stitch length. Disabled if value is 0. Be careful, a too small value can generate a huge tessellation.
-
-	TessellationParameters.m_bKeepUVPoints = A3D_TRUE; // Keep parametric points as texture points.
-
-	// Get the tessellation
-	A3DStatus Status = A3DRiRepresentationItemComputeTessellation(RepresentationItemPtr, &TessellationParameters);
-	Status = A3DRiRepresentationItemGet(RepresentationItemPtr, RepresentationItemData.GetEmptyDataPtr());
-
+A3DMiscAttribute* CreateMiscAttribute(const A3DMiscAttributeData& AttributeData)
+{
+	A3DMiscAttribute* AttributePtr = nullptr;
+	if (A3DMiscAttributeCreate(&AttributeData, &AttributePtr) != A3DStatus::A3D_SUCCESS)
 	{
-		A3DEntityGetType(RepresentationItemData->m_pTessBase, &Type);
-		ensure(Type == kA3DTypeTess3D);
+		return nullptr;
 	}
-
-	TechSoftInterfaceUtils::FTechSoftTessellationExtractor Extractor(RepresentationItemData->m_pTessBase);
-	return Extractor.FillBodyMesh(BodyMesh, FileUnit);
-#else
-	return false;
-#endif
+	return AttributePtr;
 }
 
-namespace TechSoftUtils
-{
-
-#if defined USE_TECHSOFT_SDK && !defined CADKERNEL_DEV
-TSharedPtr<FJsonObject> GetJsonObject(A3DAsmProductOccurrence* ProductOcccurence)
+A3DAsmModelFile* CreateModelFile(const A3DAsmModelFileData& ModelFileData)
 {
-	TUniqueTSObj<A3DRootBaseData> RootBaseData(ProductOcccurence);
-
-	if (RootBaseData.IsValid() && RootBaseData->m_uiSize > 0)
+	A3DAsmModelFile* ModelFilePtr = nullptr;
+	if (A3DAsmModelFileCreate(&ModelFileData, &ModelFilePtr) != A3DStatus::A3D_SUCCESS)
 	{
-		TUniqueTSObj<A3DMiscAttributeData> AttributeData(RootBaseData->m_ppAttributes[0]);
-		if (AttributeData->m_uiSize > 0 && AttributeData->m_asSingleAttributesData[0].m_eType == kA3DModellerAttributeTypeString)
-		{
-			FString JsonString = UTF8_TO_TCHAR(AttributeData->m_asSingleAttributesData[0].m_pcData);
-
-			TSharedRef<TJsonReader<>> JsonReader = TJsonReaderFactory<>::Create(JsonString);
-			TSharedPtr<FJsonObject> JsonObject;
-
-			if (FJsonSerializer::Deserialize(JsonReader, JsonObject))
-			{
-				return JsonObject;
-			}
-		}
+		return nullptr;
 	}
-
-	return {};
+	return ModelFilePtr;
 }
 
-FColor GetColorAt(uint32 ColorIndex)
+A3DTopoBrepData* CreateTopoBRep(const A3DTopoBrepDataData& TopoBRepDataData)
 {
-	TUniqueTSObjFromIndex<A3DGraphRgbColorData> ColorData(ColorIndex);
-	if (ColorData.IsValid())
-	{
-		return FColor((uint8)(ColorData->m_dRed * 255)
-			, (uint8)(ColorData->m_dGreen * 255)
-			, (uint8)(ColorData->m_dBlue * 255));
-	}
-	else
+	A3DTopoBrepData* TopoBRepDataPtr = nullptr;
+	if (A3DTopoBrepDataCreate(&TopoBRepDataData, &TopoBRepDataPtr) != A3DStatus::A3D_SUCCESS)
 	{
-		return FColor(200, 200, 200);
+		return nullptr;
 	}
+	return TopoBRepDataPtr;
 }
 
-// Replicate logic in FTechSoftFileParser::FindOrAddMaterial and the methods it calls
-void BuildCADMaterial(uint32 MaterialIndex, const A3DGraphStyleData& GraphStyleData, FCADMaterial& Material)
+A3DTopoCoEdge* CreateTopoCoEdge(const A3DTopoCoEdgeData& TopoCoEdgeData)
 {
-	if (TechSoftUtils::IsMaterialTexture(MaterialIndex))
+	A3DTopoCoEdge* TopoCoEdgePtr = nullptr;
+	if (A3DTopoCoEdgeCreate(&TopoCoEdgeData, &TopoCoEdgePtr) != A3DStatus::A3D_SUCCESS)
 	{
-		TUniqueTSObjFromIndex<A3DGraphTextureApplicationData> TextureData(MaterialIndex);
-		if (TextureData.IsValid())
-		{
-			MaterialIndex = TextureData->m_uiMaterialIndex;
-		}
+		return nullptr;
 	}
+	return TopoCoEdgePtr;
+}
 
-	TUniqueTSObjFromIndex<A3DGraphMaterialData> MaterialData(MaterialIndex);
-	if (MaterialData.IsValid())
+A3DTopoConnex* CreateTopoConnex(const A3DTopoConnexData& TopoConnexData)
+{
+	A3DTopoConnex* TopoConnexPtr = nullptr;
+	if (A3DTopoConnexCreate(&TopoConnexData, &TopoConnexPtr) != A3DStatus::A3D_SUCCESS)
 	{
-		Material.Diffuse = GetColorAt(MaterialData->m_uiDiffuse);
-		Material.Ambient = GetColorAt(MaterialData->m_uiAmbient);
-		Material.Specular = GetColorAt(MaterialData->m_uiSpecular);
-		Material.Shininess = MaterialData->m_dShininess;
-		if (GraphStyleData.m_bIsTransparencyDefined)
-		{
-			Material.Transparency = 1. - GraphStyleData.m_ucTransparency / 255.;
-		}
-		// todo: find how to convert Emissive color into ? reflexion coef...
-		// Material.Emissive = GetColor(MaterialData->m_uiEmissive);
-		// Material.Reflexion;
+		return nullptr;
 	}
+	return TopoConnexPtr;
 }
 
-// Replicates the logic in FTechSoftFileParser::ExtractGraphStyleProperties
-void GetMaterialValues(uint32 StyleIndex, FCADUUID& OutColorName, FCADUUID& OutMaterialName)
+A3DTopoEdge* CreateTopoEdge(const A3DTopoEdgeData& TopoEdgeData)
 {
-	TUniqueTSObjFromIndex<A3DGraphStyleData> GraphStyleData(StyleIndex);
-
-	if (GraphStyleData.IsValid())
+	A3DTopoEdge* TopoEdgePtr = nullptr;
+	if (A3DTopoEdgeCreate(&TopoEdgeData, &TopoEdgePtr) != A3DStatus::A3D_SUCCESS)
 	{
-		if (GraphStyleData->m_bMaterial)
-		{
-			FCADMaterial Material;
-
-			BuildCADMaterial(GraphStyleData->m_uiRgbColorIndex, *GraphStyleData, Material);
-
-			OutMaterialName = BuildMaterialName(Material);
-		}
-		else
-		{
-			TUniqueTSObjFromIndex<A3DGraphRgbColorData> ColorData(GraphStyleData->m_uiRgbColorIndex);
-			if (ColorData.IsValid())
-			{
-				const uint8 Alpha = GraphStyleData->m_bIsTransparencyDefined ? GraphStyleData->m_ucTransparency : 255;
-				const FColor ColorValue((uint8)(ColorData->m_dRed * 255), (uint8)(ColorData->m_dGreen * 255), (uint8)(ColorData->m_dBlue * 255), Alpha);
-
-				OutColorName = BuildColorName(ColorValue);
-			}
-		}
+		return nullptr;
 	}
+	return TopoEdgePtr;
 }
 
-void RestoreMaterials(const TSharedPtr<FJsonObject>& DefaultValues, FBodyMesh& BodyMesh)
+A3DTopoFace* CreateTopoFace(const A3DTopoFaceData& TopoFaceData)
 {
-	FCADUUID DefaultColorName = 0;
-	FCADUUID DefaultMaterialName = 0;
-
-	DefaultValues->TryGetNumberField(JSON_ENTRY_COLOR_NAME, DefaultColorName);
-	DefaultValues->TryGetNumberField(JSON_ENTRY_MATERIAL_NAME, DefaultMaterialName);
-
-	BodyMesh.MaterialSet.Empty();
-	BodyMesh.ColorSet.Empty();
-
-	for (FTessellationData& Tessellation : BodyMesh.Faces)
+	A3DTopoFace* TopoFacePtr = nullptr;
+	if (A3DTopoFaceCreate(&TopoFaceData, &TopoFacePtr) != A3DStatus::A3D_SUCCESS)
 	{
-		// Extract proper color or material based on style index
-		uint32 CachedStyleIndex = Tessellation.MaterialName;
-		Tessellation.MaterialName = 0;
-
-		FCADUUID ColorName = DefaultColorName;
-		FCADUUID MaterialName = DefaultMaterialName;
-
-		GetMaterialValues(CachedStyleIndex, ColorName, MaterialName);
-
-		if (ColorName)
-		{
-			Tessellation.ColorName = ColorName;
-			BodyMesh.ColorSet.Add(ColorName);
-		}
-
-		if (MaterialName)
-		{
-			Tessellation.MaterialName = MaterialName;
-			BodyMesh.MaterialSet.Add(MaterialName);
-		}
+		return nullptr;
 	}
+	return TopoFacePtr;
 }
-#endif
 
-FTechSoftInterface& GetTechSoftInterface()
+A3DTopoLoop* CreateTopoLoop(const A3DTopoLoopData& TopoLoopData)
 {
-	static FTechSoftInterface TechSoftInterface;
-	return TechSoftInterface;
+	A3DTopoLoop* TopoLoopPtr = nullptr;
+	if (A3DTopoLoopCreate(&TopoLoopData, &TopoLoopPtr) != A3DStatus::A3D_SUCCESS)
+	{
+		return nullptr;
+	}
+	return TopoLoopPtr;
 }
 
-FString GetTechSoftVersion()
+A3DCrvNurbs* CreateTopoShell(const A3DTopoShellData& TopoShellData)
 {
-#ifdef USE_TECHSOFT_SDK
-	A3DInt32 MajorVersion = 0, MinorVersion = 0;
-	A3DDllGetVersion(&MajorVersion, &MinorVersion);
-	return FString::Printf(TEXT("Techsoft %d.%d"), MajorVersion, MinorVersion);
-#endif
-	return FString();
+	A3DTopoShell* TopoShellPtr = nullptr;
+	if (A3DTopoShellCreate(&TopoShellData, &TopoShellPtr) != A3DStatus::A3D_SUCCESS)
+	{
+		return nullptr;
+	}
+	return TopoShellPtr;
 }
 
-
-bool TECHSOFT_InitializeKernel(const TCHAR* InEnginePluginsPath)
+A3DCrvNurbs* CreateCurveTransform(const A3DCrvTransformData& CurveTransformData)
 {
-	return GetTechSoftInterface().InitializeKernel(InEnginePluginsPath);
+	A3DCrvNurbs* CurveNurbsPtr = nullptr;
+	if (A3DCrvTransformCreate(&CurveTransformData, &CurveNurbsPtr) != A3DStatus::A3D_SUCCESS)
+	{
+		return nullptr;
+	}
+	return CurveNurbsPtr;
 }
 
-#ifdef USE_TECHSOFT_SDK
-
-A3DStatus GetGlobalPointer(A3DGlobal** GlobalPtr)
+A3DCrvNurbs* CreateCurveNurbs(const A3DCrvNurbsData& CurveNurbsData)
 {
-	return A3DGlobalGetPointer(GlobalPtr);
+	A3DCrvNurbs* CurveNurbsPtr = nullptr;
+	if (A3DCrvNurbsCreate(&CurveNurbsData, &CurveNurbsPtr) != A3DStatus::A3D_SUCCESS)
+	{
+		return nullptr;
+	}
+	return CurveNurbsPtr;
 }
 
-A3DStatus GetSurfaceAsNurbs(const A3DSurfBase* SurfacePtr, A3DSurfNurbsData* DataPtr, A3DDouble Tolerance, A3DBool bUseSameParameterization)
+A3DSurfNurbs* CreateSurfaceCylinder(const A3DSurfCylinderData& SurfaceCylinderData)
 {
-	return A3DSurfBaseGetAsNurbs(SurfacePtr, Tolerance, bUseSameParameterization, DataPtr);
+	A3DSurfCylinder* SurfCylinderPtr = nullptr;
+	if (A3DSurfCylinderCreate(&SurfaceCylinderData, &SurfCylinderPtr) != A3DStatus::A3D_SUCCESS)
+	{
+		return nullptr;
+	}
+	return SurfCylinderPtr;
 }
 
-A3DStatus GetCurveAsNurbs(const A3DCrvBase* CurvePtr, A3DCrvNurbsData* DataPtr, A3DDouble Tolerance, A3DBool bUseSameParameterization)
+A3DSurfNurbs* CreateSurfaceNurbs(const A3DSurfNurbsData& SurfaceNurbsData)
 {
-	return A3DCrvBaseGetAsNurbs(CurvePtr, Tolerance, bUseSameParameterization, DataPtr);
+	A3DSurfNurbs* SurfaceNurbsPtr = nullptr;
+	if (A3DSurfNurbsCreate(&SurfaceNurbsData, &SurfaceNurbsPtr) != A3DStatus::A3D_SUCCESS)
+	{
+		return nullptr;
+	}
+	return SurfaceNurbsPtr;
 }
 
-A3DStatus GetOriginalFilePathName(const A3DAsmProductOccurrence* A3DOccurrencePtr, A3DUTF8Char** FilePathUTF8Ptr)
+A3DGraphics* CreateGraphics(const A3DGraphicsData& GraphicsData)
 {
-	return A3DAsmProductOccurrenceGetOriginalFilePathName(A3DOccurrencePtr, FilePathUTF8Ptr);
+	A3DGraphics* GraphicsPtr = nullptr;
+	if (A3DGraphicsCreate(&GraphicsData, &GraphicsPtr) != A3DStatus::A3D_SUCCESS)
+	{
+		return nullptr;
+	}
+	return GraphicsPtr;
 }
 
-A3DStatus GetFilePathName(const A3DAsmProductOccurrence* A3DOccurrencePtr, A3DUTF8Char** FilePathUTF8Ptr)
+A3DStatus LinkCoEdges(A3DTopoCoEdge* CoEdgePtr, A3DTopoCoEdge* NeighbourCoEdgePtr)
 {
-	return A3DAsmProductOccurrenceGetFilePathName(A3DOccurrencePtr, FilePathUTF8Ptr);
+	return A3DTopoCoEdgeSetNeighbour(CoEdgePtr, NeighbourCoEdgePtr);
 }
 
-A3DStatus GetEntityType(const A3DEntity* EntityPtr, A3DEEntityType* EntityTypePtr)
+A3DStatus SetRootBase(A3DEntity* EntityPtr, const A3DRootBaseData& RootBaseData)
 {
-	return A3DEntityGetType(EntityPtr, EntityTypePtr);
+	return A3DRootBaseSet(EntityPtr, &RootBaseData);
 }
 
-bool IsEntityBaseWithGraphicsType(const A3DEntity* EntityPtr)
+A3DStatus SetRootBaseWithGraphics(const A3DRootBaseWithGraphicsData& RootBaseWithGraphicsData, A3DRootBaseWithGraphics* RootPtr)
 {
-	return (bool)A3DEntityIsBaseWithGraphicsType(EntityPtr);
+	return A3DRootBaseWithGraphicsSet(RootPtr, &RootBaseWithGraphicsData);
 }
 
-bool IsEntityBaseType(const A3DEntity* EntityPtr)
+A3DUns32 InsertGraphMaterial(const A3DGraphMaterialData& InMaterialData)
 {
-	return (bool)A3DEntityIsBaseType(EntityPtr);
+	A3DUns32 OutMaterialGenericIndex = A3D_DEFAULT_MATERIAL_INDEX;
+	if (A3DGlobalInsertGraphMaterial(&InMaterialData, &OutMaterialGenericIndex) != A3DStatus::A3D_SUCCESS)
+	{
+		return A3D_DEFAULT_MATERIAL_INDEX;
+	}
+	return OutMaterialGenericIndex;
 }
 
-bool IsMaterialTexture(const uint32 MaterialIndex)
+A3DUns32 InsertGraphRgbColor(const A3DGraphRgbColorData& InRgbColorData)
 {
-	A3DBool bIsTexture = false;
-	return A3DGlobalIsMaterialTexture(MaterialIndex, &bIsTexture) == A3DStatus::A3D_SUCCESS ? bool(bIsTexture) : false;
+	A3DUns32 OutIndexRgbColor = A3D_DEFAULT_COLOR_INDEX;
+	if (A3DGlobalInsertGraphRgbColor(&InRgbColorData, &OutIndexRgbColor) != A3DStatus::A3D_SUCCESS)
+	{
+		return A3D_DEFAULT_COLOR_INDEX;
+	}
+	return OutIndexRgbColor;
 }
 
-A3DEntity* GetPointerFromIndex(const uint32 Index, const A3DEEntityType Type)
+A3DUns32 InsertGraphStyle(const A3DGraphStyleData& InStyleData)
 {
-	A3DEntity* EntityPtr;
-	if (A3DMiscPointerFromIndexGet(Index, Type, &EntityPtr) != A3DStatus::A3D_SUCCESS)
+	A3DUns32 OutStyleIndex = A3D_DEFAULT_STYLE_INDEX;
+	if (A3DGlobalInsertGraphStyle(&InStyleData, &OutStyleIndex) != A3DStatus::A3D_SUCCESS)
 	{
-		return nullptr;
+		return A3D_DEFAULT_STYLE_INDEX;
 	}
-	return EntityPtr;
+	return OutStyleIndex;
 }
 
-A3DStatus HealBRep(A3DRiBrepModel** BRepToHeal, double Tolerance, A3DSewOptionsData const* SewOptions, A3DRiBrepModel*** OutNewBReps, uint32& OutNewBRepCount)
+A3DStatus ExportModelFileToPrcFile(const A3DAsmModelFile* ModelFile, const A3DRWParamsExportPrcData* ParamsExportData, const A3DUTF8Char* CADFileName, A3DRWParamsPrcWriteHelper** PrcWriteHelper)
 {
-	A3DUns32 NewBRepCount;
-	A3DStatus Status = A3DSewBrep(&BRepToHeal, 1, Tolerance, SewOptions, OutNewBReps, &NewBRepCount);
-	OutNewBRepCount = NewBRepCount;
-	return Status;
+	return A3DAsmModelFileExportToPrcFile(ModelFile, ParamsExportData, CADFileName, PrcWriteHelper);
 }
 
-A3DStatus SewBReps(A3DRiBrepModel*** BRepsToSew, uint32 const BRepCount, double Tolerance, A3DSewOptionsData const* SewOptions, A3DRiBrepModel*** OutNewBReps, uint32& OutNewBRepCount)
+A3DStatus SewModel(A3DAsmModelFile* ModelPtr, double Tolerance, A3DSewOptionsData const* SewOptions)
 {
-	A3DUns32 NewBRepCount;
-	A3DStatus Status = A3DSewBrep(BRepsToSew, BRepCount, Tolerance, SewOptions, OutNewBReps, &NewBRepCount);
-	OutNewBRepCount = NewBRepCount;
-	return Status;
+	return A3DAsmModelFileSew(&ModelPtr, Tolerance, SewOptions);
 }
 
-A3DStatus SewModel(A3DAsmModelFile** ModelPtr, double Tolerance, A3DSewOptionsData const* SewOptions)
+A3DStatus SewBReps(A3DRiBrepModel** BRepsToSew, uint32 const BRepCount, double Tolerance, A3DSewOptionsData const* SewOptions, A3DRiBrepModel*** OutNewBReps, uint32& OutNewBRepCount)
 {
-	return A3DAsmModelFileSew(ModelPtr, Tolerance, SewOptions);
+	A3DUns32 NewBRepCount = 0;
+	A3DStatus Status = A3DSewBrep(&BRepsToSew, BRepCount, Tolerance, SewOptions, OutNewBReps, &NewBRepCount);
+	OutNewBRepCount = NewBRepCount;
+	return Status;
 }
 
 #endif
+} // NS TechSoftInterface
 
-}
 
 #ifdef USE_TECHSOFT_SDK
+
+// TUniqueTSObj InitializeData -----------------------------------
+
 void TUniqueTSObj<A3DAsmModelFileData>::InitializeData()
 {
 	A3D_INITIALIZE_DATA(A3DAsmModelFileData, Data);
 }
-const A3DEntity* TUniqueTSObj<A3DAsmModelFileData>::DefaultValue = nullptr;
 
 void TUniqueTSObj<A3DAsmPartDefinitionData>::InitializeData()
 {
 	A3D_INITIALIZE_DATA(A3DAsmPartDefinitionData, Data);
 }
-const A3DEntity* TUniqueTSObj<A3DAsmPartDefinitionData>::DefaultValue = nullptr;
 
 void TUniqueTSObj<A3DAsmProductOccurrenceData>::InitializeData()
 {
 	A3D_INITIALIZE_DATA(A3DAsmProductOccurrenceData, Data);
 }
-const A3DEntity* TUniqueTSObj<A3DAsmProductOccurrenceData>::DefaultValue = nullptr;
 
 void TUniqueTSObj<A3DAsmProductOccurrenceDataCV5>::InitializeData()
 {
 	A3D_INITIALIZE_DATA(A3DAsmProductOccurrenceDataCV5, Data);
 }
-const A3DEntity* TUniqueTSObj<A3DAsmProductOccurrenceDataCV5>::DefaultValue = nullptr;
 
 void TUniqueTSObj<A3DAsmProductOccurrenceDataSLW>::InitializeData()
 {
 	A3D_INITIALIZE_DATA(A3DAsmProductOccurrenceDataSLW, Data);
 }
-const A3DEntity* TUniqueTSObj<A3DAsmProductOccurrenceDataSLW>::DefaultValue = nullptr;
 
 void TUniqueTSObj<A3DAsmProductOccurrenceDataUg>::InitializeData()
 {
 	A3D_INITIALIZE_DATA(A3DAsmProductOccurrenceDataUg, Data);
 }
-const A3DEntity* TUniqueTSObj<A3DAsmProductOccurrenceDataUg>::DefaultValue = nullptr;
 
 void TUniqueTSObj<A3DBoundingBoxData>::InitializeData()
 {
 	A3D_INITIALIZE_DATA(A3DBoundingBoxData, Data);
 }
-const A3DEntity* TUniqueTSObj<A3DBoundingBoxData>::DefaultValue = nullptr;
 
 void TUniqueTSObj<A3DCrvCircleData>::InitializeData()
 {
 	A3D_INITIALIZE_DATA(A3DCrvCircleData, Data);
 }
-const A3DEntity* TUniqueTSObj<A3DCrvCircleData>::DefaultValue = nullptr;
 
 void TUniqueTSObj<A3DCrvCompositeData>::InitializeData()
 {
 	A3D_INITIALIZE_DATA(A3DCrvCompositeData, Data);
 }
-const A3DEntity* TUniqueTSObj<A3DCrvCompositeData>::DefaultValue = nullptr;
 
 void TUniqueTSObj<A3DCrvEllipseData>::InitializeData()
 {
 	A3D_INITIALIZE_DATA(A3DCrvEllipseData, Data);
 }
-const A3DEntity* TUniqueTSObj<A3DCrvEllipseData>::DefaultValue = nullptr;
 
 void TUniqueTSObj<A3DCrvHelixData>::InitializeData()
 {
 	A3D_INITIALIZE_DATA(A3DCrvHelixData, Data);
 }
-const A3DEntity* TUniqueTSObj<A3DCrvHelixData>::DefaultValue = nullptr;
 
 void TUniqueTSObj<A3DCrvHyperbolaData>::InitializeData()
 {
 	A3D_INITIALIZE_DATA(A3DCrvHyperbolaData, Data);
 }
-const A3DEntity* TUniqueTSObj<A3DCrvHyperbolaData>::DefaultValue = nullptr;
 
 void TUniqueTSObj<A3DCrvLineData>::InitializeData()
 {
 	A3D_INITIALIZE_DATA(A3DCrvLineData, Data);
 }
-const A3DEntity* TUniqueTSObj<A3DCrvLineData>::DefaultValue = nullptr;
 
 void TUniqueTSObj<A3DCrvNurbsData>::InitializeData()
 {
 	A3D_INITIALIZE_DATA(A3DCrvNurbsData, Data);
 }
-const A3DEntity* TUniqueTSObj<A3DCrvNurbsData>::DefaultValue = nullptr;
 
 void TUniqueTSObj<A3DCrvParabolaData>::InitializeData()
 {
 	A3D_INITIALIZE_DATA(A3DCrvParabolaData, Data);
 }
-const A3DEntity* TUniqueTSObj<A3DCrvParabolaData>::DefaultValue = nullptr;
 
 void TUniqueTSObj<A3DCrvPolyLineData>::InitializeData()
 {
 	A3D_INITIALIZE_DATA(A3DCrvPolyLineData, Data);
 }
-const A3DEntity* TUniqueTSObj<A3DCrvPolyLineData>::DefaultValue = nullptr;
+
+void TUniqueTSObj<A3DCrvTransformData>::InitializeData()
+{
+	A3D_INITIALIZE_DATA(A3DCrvTransformData, Data);
+}
 
 void TUniqueTSObj<A3DGlobalData>::InitializeData()
 {
 	A3D_INITIALIZE_DATA(A3DGlobalData, Data);
 }
-const A3DEntity* TUniqueTSObj<A3DGlobalData>::DefaultValue = nullptr;
 
 void TUniqueTSObj<A3DGraphicsData>::InitializeData()
 {
 	A3D_INITIALIZE_DATA(A3DGraphicsData, Data);
 }
-const A3DEntity* TUniqueTSObj<A3DGraphicsData>::DefaultValue = nullptr;
-
-void TUniqueTSObjFromIndex<A3DGraphMaterialData>::InitializeData()
-{
-	A3D_INITIALIZE_DATA(A3DGraphMaterialData, Data);
-}
-uint32 TUniqueTSObjFromIndex<A3DGraphMaterialData>::DefaultValue = A3D_DEFAULT_MATERIAL_INDEX;
-
-void TUniqueTSObjFromIndex<A3DGraphPictureData>::InitializeData()
-{
-	A3D_INITIALIZE_DATA(A3DGraphPictureData, Data);
-}
-uint32 TUniqueTSObjFromIndex<A3DGraphPictureData>::DefaultValue = A3D_DEFAULT_PICTURE_INDEX;
-
-void TUniqueTSObjFromIndex<A3DGraphRgbColorData>::InitializeData()
-{
-	A3D_INITIALIZE_DATA(A3DGraphRgbColorData, Data);
-}
-uint32 TUniqueTSObjFromIndex<A3DGraphRgbColorData>::DefaultValue = A3D_DEFAULT_COLOR_INDEX;
-
-void TUniqueTSObjFromIndex<A3DGraphStyleData>::InitializeData()
-{
-	A3D_INITIALIZE_DATA(A3DGraphStyleData, Data);
-}
-uint32 TUniqueTSObjFromIndex<A3DGraphStyleData>::DefaultValue = A3D_DEFAULT_STYLE_INDEX;
-
-void TUniqueTSObjFromIndex<A3DGraphTextureApplicationData>::InitializeData()
-{
-	A3D_INITIALIZE_DATA(A3DGraphTextureApplicationData, Data);
-}
-uint32 TUniqueTSObjFromIndex<A3DGraphTextureApplicationData>::DefaultValue = A3D_DEFAULT_TEXTURE_APPLICATION_INDEX;
-
-void TUniqueTSObjFromIndex<A3DGraphTextureDefinitionData>::InitializeData()
-{
-	A3D_INITIALIZE_DATA(A3DGraphTextureDefinitionData, Data);
-}
-uint32 TUniqueTSObjFromIndex<A3DGraphTextureDefinitionData>::DefaultValue = A3D_DEFAULT_TEXTURE_DEFINITION_INDEX;
 
 void TUniqueTSObj<A3DIntervalData>::InitializeData()
 {
 	A3D_INITIALIZE_DATA(A3DIntervalData, Data);
 }
-const A3DEntity* TUniqueTSObj<A3DIntervalData>::DefaultValue = nullptr;
 
 void TUniqueTSObj<A3DMiscAttributeData>::InitializeData()
 {
 	A3D_INITIALIZE_DATA(A3DMiscAttributeData, Data);
 }
-const A3DEntity* TUniqueTSObj<A3DMiscAttributeData>::DefaultValue = nullptr;
 
 void TUniqueTSObj<A3DMiscCartesianTransformationData>::InitializeData()
 {
 	A3D_INITIALIZE_DATA(A3DMiscCartesianTransformationData, Data);
 }
-const A3DEntity* TUniqueTSObj<A3DMiscCartesianTransformationData>::DefaultValue = nullptr;
 
 void TUniqueTSObj<A3DMiscEntityReferenceData>::InitializeData()
 {
 	A3D_INITIALIZE_DATA(A3DMiscEntityReferenceData, Data);
 }
-const A3DEntity* TUniqueTSObj<A3DMiscEntityReferenceData>::DefaultValue = nullptr;
 
 void TUniqueTSObj<A3DMiscGeneralTransformationData>::InitializeData()
 {
 	A3D_INITIALIZE_DATA(A3DMiscGeneralTransformationData, Data);
 }
-const A3DEntity* TUniqueTSObj<A3DMiscGeneralTransformationData>::DefaultValue = nullptr;
 
 void TUniqueTSObj<A3DMiscMaterialPropertiesData>::InitializeData()
 {
 	A3D_INITIALIZE_DATA(A3DMiscMaterialPropertiesData, Data);
 }
-const A3DEntity* TUniqueTSObj<A3DMiscMaterialPropertiesData>::DefaultValue = nullptr;
 
 void TUniqueTSObj<A3DMiscReferenceOnCsysItemData>::InitializeData()
 {
 	A3D_INITIALIZE_DATA(A3DMiscReferenceOnCsysItemData, Data);
 }
-const A3DEntity* TUniqueTSObj<A3DMiscReferenceOnCsysItemData>::DefaultValue = nullptr;
+
+void TUniqueTSObj<A3DMiscReferenceOnTessData>::InitializeData()
+{
+	A3D_INITIALIZE_DATA(A3DMiscReferenceOnTessData, Data);
+}
 
 void TUniqueTSObj<A3DMiscReferenceOnTopologyData>::InitializeData()
 {
 	A3D_INITIALIZE_DATA(A3DMiscReferenceOnTopologyData, Data);
 }
-const A3DEntity* TUniqueTSObj<A3DMiscReferenceOnTopologyData>::DefaultValue = nullptr;
 
-void TUniqueTSObj<A3DMiscReferenceOnTessData>::InitializeData()
+void TUniqueTSObj<A3DMiscSingleAttributeData>::InitializeData()
 {
-	A3D_INITIALIZE_DATA(A3DMiscReferenceOnTessData, Data);
+	A3D_INITIALIZE_DATA(A3DMiscSingleAttributeData, Data);
+}
+
+void TUniqueTSObj<A3DRWParamsExportPrcData>::InitializeData()
+{
+	A3D_INITIALIZE_DATA(A3DRWParamsExportPrcData, Data);
 }
-const A3DEntity* TUniqueTSObj<A3DMiscReferenceOnTessData>::DefaultValue = nullptr;
 
 void TUniqueTSObj<A3DRiBrepModelData>::InitializeData()
 {
 	A3D_INITIALIZE_DATA(A3DRiBrepModelData, Data);
 }
-const A3DEntity* TUniqueTSObj<A3DRiBrepModelData>::DefaultValue = nullptr;
 
 void TUniqueTSObj<A3DRiCoordinateSystemData>::InitializeData()
 {
 	A3D_INITIALIZE_DATA(A3DRiCoordinateSystemData, Data);
 }
-const A3DEntity* TUniqueTSObj<A3DRiCoordinateSystemData>::DefaultValue = nullptr;
 
 void TUniqueTSObj<A3DRiDirectionData>::InitializeData()
 {
 	A3D_INITIALIZE_DATA(A3DRiDirectionData, Data);
 }
-const A3DEntity* TUniqueTSObj<A3DRiDirectionData>::DefaultValue = nullptr;
 
 void TUniqueTSObj<A3DRiPolyBrepModelData>::InitializeData()
 {
 	A3D_INITIALIZE_DATA(A3DRiPolyBrepModelData, Data);
 }
-const A3DEntity* TUniqueTSObj<A3DRiPolyBrepModelData>::DefaultValue = nullptr;
 
 void TUniqueTSObj<A3DRiRepresentationItemData>::InitializeData()
 {
 	A3D_INITIALIZE_DATA(A3DRiRepresentationItemData, Data);
 }
-const A3DEntity* TUniqueTSObj<A3DRiRepresentationItemData>::DefaultValue = nullptr;
 
 void TUniqueTSObj<A3DRiSetData>::InitializeData()
 {
 	A3D_INITIALIZE_DATA(A3DRiSetData, Data);
 }
-const A3DEntity* TUniqueTSObj<A3DRiSetData>::DefaultValue = nullptr;
 
 void TUniqueTSObj<A3DRootBaseData>::InitializeData()
 {
 	A3D_INITIALIZE_DATA(A3DRootBaseData, Data);
 }
-const A3DEntity* TUniqueTSObj<A3DRootBaseData>::DefaultValue = nullptr;
 
 void TUniqueTSObj<A3DRootBaseWithGraphicsData>::InitializeData()
 {
 	A3D_INITIALIZE_DATA(A3DRootBaseWithGraphicsData, Data);
 }
-const A3DEntity* TUniqueTSObj<A3DRootBaseWithGraphicsData>::DefaultValue = nullptr;
+
+void TUniqueTSObj<A3DRWParamsTessellationData>::InitializeData()
+{
+	A3D_INITIALIZE_DATA(A3DRWParamsTessellationData, Data);
+}
 
 void TUniqueTSObj<A3DSewOptionsData>::InitializeData()
 {
 	A3D_INITIALIZE_DATA(A3DSewOptionsData, Data);
 }
-const A3DEntity* TUniqueTSObj<A3DSewOptionsData>::DefaultValue = nullptr;
 
 void TUniqueTSObj<A3DSurfBlend01Data>::InitializeData()
 {
 	A3D_INITIALIZE_DATA(A3DSurfBlend01Data, Data);
 }
-const A3DEntity* TUniqueTSObj<A3DSurfBlend01Data>::DefaultValue = nullptr;
 
 void TUniqueTSObj<A3DSurfBlend02Data>::InitializeData()
 {
 	A3D_INITIALIZE_DATA(A3DSurfBlend02Data, Data);
 }
-const A3DEntity* TUniqueTSObj<A3DSurfBlend02Data>::DefaultValue = nullptr;
 
 void TUniqueTSObj<A3DSurfBlend03Data>::InitializeData()
 {
 	A3D_INITIALIZE_DATA(A3DSurfBlend03Data, Data);
 }
-const A3DEntity* TUniqueTSObj<A3DSurfBlend03Data>::DefaultValue = nullptr;
 
 void TUniqueTSObj<A3DSurfConeData>::InitializeData()
 {
 	A3D_INITIALIZE_DATA(A3DSurfConeData, Data);
 }
-const A3DEntity* TUniqueTSObj<A3DSurfConeData>::DefaultValue = nullptr;
 
 void TUniqueTSObj<A3DSurfCylinderData>::InitializeData()
 {
 	A3D_INITIALIZE_DATA(A3DSurfCylinderData, Data);
 }
-const A3DEntity* TUniqueTSObj<A3DSurfCylinderData>::DefaultValue = nullptr;
 
 void TUniqueTSObj<A3DSurfCylindricalData>::InitializeData()
 {
 	A3D_INITIALIZE_DATA(A3DSurfCylindricalData, Data);
 }
-const A3DEntity* TUniqueTSObj<A3DSurfCylindricalData>::DefaultValue = nullptr;
 
 void TUniqueTSObj<A3DSurfExtrusionData>::InitializeData()
 {
 	A3D_INITIALIZE_DATA(A3DSurfExtrusionData, Data);
 }
-const A3DEntity* TUniqueTSObj<A3DSurfExtrusionData>::DefaultValue = nullptr;
 
 void TUniqueTSObj<A3DSurfFromCurvesData>::InitializeData()
 {
 	A3D_INITIALIZE_DATA(A3DSurfFromCurvesData, Data);
 }
-const A3DEntity* TUniqueTSObj<A3DSurfFromCurvesData>::DefaultValue = nullptr;
 
 void TUniqueTSObj<A3DSurfNurbsData>::InitializeData()
 {
 	A3D_INITIALIZE_DATA(A3DSurfNurbsData, Data);
 }
-const A3DEntity* TUniqueTSObj<A3DSurfNurbsData>::DefaultValue = nullptr;
 
 void TUniqueTSObj<A3DSurfPipeData>::InitializeData()
 {
 	A3D_INITIALIZE_DATA(A3DSurfPipeData, Data);
 }
-const A3DEntity* TUniqueTSObj<A3DSurfPipeData>::DefaultValue = nullptr;
 
 void TUniqueTSObj<A3DSurfPlaneData>::InitializeData()
 {
 	A3D_INITIALIZE_DATA(A3DSurfPlaneData, Data);
 }
-const A3DEntity* TUniqueTSObj<A3DSurfPlaneData>::DefaultValue = nullptr;
 
 void TUniqueTSObj<A3DSurfRevolutionData>::InitializeData()
 {
 	A3D_INITIALIZE_DATA(A3DSurfRevolutionData, Data);
 }
-const A3DEntity* TUniqueTSObj<A3DSurfRevolutionData>::DefaultValue = nullptr;
 
 void TUniqueTSObj<A3DSurfRuledData>::InitializeData()
 {
 	A3D_INITIALIZE_DATA(A3DSurfRuledData, Data);
 }
-const A3DEntity* TUniqueTSObj<A3DSurfRuledData>::DefaultValue = nullptr;
 
 void TUniqueTSObj<A3DSurfSphereData>::InitializeData()
 {
 	A3D_INITIALIZE_DATA(A3DSurfSphereData, Data);
 }
-const A3DEntity* TUniqueTSObj<A3DSurfSphereData>::DefaultValue = nullptr;
 
 void TUniqueTSObj<A3DSurfTorusData>::InitializeData()
 {
 	A3D_INITIALIZE_DATA(A3DSurfTorusData, Data);
 }
-const A3DEntity* TUniqueTSObj<A3DSurfTorusData>::DefaultValue = nullptr;
 
 void TUniqueTSObj<A3DTess3DData>::InitializeData()
 {
 	A3D_INITIALIZE_DATA(A3DTess3DData, Data);
 }
-const A3DEntity* TUniqueTSObj<A3DTess3DData>::DefaultValue = nullptr;
 
 void TUniqueTSObj<A3DTessBaseData>::InitializeData()
 {
 	A3D_INITIALIZE_DATA(A3DTessBaseData, Data);
 }
-const A3DEntity* TUniqueTSObj<A3DTessBaseData>::DefaultValue = nullptr;
 
 void TUniqueTSObj<A3DTopoBodyData>::InitializeData()
 {
 	A3D_INITIALIZE_DATA(A3DTopoBodyData, Data);
 }
-const A3DEntity* TUniqueTSObj<A3DTopoBodyData>::DefaultValue = nullptr;
 
 void TUniqueTSObj<A3DTopoBrepDataData>::InitializeData()
 {
 	A3D_INITIALIZE_DATA(A3DTopoBrepDataData, Data);
 }
-const A3DEntity* TUniqueTSObj<A3DTopoBrepDataData>::DefaultValue = nullptr;
 
 void TUniqueTSObj<A3DTopoCoEdgeData>::InitializeData()
 {
 	A3D_INITIALIZE_DATA(A3DTopoCoEdgeData, Data);
 }
-const A3DEntity* TUniqueTSObj<A3DTopoCoEdgeData>::DefaultValue = nullptr;
 
 void TUniqueTSObj<A3DTopoConnexData>::InitializeData()
 {
 	A3D_INITIALIZE_DATA(A3DTopoConnexData, Data);
 }
-const A3DEntity* TUniqueTSObj<A3DTopoConnexData>::DefaultValue = nullptr;
 
 void TUniqueTSObj<A3DTopoContextData>::InitializeData()
 {
 	A3D_INITIALIZE_DATA(A3DTopoContextData, Data);
 }
-const A3DEntity* TUniqueTSObj<A3DTopoContextData>::DefaultValue = nullptr;
 
 void TUniqueTSObj<A3DTopoEdgeData>::InitializeData()
 {
 	A3D_INITIALIZE_DATA(A3DTopoEdgeData, Data);
 }
-const A3DEntity* TUniqueTSObj<A3DTopoEdgeData>::DefaultValue = nullptr;
 
 void TUniqueTSObj<A3DTopoFaceData>::InitializeData()
 {
 	A3D_INITIALIZE_DATA(A3DTopoFaceData, Data);
 }
-const A3DEntity* TUniqueTSObj<A3DTopoFaceData>::DefaultValue = nullptr;
 
 void TUniqueTSObj<A3DTopoLoopData>::InitializeData()
 {
 	A3D_INITIALIZE_DATA(A3DTopoLoopData, Data);
 }
-const A3DEntity* TUniqueTSObj<A3DTopoLoopData>::DefaultValue = nullptr;
 
 void TUniqueTSObj<A3DTopoShellData>::InitializeData()
 {
 	A3D_INITIALIZE_DATA(A3DTopoShellData, Data);
 }
-const A3DEntity* TUniqueTSObj<A3DTopoShellData>::DefaultValue = nullptr;
 
 void TUniqueTSObj<A3DTopoUniqueVertexData>::InitializeData()
 {
 	A3D_INITIALIZE_DATA(A3DTopoUniqueVertexData, Data);
 }
-const A3DEntity* TUniqueTSObj<A3DTopoUniqueVertexData>::DefaultValue = nullptr;
 
 void TUniqueTSObj<A3DTopoWireEdgeData>::InitializeData()
 {
 	A3D_INITIALIZE_DATA(A3DTopoWireEdgeData, Data);
 }
-const A3DEntity* TUniqueTSObj<A3DTopoWireEdgeData>::DefaultValue = nullptr;
 
 void TUniqueTSObj<A3DUTF8Char*>::InitializeData()
 {
 	Data = nullptr;
 }
-const A3DEntity* TUniqueTSObj<A3DUTF8Char*>::DefaultValue = nullptr;
 
-A3DStatus TUniqueTSObj<A3DAsmModelFileData>::GetData(const A3DAsmModelFile* InAsmModelFilePtr)
+
+
+// TUniqueTSObjFromIndex InitializeData -----------------------------------
+
+
+void TUniqueTSObjFromIndex<A3DGraphMaterialData>::InitializeData()
+{
+	A3D_INITIALIZE_DATA(A3DGraphMaterialData, Data);
+}
+
+void TUniqueTSObjFromIndex<A3DGraphPictureData>::InitializeData()
+{
+	A3D_INITIALIZE_DATA(A3DGraphPictureData, Data);
+}
+
+void TUniqueTSObjFromIndex<A3DGraphRgbColorData>::InitializeData()
+{
+	A3D_INITIALIZE_DATA(A3DGraphRgbColorData, Data);
+}
+
+void TUniqueTSObjFromIndex<A3DGraphStyleData>::InitializeData()
+{
+	A3D_INITIALIZE_DATA(A3DGraphStyleData, Data);
+}
+
+void TUniqueTSObjFromIndex<A3DGraphTextureApplicationData>::InitializeData()
 {
-	return A3DAsmModelFileGet(InAsmModelFilePtr, &Data);
+	A3D_INITIALIZE_DATA(A3DGraphTextureApplicationData, Data);
 }
 
-A3DStatus TUniqueTSObj<A3DAsmPartDefinitionData>::GetData(const A3DAsmPartDefinition* InAsmPartDefinitionPtr)
+void TUniqueTSObjFromIndex<A3DGraphTextureDefinitionData>::InitializeData()
 {
-	return A3DAsmPartDefinitionGet(InAsmPartDefinitionPtr, &Data);
+	A3D_INITIALIZE_DATA(A3DGraphTextureDefinitionData, Data);
 }
 
-A3DStatus TUniqueTSObj<A3DAsmProductOccurrenceData>::GetData(const A3DAsmProductOccurrence* InAsmProductOccurrencePtr)
+
+
+// TUniqueTSObj GetData -----------------------------------
+
+
+A3DStatus TUniqueTSObj<A3DAsmModelFileData>::GetData(const A3DEntity* InEntityPtr)
 {
-	return A3DAsmProductOccurrenceGet(InAsmProductOccurrencePtr, &Data);
+	return A3DAsmModelFileGet(InEntityPtr, &Data);
 }
 
-A3DStatus TUniqueTSObj<A3DAsmProductOccurrenceDataCV5>::GetData(const A3DAsmProductOccurrence* InAsmProductOccurrencePtr)
+A3DStatus TUniqueTSObj<A3DAsmPartDefinitionData>::GetData(const A3DEntity* InEntityPtr)
 {
-	return A3DAsmProductOccurrenceGetCV5(InAsmProductOccurrencePtr, &Data);
+	return A3DAsmPartDefinitionGet(InEntityPtr, &Data);
 }
 
-A3DStatus TUniqueTSObj<A3DAsmProductOccurrenceDataSLW>::GetData(const A3DAsmProductOccurrence* InAsmProductOccurrencePtr)
+A3DStatus TUniqueTSObj<A3DAsmProductOccurrenceData>::GetData(const A3DEntity* InEntityPtr)
 {
-	return A3DAsmProductOccurrenceGetSLW(InAsmProductOccurrencePtr, &Data);
+	return A3DAsmProductOccurrenceGet(InEntityPtr, &Data);
 }
 
-A3DStatus TUniqueTSObj<A3DAsmProductOccurrenceDataUg>::GetData(const A3DAsmProductOccurrence* InAsmProductOccurrencePtr)
+A3DStatus TUniqueTSObj<A3DAsmProductOccurrenceDataCV5>::GetData(const A3DEntity* InEntityPtr)
 {
-	return A3DAsmProductOccurrenceGetUg(InAsmProductOccurrencePtr, &Data);
+	return A3DAsmProductOccurrenceGetCV5(InEntityPtr, &Data);
 }
 
-A3DStatus TUniqueTSObj<A3DBoundingBoxData>::GetData(const A3DEntity* InGraphicsPtr)
+A3DStatus TUniqueTSObj<A3DAsmProductOccurrenceDataSLW>::GetData(const A3DEntity* InEntityPtr)
 {
-	return A3DMiscGetBoundingBox(InGraphicsPtr, &Data);
+	return A3DAsmProductOccurrenceGetSLW(InEntityPtr, &Data);
 }
 
-A3DStatus TUniqueTSObj<A3DCrvCircleData>::GetData(const A3DEntity* InCurvePtr)
+A3DStatus TUniqueTSObj<A3DAsmProductOccurrenceDataUg>::GetData(const A3DEntity* InEntityPtr)
 {
-	return A3DCrvCircleGet(InCurvePtr, &Data);
+	return A3DAsmProductOccurrenceGetUg(InEntityPtr, &Data);
 }
 
-A3DStatus TUniqueTSObj<A3DCrvCompositeData>::GetData(const A3DEntity* InCurvePtr)
+A3DStatus TUniqueTSObj<A3DBoundingBoxData>::GetData(const A3DEntity* InEntityPtr)
 {
-	return A3DCrvCompositeGet(InCurvePtr, &Data);
+	return A3DMiscGetBoundingBox(InEntityPtr, &Data);
 }
 
-A3DStatus TUniqueTSObj<A3DCrvEllipseData>::GetData(const A3DEntity* InCurvePtr)
+A3DStatus TUniqueTSObj<A3DCrvCircleData>::GetData(const A3DEntity* InEntityPtr)
 {
-	return A3DCrvEllipseGet(InCurvePtr, &Data);
+	return A3DCrvCircleGet(InEntityPtr, &Data);
 }
 
-A3DStatus TUniqueTSObj<A3DCrvHelixData>::GetData(const A3DEntity* InCurvePtr)
+A3DStatus TUniqueTSObj<A3DCrvCompositeData>::GetData(const A3DEntity* InEntityPtr)
 {
-	return A3DCrvHelixGet(InCurvePtr, &Data);
+	return A3DCrvCompositeGet(InEntityPtr, &Data);
 }
 
-A3DStatus TUniqueTSObj<A3DCrvHyperbolaData>::GetData(const A3DEntity* InCurvePtr)
+A3DStatus TUniqueTSObj<A3DCrvEllipseData>::GetData(const A3DEntity* InEntityPtr)
 {
-	return A3DCrvHyperbolaGet(InCurvePtr, &Data);
+	return A3DCrvEllipseGet(InEntityPtr, &Data);
 }
 
-A3DStatus TUniqueTSObj<A3DCrvLineData>::GetData(const A3DEntity* InCurvePtr)
+A3DStatus TUniqueTSObj<A3DCrvHelixData>::GetData(const A3DEntity* InEntityPtr)
 {
-	return A3DCrvLineGet(InCurvePtr, &Data);
+	return A3DCrvHelixGet(InEntityPtr, &Data);
 }
 
-A3DStatus TUniqueTSObj<A3DCrvNurbsData>::GetData(const A3DEntity* InCurvePtr)
+A3DStatus TUniqueTSObj<A3DCrvHyperbolaData>::GetData(const A3DEntity* InEntityPtr)
 {
-	return A3DCrvNurbsGet(InCurvePtr, &Data);
+	return A3DCrvHyperbolaGet(InEntityPtr, &Data);
 }
 
-A3DStatus TUniqueTSObj<A3DCrvParabolaData>::GetData(const A3DEntity* InCurvePtr)
+A3DStatus TUniqueTSObj<A3DCrvLineData>::GetData(const A3DEntity* InEntityPtr)
 {
-	return A3DCrvParabolaGet(InCurvePtr, &Data);
+	return A3DCrvLineGet(InEntityPtr, &Data);
 }
 
-A3DStatus TUniqueTSObj<A3DCrvPolyLineData>::GetData(const A3DEntity* InCurvePtr)
+A3DStatus TUniqueTSObj<A3DCrvNurbsData>::GetData(const A3DEntity* InEntityPtr)
 {
-	return A3DCrvPolyLineGet(InCurvePtr, &Data);
+	return A3DCrvNurbsGet(InEntityPtr, &Data);
 }
 
-A3DStatus TUniqueTSObj<A3DGlobalData>::GetData(const A3DEntity* GlobalPtr)
+A3DStatus TUniqueTSObj<A3DCrvParabolaData>::GetData(const A3DEntity* InEntityPtr)
 {
-	return A3DGlobalGet(GlobalPtr, &Data);
+	return A3DCrvParabolaGet(InEntityPtr, &Data);
 }
 
-A3DStatus TUniqueTSObj<A3DGraphicsData>::GetData(const A3DGraphics* InGraphicsPtr)
+A3DStatus TUniqueTSObj<A3DCrvPolyLineData>::GetData(const A3DEntity* InEntityPtr)
 {
-	return A3DGraphicsGet(InGraphicsPtr, &Data);
+	return A3DCrvPolyLineGet(InEntityPtr, &Data);
 }
 
-A3DStatus TUniqueTSObjFromIndex<A3DGraphMaterialData>::GetData(uint32 InGraphicsIndex)
+A3DStatus TUniqueTSObj<A3DCrvTransformData>::GetData(const A3DEntity* InEntityPtr)
 {
-	return A3DGlobalGetGraphMaterialData(InGraphicsIndex, &Data);
+	return A3DCrvTransformGet(InEntityPtr, &Data);
 }
 
-A3DStatus TUniqueTSObjFromIndex<A3DGraphPictureData>::GetData(uint32 InGraphicsIndex)
+A3DStatus TUniqueTSObj<A3DGlobalData>::GetData(const A3DEntity* InEntityPtr)
 {
-	return A3DGlobalGetGraphPictureData(InGraphicsIndex, &Data);
+	return A3DGlobalGet(InEntityPtr, &Data);
 }
 
-A3DStatus TUniqueTSObjFromIndex<A3DGraphRgbColorData>::GetData(uint32 ObjectIndex)
+A3DStatus TUniqueTSObj<A3DGraphicsData>::GetData(const A3DEntity* InEntityPtr)
 {
-	return A3DGlobalGetGraphRgbColorData(ObjectIndex, &Data);
+	return A3DGraphicsGet(InEntityPtr, &Data);
 }
 
-A3DStatus TUniqueTSObjFromIndex<A3DGraphStyleData>::GetData(uint32 ObjectIndex)
+A3DStatus TUniqueTSObj<A3DIntervalData>::GetData(const A3DEntity* InEntityPtr)
 {
-	return A3DGlobalGetGraphStyleData(ObjectIndex, &Data);
+	return A3DCrvGetInterval(InEntityPtr, &Data);
 }
 
-A3DStatus TUniqueTSObjFromIndex<A3DGraphTextureApplicationData>::GetData(uint32 ObjectIndex)
+A3DStatus TUniqueTSObj<A3DMiscAttributeData>::GetData(const A3DEntity* InEntityPtr)
 {
-	return A3DGlobalGetGraphTextureApplicationData(ObjectIndex, &Data);
+	return A3DMiscAttributeGet(InEntityPtr, &Data);
 }
 
-A3DStatus TUniqueTSObjFromIndex<A3DGraphTextureDefinitionData>::GetData(uint32 ObjectIndex)
+A3DStatus TUniqueTSObj<A3DMiscCartesianTransformationData>::GetData(const A3DEntity* InEntityPtr)
 {
-	return A3DGlobalGetGraphTextureDefinitionData(ObjectIndex, &Data);
+	return A3DMiscCartesianTransformationGet(InEntityPtr, &Data);
 }
 
-A3DStatus TUniqueTSObj<A3DIntervalData>::GetData(const A3DEntity* InCurvePtr)
+A3DStatus TUniqueTSObj<A3DMiscEntityReferenceData>::GetData(const A3DEntity* InEntityPtr)
 {
-	return A3DCrvGetInterval(InCurvePtr, &Data);
+	return A3DMiscEntityReferenceGet(InEntityPtr, &Data);
 }
 
-A3DStatus TUniqueTSObj<A3DMiscAttributeData>::GetData(const A3DMiscAttribute* InMiscAttributePtr)
+A3DStatus TUniqueTSObj<A3DMiscGeneralTransformationData>::GetData(const A3DEntity* InEntityPtr)
 {
-	return A3DMiscAttributeGet(InMiscAttributePtr, &Data);
+	return A3DMiscGeneralTransformationGet(InEntityPtr, &Data);
 }
 
-A3DStatus TUniqueTSObj<A3DMiscCartesianTransformationData>::GetData(const A3DMiscCartesianTransformation* InMiscCartesianTransformationPtr)
+A3DStatus TUniqueTSObj<A3DMiscMaterialPropertiesData>::GetData(const A3DEntity* InEntityPtr)
 {
-	return A3DMiscCartesianTransformationGet(InMiscCartesianTransformationPtr, &Data);
+	return A3DMiscGetMaterialProperties(InEntityPtr, &Data);
 }
 
-A3DStatus TUniqueTSObj<A3DMiscEntityReferenceData>::GetData(const A3DMiscEntityReference* InMiscEntityReferencePtr)
+A3DStatus TUniqueTSObj<A3DMiscReferenceOnCsysItemData>::GetData(const A3DEntity* InEntityPtr)
 {
-	return A3DMiscEntityReferenceGet(InMiscEntityReferencePtr, &Data);
+	return A3DMiscReferenceOnCsysItemGet(InEntityPtr, &Data);
 }
 
-A3DStatus TUniqueTSObj<A3DMiscGeneralTransformationData>::GetData(const A3DMiscGeneralTransformation* InMiscGeneralTransformationPtr)
+A3DStatus TUniqueTSObj<A3DMiscReferenceOnTessData>::GetData(const A3DEntity* InEntityPtr)
 {
-	return A3DMiscGeneralTransformationGet(InMiscGeneralTransformationPtr, &Data);
+	return A3DMiscReferenceOnTessGet(InEntityPtr, &Data);
 }
 
-A3DStatus TUniqueTSObj<A3DMiscMaterialPropertiesData>::GetData(const A3DEntity* InMiscMaterialPropertiesPtr)
+A3DStatus TUniqueTSObj<A3DMiscReferenceOnTopologyData>::GetData(const A3DEntity* InEntityPtr)
 {
-	return A3DMiscGetMaterialProperties(InMiscMaterialPropertiesPtr, &Data);
+	return A3DMiscReferenceOnTopologyGet(InEntityPtr, &Data);
 }
 
-A3DStatus TUniqueTSObj<A3DMiscReferenceOnTopologyData>::GetData(const A3DEntity* InMiscReferenceOnTopologyPtr)
+A3DStatus TUniqueTSObj<A3DMiscSingleAttributeData>::GetData(const A3DEntity* InEntityPtr)
 {
-	return A3DMiscReferenceOnTopologyGet(InMiscReferenceOnTopologyPtr, &Data);
+	return A3DStatus::A3D_ERROR;
 }
 
-A3DStatus TUniqueTSObj<A3DMiscReferenceOnCsysItemData>::GetData(const A3DEntity* InMiscReferenceOnCsysItem)
+A3DStatus TUniqueTSObj<A3DRWParamsExportPrcData>::GetData(const A3DEntity* InEntityPtr)
 {
-	return A3DMiscReferenceOnCsysItemGet(InMiscReferenceOnCsysItem, &Data);
+	return A3DStatus::A3D_ERROR;
 }
 
-A3DStatus TUniqueTSObj<A3DMiscReferenceOnTessData>::GetData(const A3DEntity* InMiscReferenceOnTess)
+A3DStatus TUniqueTSObj<A3DRiBrepModelData>::GetData(const A3DEntity* InEntityPtr)
 {
-	return A3DMiscReferenceOnTessGet(InMiscReferenceOnTess, &Data);
+	return A3DRiBrepModelGet(InEntityPtr, &Data);
 }
 
-A3DStatus TUniqueTSObj<A3DRiBrepModelData>::GetData(const A3DRiBrepModel* InRiBrepModelPtr)
+A3DStatus TUniqueTSObj<A3DRiCoordinateSystemData>::GetData(const A3DEntity* InEntityPtr)
 {
-	return A3DRiBrepModelGet(InRiBrepModelPtr, &Data);
+	return A3DRiCoordinateSystemGet(InEntityPtr, &Data);
 }
 
-A3DStatus TUniqueTSObj<A3DRiCoordinateSystemData>::GetData(const A3DRiCoordinateSystem* InRiCoordinateSystemPtr)
+A3DStatus TUniqueTSObj<A3DRiDirectionData>::GetData(const A3DEntity* InEntityPtr)
 {
-	return A3DRiCoordinateSystemGet(InRiCoordinateSystemPtr, &Data);
+	return A3DRiDirectionGet(InEntityPtr, &Data);
 }
 
-A3DStatus TUniqueTSObj<A3DRiDirectionData>::GetData(const A3DRiDirection* InRiDirectionPtr)
+A3DStatus TUniqueTSObj<A3DRiPolyBrepModelData>::GetData(const A3DEntity* InEntityPtr)
 {
-	return A3DRiDirectionGet(InRiDirectionPtr, &Data);
+	return A3DRiPolyBrepModelGet(InEntityPtr, &Data);
 }
 
-A3DStatus TUniqueTSObj<A3DRiPolyBrepModelData>::GetData(const A3DRiPolyBrepModel* InRiPolyBrepModelPtr)
+A3DStatus TUniqueTSObj<A3DRiRepresentationItemData>::GetData(const A3DEntity* InEntityPtr)
 {
-	return A3DRiPolyBrepModelGet(InRiPolyBrepModelPtr, &Data);
+	return A3DRiRepresentationItemGet(InEntityPtr, &Data);
 }
 
-A3DStatus TUniqueTSObj<A3DRiRepresentationItemData>::GetData(const A3DRiRepresentationItem* InRiRepresentationItemPtr)
+A3DStatus TUniqueTSObj<A3DRiSetData>::GetData(const A3DEntity* InEntityPtr)
 {
-	return A3DRiRepresentationItemGet(InRiRepresentationItemPtr, &Data);
+	return A3DRiSetGet(InEntityPtr, &Data);
 }
 
-A3DStatus TUniqueTSObj<A3DRiSetData>::GetData(const A3DRiSet* InRiSetPtr)
+A3DStatus TUniqueTSObj<A3DRootBaseData>::GetData(const A3DEntity* InEntityPtr)
 {
-	return A3DRiSetGet(InRiSetPtr, &Data);
+	return A3DRootBaseGet(InEntityPtr, &Data);
 }
 
-A3DStatus TUniqueTSObj<A3DRootBaseData>::GetData(const A3DRootBase* InRootBasePtr)
+A3DStatus TUniqueTSObj<A3DRootBaseWithGraphicsData>::GetData(const A3DEntity* InEntityPtr)
 {
-	return A3DRootBaseGet(InRootBasePtr, &Data);
+	return A3DRootBaseWithGraphicsGet(InEntityPtr, &Data);
 }
 
-A3DStatus TUniqueTSObj<A3DRootBaseWithGraphicsData>::GetData(const A3DRootBaseWithGraphics* InRootBaseWithGraphicsPtr)
+A3DStatus TUniqueTSObj<A3DRWParamsTessellationData>::GetData(const A3DEntity* InEntityPtr)
 {
-	return A3DRootBaseWithGraphicsGet(InRootBaseWithGraphicsPtr, &Data);
+	return A3DStatus::A3D_ERROR;
 }
 
-A3DStatus TUniqueTSObj<A3DSewOptionsData>::GetData(const A3DEntity* InRootBaseWithGraphicsPtr)
+A3DStatus TUniqueTSObj<A3DSewOptionsData>::GetData(const A3DEntity* InEntityPtr)
 {
 	return A3DStatus::A3D_ERROR;
 }
 
-A3DStatus TUniqueTSObj<A3DSurfBlend01Data>::GetData(const A3DEntity* InSurfPtr)
+A3DStatus TUniqueTSObj<A3DSurfBlend01Data>::GetData(const A3DEntity* InEntityPtr)
 {
-	return A3DSurfBlend01Get(InSurfPtr, &Data);
+	return A3DSurfBlend01Get(InEntityPtr, &Data);
 }
 
-A3DStatus TUniqueTSObj<A3DSurfBlend02Data>::GetData(const A3DEntity* InSurfacePtr)
+A3DStatus TUniqueTSObj<A3DSurfBlend02Data>::GetData(const A3DEntity* InEntityPtr)
 {
-	return A3DSurfBlend02Get(InSurfacePtr, &Data);
+	return A3DSurfBlend02Get(InEntityPtr, &Data);
 }
 
-A3DStatus TUniqueTSObj<A3DSurfBlend03Data>::GetData(const A3DEntity* InSurfacePtr)
+A3DStatus TUniqueTSObj<A3DSurfBlend03Data>::GetData(const A3DEntity* InEntityPtr)
 {
-	return A3DSurfBlend03Get(InSurfacePtr, &Data);
+	return A3DSurfBlend03Get(InEntityPtr, &Data);
 }
 
-A3DStatus TUniqueTSObj<A3DSurfConeData>::GetData(const A3DEntity* InSurfacePtr)
+A3DStatus TUniqueTSObj<A3DSurfConeData>::GetData(const A3DEntity* InEntityPtr)
 {
-	return A3DSurfConeGet(InSurfacePtr, &Data);
+	return A3DSurfConeGet(InEntityPtr, &Data);
 }
 
-A3DStatus TUniqueTSObj<A3DSurfCylinderData>::GetData(const A3DEntity* InSurfacePtr)
+A3DStatus TUniqueTSObj<A3DSurfCylinderData>::GetData(const A3DEntity* InEntityPtr)
 {
-	return A3DSurfCylinderGet(InSurfacePtr, &Data);
+	return A3DSurfCylinderGet(InEntityPtr, &Data);
 }
 
-A3DStatus TUniqueTSObj<A3DSurfCylindricalData>::GetData(const A3DEntity* InSurfacePtr)
+A3DStatus TUniqueTSObj<A3DSurfCylindricalData>::GetData(const A3DEntity* InEntityPtr)
 {
-	return A3DSurfCylindricalGet(InSurfacePtr, &Data);
+	return A3DSurfCylindricalGet(InEntityPtr, &Data);
 }
 
-A3DStatus TUniqueTSObj<A3DSurfExtrusionData>::GetData(const A3DEntity* InSurfacePtr)
+A3DStatus TUniqueTSObj<A3DSurfExtrusionData>::GetData(const A3DEntity* InEntityPtr)
 {
-	return A3DSurfExtrusionGet(InSurfacePtr, &Data);
+	return A3DSurfExtrusionGet(InEntityPtr, &Data);
 }
 
-A3DStatus TUniqueTSObj<A3DSurfFromCurvesData>::GetData(const A3DEntity* InSurfacePtr)
+A3DStatus TUniqueTSObj<A3DSurfFromCurvesData>::GetData(const A3DEntity* InEntityPtr)
 {
-	return A3DSurfFromCurvesGet(InSurfacePtr, &Data);
+	return A3DSurfFromCurvesGet(InEntityPtr, &Data);
 }
 
-A3DStatus TUniqueTSObj<A3DSurfNurbsData>::GetData(const A3DEntity* InSurfacePtr)
+A3DStatus TUniqueTSObj<A3DSurfNurbsData>::GetData(const A3DEntity* InEntityPtr)
 {
-	return A3DSurfNurbsGet(InSurfacePtr, &Data);
+	return A3DSurfNurbsGet(InEntityPtr, &Data);
 }
 
-A3DStatus TUniqueTSObj<A3DSurfPipeData>::GetData(const A3DEntity* InSurfacePtr)
+A3DStatus TUniqueTSObj<A3DSurfPipeData>::GetData(const A3DEntity* InEntityPtr)
 {
-	return A3DSurfPipeGet(InSurfacePtr, &Data);
+	return A3DSurfPipeGet(InEntityPtr, &Data);
 }
 
-A3DStatus TUniqueTSObj<A3DSurfPlaneData>::GetData(const A3DEntity* InSurfacePtr)
+A3DStatus TUniqueTSObj<A3DSurfPlaneData>::GetData(const A3DEntity* InEntityPtr)
 {
-	return A3DSurfPlaneGet(InSurfacePtr, &Data);
+	return A3DSurfPlaneGet(InEntityPtr, &Data);
 }
 
-A3DStatus TUniqueTSObj<A3DSurfRevolutionData>::GetData(const A3DEntity* InSurfacePtr)
+A3DStatus TUniqueTSObj<A3DSurfRevolutionData>::GetData(const A3DEntity* InEntityPtr)
 {
-	return A3DSurfRevolutionGet(InSurfacePtr, &Data);
+	return A3DSurfRevolutionGet(InEntityPtr, &Data);
 }
 
-A3DStatus TUniqueTSObj<A3DSurfRuledData>::GetData(const A3DEntity* InSurfacePtr)
+A3DStatus TUniqueTSObj<A3DSurfRuledData>::GetData(const A3DEntity* InEntityPtr)
 {
-	return A3DSurfRuledGet(InSurfacePtr, &Data);
+	return A3DSurfRuledGet(InEntityPtr, &Data);
 }
 
-A3DStatus TUniqueTSObj<A3DSurfSphereData>::GetData(const A3DEntity* InSurfacePtr)
+A3DStatus TUniqueTSObj<A3DSurfSphereData>::GetData(const A3DEntity* InEntityPtr)
 {
-	return A3DSurfSphereGet(InSurfacePtr, &Data);
+	return A3DSurfSphereGet(InEntityPtr, &Data);
 }
 
-A3DStatus TUniqueTSObj<A3DSurfTorusData>::GetData(const A3DEntity* InSurfacePtr)
+A3DStatus TUniqueTSObj<A3DSurfTorusData>::GetData(const A3DEntity* InEntityPtr)
 {
-	return A3DSurfTorusGet(InSurfacePtr, &Data);
+	return A3DSurfTorusGet(InEntityPtr, &Data);
 }
 
-A3DStatus TUniqueTSObj<A3DTess3DData>::GetData(const A3DTess3D* InTess3DPtr)
+A3DStatus TUniqueTSObj<A3DTess3DData>::GetData(const A3DEntity* InEntityPtr)
 {
-	return A3DTess3DGet(InTess3DPtr, &Data);
+	return A3DTess3DGet(InEntityPtr, &Data);
 }
 
-A3DStatus TUniqueTSObj<A3DTessBaseData>::GetData(const A3DTessBase* InTessBasePtr)
+A3DStatus TUniqueTSObj<A3DTessBaseData>::GetData(const A3DEntity* InEntityPtr)
 {
-	return A3DTessBaseGet(InTessBasePtr, &Data);
+	return A3DTessBaseGet(InEntityPtr, &Data);
 }
 
-A3DStatus TUniqueTSObj<A3DTopoBodyData>::GetData(const A3DEntity* InBodyPtr)
+A3DStatus TUniqueTSObj<A3DTopoBodyData>::GetData(const A3DEntity* InEntityPtr)
 {
-	return A3DTopoBodyGet(InBodyPtr, &Data);
+	return A3DTopoBodyGet(InEntityPtr, &Data);
 }
 
-A3DStatus TUniqueTSObj<A3DTopoBrepDataData>::GetData(const A3DTopoBrepData* InTopoBrepDataPtr)
+A3DStatus TUniqueTSObj<A3DTopoBrepDataData>::GetData(const A3DEntity* InEntityPtr)
 {
-	return A3DTopoBrepDataGet(InTopoBrepDataPtr, &Data);
+	return A3DTopoBrepDataGet(InEntityPtr, &Data);
 }
 
-A3DStatus TUniqueTSObj<A3DTopoCoEdgeData>::GetData(const A3DEntity* InGraphicsPtr)
+A3DStatus TUniqueTSObj<A3DTopoCoEdgeData>::GetData(const A3DEntity* InEntityPtr)
 {
-	return A3DTopoCoEdgeGet(InGraphicsPtr, &Data);
+	return A3DTopoCoEdgeGet(InEntityPtr, &Data);
 }
 
-A3DStatus TUniqueTSObj<A3DTopoConnexData>::GetData(const A3DEntity* InGraphicsPtr)
+A3DStatus TUniqueTSObj<A3DTopoConnexData>::GetData(const A3DEntity* InEntityPtr)
 {
-	return A3DTopoConnexGet(InGraphicsPtr, &Data);
+	return A3DTopoConnexGet(InEntityPtr, &Data);
 }
 
-A3DStatus TUniqueTSObj<A3DTopoContextData>::GetData(const A3DEntity* InGraphicsPtr)
+A3DStatus TUniqueTSObj<A3DTopoContextData>::GetData(const A3DEntity* InEntityPtr)
 {
-	return A3DTopoContextGet(InGraphicsPtr, &Data);
+	return A3DTopoContextGet(InEntityPtr, &Data);
 }
 
-A3DStatus TUniqueTSObj<A3DTopoEdgeData>::GetData(const A3DTopoEdge* InTopoEdgePtr)
+A3DStatus TUniqueTSObj<A3DTopoEdgeData>::GetData(const A3DEntity* InEntityPtr)
 {
-	return A3DTopoEdgeGet(InTopoEdgePtr, &Data);
+	return A3DTopoEdgeGet(InEntityPtr, &Data);
 }
 
-A3DStatus TUniqueTSObj<A3DTopoFaceData>::GetData(const A3DTopoFace* InTopoFacePtr)
+A3DStatus TUniqueTSObj<A3DTopoFaceData>::GetData(const A3DEntity* InEntityPtr)
 {
-	return A3DTopoFaceGet(InTopoFacePtr, &Data);
+	return A3DTopoFaceGet(InEntityPtr, &Data);
 }
 
-A3DStatus TUniqueTSObj<A3DTopoLoopData>::GetData(const A3DEntity* InTopoLoopPtr)
+A3DStatus TUniqueTSObj<A3DTopoLoopData>::GetData(const A3DEntity* InEntityPtr)
 {
-	return A3DTopoLoopGet(InTopoLoopPtr, &Data);
+	return A3DTopoLoopGet(InEntityPtr, &Data);
 }
 
-A3DStatus TUniqueTSObj<A3DTopoShellData>::GetData(const A3DEntity* InTopoPtr)
+A3DStatus TUniqueTSObj<A3DTopoShellData>::GetData(const A3DEntity* InEntityPtr)
 {
-	return A3DTopoShellGet(InTopoPtr, &Data);
+	return A3DTopoShellGet(InEntityPtr, &Data);
 }
 
-A3DStatus TUniqueTSObj<A3DTopoUniqueVertexData>::GetData(const A3DTopoUniqueVertex* InTopoUniqueVertexPtr)
+A3DStatus TUniqueTSObj<A3DTopoUniqueVertexData>::GetData(const A3DEntity* InEntityPtr)
 {
-	return A3DTopoUniqueVertexGet(InTopoUniqueVertexPtr, &Data);
+	return A3DTopoUniqueVertexGet(InEntityPtr, &Data);
 }
 
-A3DStatus TUniqueTSObj<A3DTopoWireEdgeData>::GetData(const A3DTopoWireEdge* InTopoWireEdgePtr)
+A3DStatus TUniqueTSObj<A3DTopoWireEdgeData>::GetData(const A3DEntity* InEntityPtr)
 {
-	return A3DTopoWireEdgeGet(InTopoWireEdgePtr, &Data);
+	return A3DTopoWireEdgeGet(InEntityPtr, &Data);
 }
 
-A3DStatus TUniqueTSObj<A3DUTF8Char*>::GetData(const A3DAsmProductOccurrence* InOccurrencePtr)
+A3DStatus TUniqueTSObj<A3DUTF8Char*>::GetData(const A3DEntity* InEntityPtr)
 {
-	return A3DAsmProductOccurrenceGetFilePathName(InOccurrencePtr, &Data);
+	return A3DStatus::A3D_ERROR;
 }
+
+
+
+// TUniqueTSObjFromIndex GetData -----------------------------------
+
+
+A3DStatus TUniqueTSObjFromIndex<A3DGraphMaterialData>::GetData(const uint32 InEntityIndex)
+{
+	return A3DGlobalGetGraphMaterialData(InEntityIndex, &Data);
+}
+
+A3DStatus TUniqueTSObjFromIndex<A3DGraphPictureData>::GetData(const uint32 InEntityIndex)
+{
+	return A3DGlobalGetGraphPictureData(InEntityIndex, &Data);
+}
+
+A3DStatus TUniqueTSObjFromIndex<A3DGraphRgbColorData>::GetData(const uint32 InEntityIndex)
+{
+	return A3DGlobalGetGraphRgbColorData(InEntityIndex, &Data);
+}
+
+A3DStatus TUniqueTSObjFromIndex<A3DGraphStyleData>::GetData(const uint32 InEntityIndex)
+{
+	return A3DGlobalGetGraphStyleData(InEntityIndex, &Data);
+}
+
+A3DStatus TUniqueTSObjFromIndex<A3DGraphTextureApplicationData>::GetData(const uint32 InEntityIndex)
+{
+	return A3DGlobalGetGraphTextureApplicationData(InEntityIndex, &Data);
+}
+
+A3DStatus TUniqueTSObjFromIndex<A3DGraphTextureDefinitionData>::GetData(const uint32 InEntityIndex)
+{
+	return A3DGlobalGetGraphTextureDefinitionData(InEntityIndex, &Data);
+}
+
+// DefaultValue -----------------------------------
+
+const A3DEntity* TUniqueTSObj<A3DAsmModelFileData>::GetDefaultIndexerValue() const { return nullptr; }
+const A3DEntity* TUniqueTSObj<A3DAsmPartDefinitionData>::GetDefaultIndexerValue() const { return nullptr; }
+const A3DEntity* TUniqueTSObj<A3DAsmProductOccurrenceData>::GetDefaultIndexerValue() const { return nullptr; }
+const A3DEntity* TUniqueTSObj<A3DAsmProductOccurrenceDataCV5>::GetDefaultIndexerValue() const { return nullptr; }
+const A3DEntity* TUniqueTSObj<A3DAsmProductOccurrenceDataSLW>::GetDefaultIndexerValue() const { return nullptr; }
+const A3DEntity* TUniqueTSObj<A3DAsmProductOccurrenceDataUg>::GetDefaultIndexerValue() const { return nullptr; }
+const A3DEntity* TUniqueTSObj<A3DBoundingBoxData>::GetDefaultIndexerValue() const { return nullptr; }
+const A3DEntity* TUniqueTSObj<A3DCrvCircleData>::GetDefaultIndexerValue() const { return nullptr; }
+const A3DEntity* TUniqueTSObj<A3DCrvCompositeData>::GetDefaultIndexerValue() const { return nullptr; }
+const A3DEntity* TUniqueTSObj<A3DCrvEllipseData>::GetDefaultIndexerValue() const { return nullptr; }
+const A3DEntity* TUniqueTSObj<A3DCrvHelixData>::GetDefaultIndexerValue() const { return nullptr; }
+const A3DEntity* TUniqueTSObj<A3DCrvHyperbolaData>::GetDefaultIndexerValue() const { return nullptr; }
+const A3DEntity* TUniqueTSObj<A3DCrvLineData>::GetDefaultIndexerValue() const { return nullptr; }
+const A3DEntity* TUniqueTSObj<A3DCrvNurbsData>::GetDefaultIndexerValue() const { return nullptr; }
+const A3DEntity* TUniqueTSObj<A3DCrvParabolaData>::GetDefaultIndexerValue() const { return nullptr; }
+const A3DEntity* TUniqueTSObj<A3DCrvPolyLineData>::GetDefaultIndexerValue() const { return nullptr; }
+const A3DEntity* TUniqueTSObj<A3DCrvTransformData>::GetDefaultIndexerValue() const { return nullptr; }
+const A3DEntity* TUniqueTSObj<A3DGlobalData>::GetDefaultIndexerValue() const { return nullptr; }
+const A3DEntity* TUniqueTSObj<A3DGraphicsData>::GetDefaultIndexerValue() const { return nullptr; }
+const A3DEntity* TUniqueTSObj<A3DIntervalData>::GetDefaultIndexerValue() const { return nullptr; }
+const A3DEntity* TUniqueTSObj<A3DMiscAttributeData>::GetDefaultIndexerValue() const { return nullptr; }
+const A3DEntity* TUniqueTSObj<A3DMiscCartesianTransformationData>::GetDefaultIndexerValue() const { return nullptr; }
+const A3DEntity* TUniqueTSObj<A3DMiscEntityReferenceData>::GetDefaultIndexerValue() const { return nullptr; }
+const A3DEntity* TUniqueTSObj<A3DMiscGeneralTransformationData>::GetDefaultIndexerValue() const { return nullptr; }
+const A3DEntity* TUniqueTSObj<A3DMiscMaterialPropertiesData>::GetDefaultIndexerValue() const { return nullptr; }
+const A3DEntity* TUniqueTSObj<A3DMiscReferenceOnCsysItemData>::GetDefaultIndexerValue() const { return nullptr; }
+const A3DEntity* TUniqueTSObj<A3DMiscReferenceOnTessData>::GetDefaultIndexerValue() const { return nullptr; }
+const A3DEntity* TUniqueTSObj<A3DMiscReferenceOnTopologyData>::GetDefaultIndexerValue() const { return nullptr; }
+const A3DEntity* TUniqueTSObj<A3DMiscSingleAttributeData>::GetDefaultIndexerValue() const { return nullptr; }
+const A3DEntity* TUniqueTSObj<A3DRWParamsExportPrcData>::GetDefaultIndexerValue() const { return nullptr; }
+const A3DEntity* TUniqueTSObj<A3DRiBrepModelData>::GetDefaultIndexerValue() const { return nullptr; }
+const A3DEntity* TUniqueTSObj<A3DRiCoordinateSystemData>::GetDefaultIndexerValue() const { return nullptr; }
+const A3DEntity* TUniqueTSObj<A3DRiDirectionData>::GetDefaultIndexerValue() const { return nullptr; }
+const A3DEntity* TUniqueTSObj<A3DRiPolyBrepModelData>::GetDefaultIndexerValue() const { return nullptr; }
+const A3DEntity* TUniqueTSObj<A3DRiRepresentationItemData>::GetDefaultIndexerValue() const { return nullptr; }
+const A3DEntity* TUniqueTSObj<A3DRiSetData>::GetDefaultIndexerValue() const { return nullptr; }
+const A3DEntity* TUniqueTSObj<A3DRootBaseData>::GetDefaultIndexerValue() const { return nullptr; }
+const A3DEntity* TUniqueTSObj<A3DRootBaseWithGraphicsData>::GetDefaultIndexerValue() const { return nullptr; }
+const A3DEntity* TUniqueTSObj<A3DRWParamsTessellationData>::GetDefaultIndexerValue() const { return nullptr; }
+const A3DEntity* TUniqueTSObj<A3DSewOptionsData>::GetDefaultIndexerValue() const { return nullptr; }
+const A3DEntity* TUniqueTSObj<A3DSurfBlend01Data>::GetDefaultIndexerValue() const { return nullptr; }
+const A3DEntity* TUniqueTSObj<A3DSurfBlend02Data>::GetDefaultIndexerValue() const { return nullptr; }
+const A3DEntity* TUniqueTSObj<A3DSurfBlend03Data>::GetDefaultIndexerValue() const { return nullptr; }
+const A3DEntity* TUniqueTSObj<A3DSurfConeData>::GetDefaultIndexerValue() const { return nullptr; }
+const A3DEntity* TUniqueTSObj<A3DSurfCylinderData>::GetDefaultIndexerValue() const { return nullptr; }
+const A3DEntity* TUniqueTSObj<A3DSurfCylindricalData>::GetDefaultIndexerValue() const { return nullptr; }
+const A3DEntity* TUniqueTSObj<A3DSurfExtrusionData>::GetDefaultIndexerValue() const { return nullptr; }
+const A3DEntity* TUniqueTSObj<A3DSurfFromCurvesData>::GetDefaultIndexerValue() const { return nullptr; }
+const A3DEntity* TUniqueTSObj<A3DSurfNurbsData>::GetDefaultIndexerValue() const { return nullptr; }
+const A3DEntity* TUniqueTSObj<A3DSurfPipeData>::GetDefaultIndexerValue() const { return nullptr; }
+const A3DEntity* TUniqueTSObj<A3DSurfPlaneData>::GetDefaultIndexerValue() const { return nullptr; }
+const A3DEntity* TUniqueTSObj<A3DSurfRevolutionData>::GetDefaultIndexerValue() const { return nullptr; }
+const A3DEntity* TUniqueTSObj<A3DSurfRuledData>::GetDefaultIndexerValue() const { return nullptr; }
+const A3DEntity* TUniqueTSObj<A3DSurfSphereData>::GetDefaultIndexerValue() const { return nullptr; }
+const A3DEntity* TUniqueTSObj<A3DSurfTorusData>::GetDefaultIndexerValue() const { return nullptr; }
+const A3DEntity* TUniqueTSObj<A3DTess3DData>::GetDefaultIndexerValue() const { return nullptr; }
+const A3DEntity* TUniqueTSObj<A3DTessBaseData>::GetDefaultIndexerValue() const { return nullptr; }
+const A3DEntity* TUniqueTSObj<A3DTopoBodyData>::GetDefaultIndexerValue() const { return nullptr; }
+const A3DEntity* TUniqueTSObj<A3DTopoBrepDataData>::GetDefaultIndexerValue() const { return nullptr; }
+const A3DEntity* TUniqueTSObj<A3DTopoCoEdgeData>::GetDefaultIndexerValue() const { return nullptr; }
+const A3DEntity* TUniqueTSObj<A3DTopoConnexData>::GetDefaultIndexerValue() const { return nullptr; }
+const A3DEntity* TUniqueTSObj<A3DTopoContextData>::GetDefaultIndexerValue() const { return nullptr; }
+const A3DEntity* TUniqueTSObj<A3DTopoEdgeData>::GetDefaultIndexerValue() const { return nullptr; }
+const A3DEntity* TUniqueTSObj<A3DTopoFaceData>::GetDefaultIndexerValue() const { return nullptr; }
+const A3DEntity* TUniqueTSObj<A3DTopoLoopData>::GetDefaultIndexerValue() const { return nullptr; }
+const A3DEntity* TUniqueTSObj<A3DTopoShellData>::GetDefaultIndexerValue() const { return nullptr; }
+const A3DEntity* TUniqueTSObj<A3DTopoUniqueVertexData>::GetDefaultIndexerValue() const { return nullptr; }
+const A3DEntity* TUniqueTSObj<A3DTopoWireEdgeData>::GetDefaultIndexerValue() const { return nullptr; }
+
+uint32 TUniqueTSObjFromIndex<A3DGraphMaterialData>::GetDefaultIndexerValue() const { return A3D_DEFAULT_MATERIAL_INDEX; }
+uint32 TUniqueTSObjFromIndex<A3DGraphPictureData>::GetDefaultIndexerValue() const { return A3D_DEFAULT_PICTURE_INDEX; }
+uint32 TUniqueTSObjFromIndex<A3DGraphRgbColorData>::GetDefaultIndexerValue() const { return A3D_DEFAULT_COLOR_INDEX; }
+uint32 TUniqueTSObjFromIndex<A3DGraphStyleData>::GetDefaultIndexerValue() const { return A3D_DEFAULT_STYLE_INDEX; }
+uint32 TUniqueTSObjFromIndex<A3DGraphTextureApplicationData>::GetDefaultIndexerValue() const { return A3D_DEFAULT_TEXTURE_APPLICATION_INDEX; }
+uint32 TUniqueTSObjFromIndex<A3DGraphTextureDefinitionData>::GetDefaultIndexerValue() const { return A3D_DEFAULT_TEXTURE_DEFINITION_INDEX; }
+
+const A3DEntity* TUniqueTSObj<A3DUTF8Char*>::GetDefaultIndexerValue() const { return nullptr; }
+
 #endif
 
-}
+} // NS CADLibrary
+
diff --git a/Engine/Plugins/Enterprise/DatasmithCADImporter/Source/CADInterfaces/Private/TechSoftUtils.cpp b/Engine/Plugins/Enterprise/DatasmithCADImporter/Source/CADInterfaces/Private/TechSoftUtils.cpp
new file mode 100644
--- /dev/null
+++ b/Engine/Plugins/Enterprise/DatasmithCADImporter/Source/CADInterfaces/Private/TechSoftUtils.cpp
@@ -0,0 +1,582 @@
+// Copyright Epic Games, Inc. All Rights Reserved.
+
+#include "TechSoftUtils.h"
+
+
+#include "TUniqueTechSoftObj.h"
+
+#include "CADOptions.h"
+
+#include "Math/Color.h"
+#include "Serialization/JsonSerializer.h"
+#include "Serialization/JsonWriter.h"
+
+#ifndef CADKERNEL_DEV
+#include "TechSoftInterfaceUtils.inl"
+#endif
+
+#include <string>
+
+namespace CADLibrary
+{
+
+namespace TechSoftUtils
+{
+
+#ifdef USE_TECHSOFT_SDK
+TSharedPtr<FJsonObject> GetJsonObject(A3DAsmProductOccurrence* ProductOcccurence);
+void RestoreMaterials(const TSharedPtr<FJsonObject>& DefaultValues, CADLibrary::FBodyMesh& BodyMesh);
+void SaveModelFileToPcrFile(void* ModelFile, const FString& Filename);
+A3DUns32 CreateRGBColor(FColor& Color);
+void SetRootOccurenceAttributes(A3DEntity* Entity);
+#endif
+
+bool GetBodyFromPcrFile(const FString& Filename, const FImportParameters& ImportParameters, FBodyMesh& BodyMesh)
+{
+	bool bExtractionSuccessful = false;
+
+#if defined USE_TECHSOFT_SDK && !defined CADKERNEL_DEV
+	A3DRWParamsPrcReadHelper* ReadHelper = nullptr;
+
+	FUniqueTechSoftModelFile ModelFile = TechSoftInterface::LoadModelFileFromPrcFile(TCHAR_TO_UTF8(*Filename), &ReadHelper);
+
+	if (!ModelFile.IsValid())
+	{
+		return false;
+	}
+
+	TUniqueTSObj<A3DAsmModelFileData> ModelFileData(ModelFile.Get());
+	if (!ModelFileData.IsValid() || ModelFileData->m_uiPOccurrencesSize == 0)
+	{
+		return false;
+	}
+
+	TUniqueTSObj<A3DAsmProductOccurrenceData> OccurenceData(ModelFileData->m_ppPOccurrences[0]);
+	if (!OccurenceData.IsValid() || !OccurenceData->m_pPart)
+	{
+		return false;
+	}
+
+	TUniqueTSObj<A3DAsmPartDefinitionData> PartDefinitionData(OccurenceData->m_pPart);
+	if (!PartDefinitionData.IsValid() || PartDefinitionData->m_uiRepItemsSize == 0)
+	{
+		return false;
+	}
+
+	TSharedPtr<FJsonObject> JsonObject = GetJsonObject(ModelFileData->m_ppPOccurrences[0]);
+	if (JsonObject.IsValid())
+	{
+		double FileUnit = 1.0;
+
+		JsonObject->TryGetNumberField(JSON_ENTRY_FILE_UNIT, FileUnit);
+
+		if (FillBodyMesh(PartDefinitionData->m_ppRepItems[0], ImportParameters, FileUnit, BodyMesh))
+		{
+			bExtractionSuccessful = true;
+			RestoreMaterials(JsonObject, BodyMesh);
+		}
+	}
+
+#endif
+	return bExtractionSuccessful;
+}
+
+FUniqueTechSoftModelFile SaveBodiesToPcrFile(void** Bodies, uint32 BodyCount, const FString& Filename, const FString& JsonString)
+{
+#if defined USE_TECHSOFT_SDK && !defined CADKERNEL_DEV
+	if (!Bodies)
+	{
+		return FUniqueTechSoftModelFile();
+	}
+
+	// Create PartDefinition
+	A3DRiRepresentationItem** RepresentationItems = (A3DRiRepresentationItem**)Bodies;
+
+	TUniqueTSObj<A3DAsmPartDefinitionData> PartDefinitionData;
+	PartDefinitionData->m_uiRepItemsSize = BodyCount;
+	PartDefinitionData->m_ppRepItems = RepresentationItems;
+	A3DAsmPartDefinition* PartDefinition = TechSoftInterface::CreateAsmPartDefinition(*PartDefinitionData);
+
+	TUniqueTSObj<A3DAsmProductOccurrenceData> ProductOccurrenceData;
+	ProductOccurrenceData->m_pPart = PartDefinition;
+	A3DAsmProductOccurrence* ProductOccurrence = TechSoftInterface::CreateAsmProductOccurrence(*ProductOccurrenceData);
+
+	// Add MaterialTable as attribute to ProductOccurrence
+	std::string StringAnsi(TCHAR_TO_UTF8(*JsonString));
+
+	TUniqueTSObj<A3DMiscSingleAttributeData> SingleAttributeData;
+	SingleAttributeData->m_eType = kA3DModellerAttributeTypeString;
+	SingleAttributeData->m_pcTitle = "MaterialTable";
+	SingleAttributeData->m_pcData = (char*)StringAnsi.c_str();
+
+	TUniqueTSObj<A3DMiscAttributeData> AttributesData;
+	AttributesData->m_pcTitle = SingleAttributeData->m_pcTitle;
+	AttributesData->m_asSingleAttributesData = SingleAttributeData.GetPtr();
+	AttributesData->m_uiSize = 1;
+	A3DMiscAttribute* Attributes = TechSoftInterface::CreateMiscAttribute(*AttributesData);
+
+	TUniqueTSObj<A3DRootBaseData> RootBaseData;
+	RootBaseData->m_pcName = SingleAttributeData->m_pcTitle;
+	RootBaseData->m_ppAttributes = &Attributes;
+	RootBaseData->m_uiSize = 1;
+	TechSoftInterface::SetRootBase(ProductOccurrence, *RootBaseData);
+
+	// Create ModelFile
+	TUniqueTSObj<A3DAsmModelFileData> ModelFileData;
+	ModelFileData->m_uiPOccurrencesSize = 1;
+	ModelFileData->m_dUnit = 1.0;
+	ModelFileData->m_ppPOccurrences = &ProductOccurrence;
+	FUniqueTechSoftModelFile ModelFile = TechSoftInterface::CreateAsmModelFile(*ModelFileData);
+
+	// Save ModelFile to Pcr file
+	SaveModelFileToPcrFile(ModelFile.Get(), Filename);
+
+	// #ueent_techsoft: Deleting the model seems to delete the entire content. To be double-checked
+	TechSoftInterface::DeleteEntity(Attributes);
+
+	return ModelFile;
+#else
+	return FUniqueTechSoftModelFile();
+#endif
+}
+
+bool FillBodyMesh(void* BodyPtr, const FImportParameters& ImportParameters, double FileUnit, FBodyMesh& BodyMesh)
+{
+#if defined USE_TECHSOFT_SDK && !defined CADKERNEL_DEV
+	A3DRiRepresentationItem* RepresentationItemPtr = (A3DRiRepresentationItem*)BodyPtr;
+
+	A3DEEntityType Type;
+	A3DEntityGetType(RepresentationItemPtr, &Type);
+	if (Type == kA3DTypeRiPolyBrepModel)
+	{
+		TUniqueTSObj<A3DRiRepresentationItemData> RepresentationItemData(RepresentationItemPtr);
+		TechSoftInterfaceUtils::FTechSoftTessellationExtractor Extractor(RepresentationItemData->m_pTessBase);
+		return Extractor.FillBodyMesh(BodyMesh, FileUnit);
+	}
+
+
+	// TUniqueTechSoftObj does not work in this case
+	TUniqueTSObj<A3DRWParamsTessellationData> TessellationParameters;
+
+	TessellationParameters->m_eTessellationLevelOfDetail = kA3DTessLODUserDefined; // Enum to specify predefined values for some following members.
+	TessellationParameters->m_bUseHeightInsteadOfRatio = A3D_TRUE;
+	TessellationParameters->m_dMaxChordHeight = ImportParameters.GetChordTolerance() * 10. / FileUnit;
+	TessellationParameters->m_dAngleToleranceDeg = ImportParameters.GetMaxNormalAngle();
+	TessellationParameters->m_dMaximalTriangleEdgeLength = 0; //ImportParameters.MaxEdgeLength;
+
+	TessellationParameters->m_bAccurateTessellation = A3D_FALSE;  // A3D_FALSE' indicates the tessellation is set for visualization
+	TessellationParameters->m_bAccurateTessellationWithGrid = A3D_FALSE; // Enable accurate tessellation with faces inner points on a grid.
+	TessellationParameters->m_dAccurateTessellationWithGridMaximumStitchLength = 0; 	// Maximal grid stitch length. Disabled if value is 0. Be careful, a too small value can generate a huge tessellation.
+
+	TessellationParameters->m_bKeepUVPoints = A3D_TRUE; // Keep parametric points as texture points.
+
+	// Get the tessellation
+	A3DStatus Status = A3DRiRepresentationItemComputeTessellation(RepresentationItemPtr, TessellationParameters.GetPtr());
+	TUniqueTSObj<A3DRiRepresentationItemData> RepresentationItemData(RepresentationItemPtr);
+	if (!RepresentationItemData.IsValid())
+	{
+		return false;
+	}
+
+	{
+		A3DEntityGetType(RepresentationItemData->m_pTessBase, &Type);
+		ensure(Type == kA3DTypeTess3D);
+	}
+
+	TechSoftInterfaceUtils::FTechSoftTessellationExtractor Extractor(RepresentationItemData->m_pTessBase);
+	return Extractor.FillBodyMesh(BodyMesh, FileUnit);
+#else
+	return false;
+#endif
+}
+
+#if defined USE_TECHSOFT_SDK && !defined CADKERNEL_DEV
+TSharedPtr<FJsonObject> GetJsonObject(A3DAsmProductOccurrence* ProductOcccurence)
+{
+	TUniqueTSObj<A3DRootBaseData> RootBaseData(ProductOcccurence);
+
+	if (RootBaseData.IsValid() && RootBaseData->m_uiSize > 0)
+	{
+		TUniqueTSObj<A3DMiscAttributeData> AttributeData(RootBaseData->m_ppAttributes[0]);
+		if (AttributeData->m_uiSize > 0 && AttributeData->m_asSingleAttributesData[0].m_eType == kA3DModellerAttributeTypeString)
+		{
+			FString JsonString = UTF8_TO_TCHAR(AttributeData->m_asSingleAttributesData[0].m_pcData);
+
+			TSharedRef<TJsonReader<>> JsonReader = TJsonReaderFactory<>::Create(JsonString);
+			TSharedPtr<FJsonObject> JsonObject;
+
+			if (FJsonSerializer::Deserialize(JsonReader, JsonObject))
+			{
+				return JsonObject;
+			}
+		}
+	}
+
+	return {};
+}
+
+// Replicate logic in FTechSoftFileParser::FindOrAddMaterial and the methods it calls
+void BuildCADMaterial(uint32 MaterialIndex, const A3DGraphStyleData& GraphStyleData, FCADMaterial& Material)
+{
+	if (TechSoftInterface::IsMaterialTexture(MaterialIndex))
+	{
+		TUniqueTSObjFromIndex<A3DGraphTextureApplicationData> TextureData(MaterialIndex);
+		if (TextureData.IsValid())
+		{
+			MaterialIndex = TextureData->m_uiMaterialIndex;
+		}
+	}
+
+	TUniqueTSObjFromIndex<A3DGraphMaterialData> MaterialData(MaterialIndex);
+	if (MaterialData.IsValid())
+	{
+		Material.Diffuse = GetColorAt(MaterialData->m_uiDiffuse);
+		Material.Ambient = GetColorAt(MaterialData->m_uiAmbient);
+		Material.Specular = GetColorAt(MaterialData->m_uiSpecular);
+		Material.Shininess = MaterialData->m_dShininess;
+		if (GraphStyleData.m_bIsTransparencyDefined)
+		{
+			Material.Transparency = 1. - GraphStyleData.m_ucTransparency / 255.;
+		}
+		// todo: find how to convert Emissive color into ? reflexion coef...
+		// Material.Emissive = GetColor(MaterialData->m_uiEmissive);
+		// Material.Reflexion;
+	}
+}
+#endif
+
+// Replicates the logic in FTechSoftFileParser::ExtractGraphStyleProperties
+void GetMaterialValues(uint32 StyleIndex, FCADUUID& OutColorName, FCADUUID& OutMaterialName)
+{
+#if defined USE_TECHSOFT_SDK && !defined CADKERNEL_DEV
+	TUniqueTSObjFromIndex<A3DGraphStyleData> GraphStyleData(StyleIndex);
+
+	if (GraphStyleData.IsValid())
+	{
+		if (GraphStyleData->m_bMaterial)
+		{
+			FCADMaterial Material;
+
+			BuildCADMaterial(GraphStyleData->m_uiRgbColorIndex, *GraphStyleData, Material);
+
+			OutMaterialName = BuildMaterialName(Material);
+		}
+		else
+		{
+			TUniqueTSObjFromIndex<A3DGraphRgbColorData> ColorData(GraphStyleData->m_uiRgbColorIndex);
+			if (ColorData.IsValid())
+			{
+				const uint8 Alpha = GraphStyleData->m_bIsTransparencyDefined ? GraphStyleData->m_ucTransparency : 255;
+				const FColor ColorValue((uint8)(ColorData->m_dRed * 255), (uint8)(ColorData->m_dGreen * 255), (uint8)(ColorData->m_dBlue * 255), Alpha);
+
+				OutColorName = BuildColorName(ColorValue);
+			}
+		}
+	}
+#endif
+}
+
+void RestoreMaterials(const TSharedPtr<FJsonObject>& DefaultValues, FBodyMesh& BodyMesh)
+{
+	FCADUUID DefaultColorName = 0;
+	FCADUUID DefaultMaterialName = 0;
+
+	DefaultValues->TryGetNumberField(JSON_ENTRY_COLOR_NAME, DefaultColorName);
+	DefaultValues->TryGetNumberField(JSON_ENTRY_MATERIAL_NAME, DefaultMaterialName);
+
+	BodyMesh.MaterialSet.Empty();
+	BodyMesh.ColorSet.Empty();
+
+	for (FTessellationData& Tessellation : BodyMesh.Faces)
+	{
+		// Extract proper color or material based on style index
+		uint32 CachedStyleIndex = Tessellation.MaterialName;
+		Tessellation.MaterialName = 0;
+
+		FCADUUID ColorName = DefaultColorName;
+		FCADUUID MaterialName = DefaultMaterialName;
+
+		GetMaterialValues(CachedStyleIndex, ColorName, MaterialName);
+
+		if (ColorName)
+		{
+			Tessellation.ColorName = ColorName;
+			BodyMesh.ColorSet.Add(ColorName);
+		}
+
+		if (MaterialName)
+		{
+			Tessellation.MaterialName = MaterialName;
+			BodyMesh.MaterialSet.Add(MaterialName);
+		}
+	}
+}
+
+FColor GetColorAt(uint32 ColorIndex)
+{
+#if defined USE_TECHSOFT_SDK && !defined CADKERNEL_DEV
+	TUniqueTSObjFromIndex<A3DGraphRgbColorData> ColorData(ColorIndex);
+	if (ColorData.IsValid())
+	{
+		return FColor((uint8)(ColorData->m_dRed * 255)
+			, (uint8)(ColorData->m_dGreen * 255)
+			, (uint8)(ColorData->m_dBlue * 255));
+	}
+#endif
+	return FColor(200, 200, 200);
+}
+
+void SaveModelFileToPcrFile(void* ModelFile, const FString& Filename)
+{
+#if defined USE_TECHSOFT_SDK && !defined CADKERNEL_DEV
+	TUniqueTSObj<A3DRWParamsExportPrcData> ParamsExportData;
+	ParamsExportData->m_bCompressBrep = false;
+	ParamsExportData->m_bCompressTessellation = false;
+
+	A3DUTF8Char HsfFileName[_MAX_PATH];
+	FCStringAnsi::Strncpy(HsfFileName, TCHAR_TO_UTF8(*Filename), _MAX_PATH);
+
+	TechSoftInterface::ExportModelFileToPrcFile(ModelFile, ParamsExportData.GetPtr(), HsfFileName, nullptr);
+#endif
+}
+
+#if defined USE_TECHSOFT_SDK && !defined CADKERNEL_DEV
+A3DUns32 CreateRGBColor(FColor& Color)
+{
+	TUniqueTSObjFromIndex<A3DGraphRgbColorData> RgbColor;
+
+	RgbColor->m_dRed = Color.R / 255.;
+	RgbColor->m_dGreen = Color.G / 255.;
+	RgbColor->m_dBlue = Color.B / 255.;
+	return TechSoftInterface::InsertGraphRgbColor(*RgbColor);
+}
+
+int32 SetEntityGraphicsColor(A3DEntity* InEntity, FColor Color)
+{
+	TUniqueTSObj<A3DRootBaseWithGraphicsData> BaseWithGraphicsData(InEntity);
+
+	//Create a style color
+	A3DUns32 ColorIndex = CreateRGBColor(Color);
+
+	A3DUns32 StyleIndex = 0;
+	TUniqueTSObjFromIndex<A3DGraphStyleData> StyleData;
+	StyleData->m_bMaterial = false;
+	StyleData->m_bVPicture = false;
+	StyleData->m_dWidth = 0.1; // default
+	A3DUns8 Alpha = Color.A;
+	if (Alpha > 0)
+	{
+		StyleData->m_bIsTransparencyDefined = true;
+		StyleData->m_ucTransparency = 255 - Alpha;
+	}
+	else
+	{
+		StyleData->m_bIsTransparencyDefined = false;
+		StyleData->m_ucTransparency = 255;
+	}
+
+	StyleData->m_bSpecialCulling = false;
+	StyleData->m_bBackCulling = false;
+	StyleData->m_uiRgbColorIndex = ColorIndex;
+	StyleIndex = TechSoftInterface::InsertGraphStyle(*StyleData);
+
+	TUniqueTSObj<A3DGraphicsData> GraphicsData;
+
+	GraphicsData->m_uiStyleIndex = StyleIndex;
+	GraphicsData->m_usBehaviour = kA3DGraphicsShow;
+	GraphicsData->m_usBehaviour |= kA3DGraphicsSonHeritColor;
+
+	BaseWithGraphicsData->m_pGraphics = TechSoftInterface::CreateGraphics(*GraphicsData);
+	if (BaseWithGraphicsData->m_pGraphics == nullptr)
+	{
+		return A3DStatus::A3D_ERROR;
+	}
+
+	return TechSoftInterface::SetRootBaseWithGraphics(*BaseWithGraphicsData, InEntity);
+}
+
+void SetRootOccurenceAttributes(A3DEntity* Entity)
+{
+	A3DMiscAttribute* Attributes[3];
+
+	TUniqueTSObj<A3DMiscSingleAttributeData> Single;
+	Single->m_eType = kA3DModellerAttributeTypeString;
+
+	TUniqueTSObj<A3DMiscAttributeData> AttributeData;
+	AttributeData->m_asSingleAttributesData = Single.GetPtr();
+	AttributeData->m_uiSize = 1;
+
+	Single->m_pcTitle = (char*)"Title";
+	Single->m_pcData = (char*)"Body model";
+	AttributeData->m_pcTitle = Single->m_pcTitle;
+	Attributes[0] = TechSoftInterface::CreateMiscAttribute(*AttributeData);
+
+	Single->m_pcTitle = (char*)"Author";
+	Single->m_pcData = (char*)"Unreal Engine";
+	AttributeData->m_pcTitle = Single->m_pcTitle;
+	Attributes[1] = TechSoftInterface::CreateMiscAttribute(*AttributeData);
+
+	Single->m_pcTitle = (char*)"Company";
+	Single->m_pcData = (char*)"Epic Games";
+	AttributeData->m_pcTitle = Single->m_pcTitle;
+	Attributes[2] = TechSoftInterface::CreateMiscAttribute(*AttributeData);
+
+	TUniqueTSObj<A3DRootBaseData> RootData;
+	RootData->m_pcName = (char*)"Body model";
+	RootData->m_ppAttributes = Attributes;
+	RootData->m_uiSize = 3;
+	TechSoftInterface::SetRootBase(Entity, *RootData);
+
+	for (A3DUns32 i = 0; i < RootData->m_uiSize; ++i)
+	{
+		TechSoftInterface::DeleteEntity(Attributes[i]);
+	}
+}
+
+A3DAsmModelFile* CreateModelFile(A3DAsmProductOccurrence* Occurrence)
+{
+	if (Occurrence != NULL)
+	{
+		return nullptr;
+	}
+
+	CADLibrary::TUniqueTSObj<A3DAsmModelFileData> ModelFileData;
+
+	// In this model there will be only one product occurrence
+	ModelFileData->m_uiPOccurrencesSize = 1;
+	ModelFileData->m_dUnit = 1.0;
+	ModelFileData->m_ppPOccurrences = &Occurrence;
+
+	return TechSoftInterface::CreateModelFile(*ModelFileData);
+}
+
+A3DAsmProductOccurrence* CreateOccurrence(A3DAsmPartDefinition* Part)
+{
+	if (Part == nullptr)
+	{
+		return nullptr;
+	}
+	TUniqueTSObj<A3DAsmProductOccurrenceData> ProductOccurrenceData;
+
+	ProductOccurrenceData->m_pPart = Part;
+	A3DAsmProductOccurrence* ProductOccurrencePtr = TechSoftInterface::CreateAsmProductOccurrence(*ProductOccurrenceData);
+	if (ProductOccurrencePtr != nullptr)
+	{
+		SetRootOccurenceAttributes(ProductOccurrencePtr);
+	}
+	return ProductOccurrencePtr;
+}
+
+A3DAsmPartDefinition* CreatePart(TArray<A3DRiRepresentationItem*>& RepresentationItems)
+{
+	if (RepresentationItems.IsEmpty())
+	{
+		return nullptr;
+	}
+
+	TUniqueTSObj<A3DAsmPartDefinitionData> PartDefinitionData;
+
+	PartDefinitionData->m_uiRepItemsSize = RepresentationItems.Num();
+	PartDefinitionData->m_ppRepItems = RepresentationItems.GetData();
+
+	return TechSoftInterface::CreateAsmPartDefinition(*PartDefinitionData);
+}
+
+A3DRiRepresentationItem* CreateRIBRep(A3DTopoShell* TopoShellPtr)
+{
+	if (TopoShellPtr == nullptr)
+	{
+		return nullptr;
+	}
+
+	A3DTopoConnex* TopoConnexPtr = nullptr;
+	{
+		CADLibrary::TUniqueTSObj<A3DTopoConnexData> TopoConnexData;
+		TopoConnexData->m_ppShells = &TopoShellPtr;
+		TopoConnexData->m_uiShellSize = 1;
+		TopoConnexPtr = TechSoftInterface::CreateTopoConnex(*TopoConnexData);
+		if (TopoConnexPtr == nullptr)
+		{
+			return nullptr;
+		}
+	}
+
+	A3DTopoBrepData* TopoBRepPtr = nullptr;
+	{
+		CADLibrary::TUniqueTSObj<A3DTopoBrepDataData> TopoBRepData;
+		TopoBRepData->m_uiConnexSize = 1;
+		TopoBRepData->m_ppConnexes = &TopoConnexPtr;
+		TopoBRepPtr = TechSoftInterface::CreateTopoBRep(*TopoBRepData);
+		if (TopoBRepPtr == nullptr)
+		{
+			return nullptr;
+		}
+	}
+
+	CADLibrary::TUniqueTSObj<A3DRiBrepModelData> RiBRepModelData;
+	RiBRepModelData->m_pBrepData = TopoBRepPtr;
+	RiBRepModelData->m_bSolid = false;
+	return TechSoftInterface::CreateRiBRepModel(*RiBRepModelData);
+}
+
+A3DTopoEdge* CreateTopoEdge()
+{
+	CADLibrary::TUniqueTSObj<A3DTopoEdgeData> EdgeData;
+	return TechSoftInterface::CreateTopoEdge(*EdgeData);
+}
+
+A3DTopoFace* CreateTopoFaceWithNaturalLoop(A3DSurfBase* CarrierSurface)
+{
+	CADLibrary::TUniqueTSObj<A3DTopoFaceData> Face;
+	Face->m_pSurface = CarrierSurface;
+	Face->m_bHasTrimDomain = false;
+	Face->m_ppLoops = nullptr;
+	Face->m_uiLoopSize = 0;
+	Face->m_uiOuterLoopIndex = 0;
+	Face->m_dTolerance = 0.01; //mm
+
+	return CADLibrary::TechSoftInterface::CreateTopoFace(*Face);
+}
+
+A3DCrvNurbs* CreateTrimNurbsCurve(A3DCrvNurbs* CurveNurbsPtr, double UMin, double UMax, bool bIs2D)
+{
+	if (CurveNurbsPtr == nullptr)
+	{
+		return nullptr;
+	}
+
+	TUniqueTSObj<A3DCrvTransformData> TransformCurveData;
+
+	TransformCurveData->m_bIs2D = bIs2D;
+	TransformCurveData->m_sParam.m_sInterval.m_dMin = UMin;
+	TransformCurveData->m_sParam.m_sInterval.m_dMax = UMax;
+	TransformCurveData->m_sParam.m_dCoeffA = 1.;
+	TransformCurveData->m_sParam.m_dCoeffB = 0.;
+	TransformCurveData->m_pBasisCrv = CurveNurbsPtr;
+	TransformCurveData->m_pTransfo = nullptr;
+
+	TransformCurveData->m_sTrsf.m_sXVector.m_dX = 1.;
+	TransformCurveData->m_sTrsf.m_sYVector.m_dY = 1.;
+	TransformCurveData->m_sTrsf.m_sScale.m_dX = 1.;
+	TransformCurveData->m_sTrsf.m_sScale.m_dY = 1.;
+	TransformCurveData->m_sTrsf.m_sScale.m_dZ = 1.;
+
+	A3DCrvTransform* CurveTransformPtr = CADLibrary::TechSoftInterface::CreateCurveTransform(*TransformCurveData);
+	if (CurveTransformPtr == nullptr)
+	{
+		return nullptr;
+	}
+
+	TUniqueTSObj<A3DCrvNurbsData> TransformedNurbsData;
+	if (CADLibrary::TechSoftInterface::GetCurveAsNurbs(CurveTransformPtr, &*TransformedNurbsData, 0.01 /*mm*/, /*bUseSameParameterization*/ true) != A3DStatus::A3D_SUCCESS)
+	{
+		return nullptr;
+	}
+
+	return CADLibrary::TechSoftInterface::CreateCurveNurbs(*TransformedNurbsData);
+}
+#endif
+
+} // NS TechSoftUtils
+
+} // NS CADLibrary
diff --git a/Engine/Plugins/Enterprise/DatasmithCADImporter/Source/CADInterfaces/Public/TUniqueTechSoftObj.h b/Engine/Plugins/Enterprise/DatasmithCADImporter/Source/CADInterfaces/Public/TUniqueTechSoftObj.h
new file mode 100644
--- /dev/null
+++ b/Engine/Plugins/Enterprise/DatasmithCADImporter/Source/CADInterfaces/Public/TUniqueTechSoftObj.h
@@ -0,0 +1,467 @@
+// Copyright Epic Games, Inc. All Rights Reserved.
+#pragma once
+
+#ifdef USE_TECHSOFT_SDK
+
+#include "TechSoftInterface.h"
+
+namespace CADLibrary
+{
+
+// Single-ownership smart TeshSoft object
+// Use this when you need to manage TechSoft object's lifetime.
+//
+// TechSoft give access to void pointers
+// According to the context, the class name of the void pointer is known but the class is unknown
+// i.e. A3DSDKTypes.h defines all type like :
+// 	   typedef void A3DEntity;		
+// 	   typedef void A3DAsmModelFile; ...
+// 
+// From a pointer, TechSoft give access to a copy of the associated structure :
+//
+// const A3DXXXXX* pPointer;
+// A3DXXXXXData sData; // the structure
+// A3D_INITIALIZE_DATA(A3DXXXXXData, sData); // initialization of the structure
+// A3DXXXXXXGet(pPointer, &sData); // Copy of the data of the pointer in the structure
+// ...
+// A3DXXXXXXGet(NULL, &sData); // Free the structure
+//
+// A3D_INITIALIZE_DATA, and all A3DXXXXXXGet methods are TechSoft macro
+//
+
+template<class ObjectType, class IndexerType>
+class TUniqueTSObjBase
+{
+public:
+
+	/**
+	 * Constructor of an initialized ObjectType object
+	 */
+	TUniqueTSObjBase()
+		: bDataFromTechSoft(false)
+		, Status(A3DStatus::A3D_SUCCESS)
+	{
+		InitializeData();
+	}
+
+	/**
+	 * Constructor of an filled ObjectType object with the data of DataPtr
+	 * @param DataPtr: the pointer of the data to copy
+	 */
+	explicit TUniqueTSObjBase(IndexerType DataPtr)
+		: bDataFromTechSoft(false)
+	{
+		InitializeData();
+		FillFrom(DataPtr);
+	}
+
+	~TUniqueTSObjBase()
+	{
+		ResetData();
+	}
+
+	/**
+	 * Fill the structure with the data of a new DataPtr
+	 */
+	A3DStatus FillFrom(IndexerType EntityPtr)
+	{
+		ResetData();
+	
+		if (EntityPtr == GetDefaultIndexerValue())
+		{
+			Status = A3DStatus::A3D_ERROR;
+		}
+		else
+		{
+			Status = GetData(EntityPtr);
+			if (Status == A3DStatus::A3D_SUCCESS)
+			{
+				bDataFromTechSoft = true;
+			}
+		}
+		return Status;
+	}
+
+	template<typename... InArgTypes>
+	A3DStatus FillWith(A3DStatus(*Getter)(const A3DEntity*, ObjectType*, InArgTypes&&...), const A3DEntity* EntityPtr, InArgTypes&&... Args)
+	{
+		ResetData();
+
+		if (EntityPtr == GetDefaultIndexerValue())
+		{
+			Status = A3DStatus::A3D_ERROR;
+		}
+		else
+		{
+			Status = Getter(EntityPtr, &Data, Forward<InArgTypes>(Args)...);
+			if(Status == A3DStatus::A3D_SUCCESS)
+			{
+				bDataFromTechSoft = true;
+			}
+		}
+
+		return Status;
+	}
+
+	/**
+	 * Empty the structure
+	 */
+	void Reset()
+	{
+		ResetData();
+	}
+
+	/**
+	 * Return
+	 *  - A3DStatus::A3D_SUCCESS if the data is filled
+	 *  - A3DStatus::A3D_ERROR if the data is empty
+	 */
+	A3DStatus GetStatus()
+	{
+		return Status;
+	}
+
+	/**
+	 * Return true if the data is filled
+	 */
+	const bool IsValid() const
+	{
+		return Status == A3DStatus::A3D_SUCCESS;
+	}
+
+	// Non-copyable
+	TUniqueTSObjBase(const TUniqueTSObjBase&) = delete;
+	TUniqueTSObjBase& operator=(const TUniqueTSObjBase&) = delete;
+
+	// Conversion methods
+
+	const ObjectType& operator*() const
+	{
+		return Data;
+	}
+
+	ObjectType& operator*()
+	{
+		check(IsValid());
+		return Data;
+	}
+
+	const ObjectType* operator->() const
+	{
+		check(IsValid());
+		return &Data;
+	}
+
+	ObjectType* operator->()
+	{
+		check(IsValid());
+		return &Data;
+	}
+
+	/**
+	 * Return the structure pointer
+	 */
+	ObjectType* GetPtr()
+	{
+		if (Status != A3DStatus::A3D_SUCCESS)
+		{
+			return nullptr;
+		}
+		return &Data;
+	}
+
+private:
+	ObjectType Data;
+	bool bDataFromTechSoft = false;
+	A3DStatus Status = A3DStatus::A3D_ERROR;
+
+	/**
+	 * DefaultValue is used to initialize "Data" with GetData method
+	 * According to IndexerType, the value is either nullptr for const A3DEntity* either something like "A3D_DEFAULT_MATERIAL_INDEX" ((A3DUns16)-1) for uint32
+	 * @see ResetData
+	 */
+	CADINTERFACES_API IndexerType GetDefaultIndexerValue() const;
+
+	CADINTERFACES_API void InitializeData()
+#ifdef USE_TECHSOFT_SDK
+		;
+#else
+	{
+		return A3DStatus::A3D_ERROR;
+	}
+#endif
+
+	CADINTERFACES_API A3DStatus GetData(IndexerType AsmModelFilePtr);
+#ifdef USE_TECHSOFT_SDK
+	;
+#else
+	{
+		return A3DStatus::A3D_ERROR;
+	}
+#endif
+
+	void ResetData()
+	{
+		if(bDataFromTechSoft)
+		{
+			GetData(GetDefaultIndexerValue());
+		}
+		else
+		{
+			InitializeData();
+		}
+		Status = A3DStatus::A3D_SUCCESS;
+		bDataFromTechSoft = false;
+	}
+};
+
+
+template<class ObjectType>
+using TUniqueTSObj = TUniqueTSObjBase<ObjectType, const A3DEntity*>;
+
+template<class ObjectType>
+using TUniqueTSObjFromIndex = TUniqueTSObjBase<ObjectType, uint32>;
+
+
+// TUniqueTSObj -----------------------------------
+
+CADINTERFACES_API void TUniqueTSObj<A3DAsmModelFileData>::InitializeData();
+CADINTERFACES_API void TUniqueTSObj<A3DAsmPartDefinitionData>::InitializeData();
+CADINTERFACES_API void TUniqueTSObj<A3DAsmProductOccurrenceData>::InitializeData();
+CADINTERFACES_API void TUniqueTSObj<A3DAsmProductOccurrenceDataCV5>::InitializeData();
+CADINTERFACES_API void TUniqueTSObj<A3DAsmProductOccurrenceDataSLW>::InitializeData();
+CADINTERFACES_API void TUniqueTSObj<A3DAsmProductOccurrenceDataUg>::InitializeData();
+CADINTERFACES_API void TUniqueTSObj<A3DBoundingBoxData>::InitializeData();
+CADINTERFACES_API void TUniqueTSObj<A3DCrvCircleData>::InitializeData();
+CADINTERFACES_API void TUniqueTSObj<A3DCrvCompositeData>::InitializeData();
+CADINTERFACES_API void TUniqueTSObj<A3DCrvEllipseData>::InitializeData();
+CADINTERFACES_API void TUniqueTSObj<A3DCrvHelixData>::InitializeData();
+CADINTERFACES_API void TUniqueTSObj<A3DCrvHyperbolaData>::InitializeData();
+CADINTERFACES_API void TUniqueTSObj<A3DCrvLineData>::InitializeData();
+CADINTERFACES_API void TUniqueTSObj<A3DCrvNurbsData>::InitializeData();
+CADINTERFACES_API void TUniqueTSObj<A3DCrvParabolaData>::InitializeData();
+CADINTERFACES_API void TUniqueTSObj<A3DCrvPolyLineData>::InitializeData();
+CADINTERFACES_API void TUniqueTSObj<A3DCrvTransformData>::InitializeData();
+CADINTERFACES_API void TUniqueTSObj<A3DGlobalData>::InitializeData();
+CADINTERFACES_API void TUniqueTSObj<A3DGraphicsData>::InitializeData();
+CADINTERFACES_API void TUniqueTSObj<A3DIntervalData>::InitializeData();
+CADINTERFACES_API void TUniqueTSObj<A3DMiscAttributeData>::InitializeData();
+CADINTERFACES_API void TUniqueTSObj<A3DMiscCartesianTransformationData>::InitializeData();
+CADINTERFACES_API void TUniqueTSObj<A3DMiscEntityReferenceData>::InitializeData();
+CADINTERFACES_API void TUniqueTSObj<A3DMiscGeneralTransformationData>::InitializeData();
+CADINTERFACES_API void TUniqueTSObj<A3DMiscMaterialPropertiesData>::InitializeData();
+CADINTERFACES_API void TUniqueTSObj<A3DMiscReferenceOnCsysItemData>::InitializeData();
+CADINTERFACES_API void TUniqueTSObj<A3DMiscReferenceOnTessData>::InitializeData();
+CADINTERFACES_API void TUniqueTSObj<A3DMiscReferenceOnTopologyData>::InitializeData();
+CADINTERFACES_API void TUniqueTSObj<A3DMiscSingleAttributeData>::InitializeData();
+CADINTERFACES_API void TUniqueTSObj<A3DRWParamsExportPrcData>::InitializeData();
+CADINTERFACES_API void TUniqueTSObj<A3DRiBrepModelData>::InitializeData();
+CADINTERFACES_API void TUniqueTSObj<A3DRiCoordinateSystemData>::InitializeData();
+CADINTERFACES_API void TUniqueTSObj<A3DRiDirectionData>::InitializeData();
+CADINTERFACES_API void TUniqueTSObj<A3DRiPolyBrepModelData>::InitializeData();
+CADINTERFACES_API void TUniqueTSObj<A3DRiRepresentationItemData>::InitializeData();
+CADINTERFACES_API void TUniqueTSObj<A3DRiSetData>::InitializeData();
+CADINTERFACES_API void TUniqueTSObj<A3DRootBaseData>::InitializeData();
+CADINTERFACES_API void TUniqueTSObj<A3DRootBaseWithGraphicsData>::InitializeData();
+CADINTERFACES_API void TUniqueTSObj<A3DRWParamsTessellationData>::InitializeData();
+CADINTERFACES_API void TUniqueTSObj<A3DSewOptionsData>::InitializeData();
+CADINTERFACES_API void TUniqueTSObj<A3DSurfBlend01Data>::InitializeData();
+CADINTERFACES_API void TUniqueTSObj<A3DSurfBlend02Data>::InitializeData();
+CADINTERFACES_API void TUniqueTSObj<A3DSurfBlend03Data>::InitializeData();
+CADINTERFACES_API void TUniqueTSObj<A3DSurfConeData>::InitializeData();
+CADINTERFACES_API void TUniqueTSObj<A3DSurfCylinderData>::InitializeData();
+CADINTERFACES_API void TUniqueTSObj<A3DSurfCylindricalData>::InitializeData();
+CADINTERFACES_API void TUniqueTSObj<A3DSurfExtrusionData>::InitializeData();
+CADINTERFACES_API void TUniqueTSObj<A3DSurfFromCurvesData>::InitializeData();
+CADINTERFACES_API void TUniqueTSObj<A3DSurfNurbsData>::InitializeData();
+CADINTERFACES_API void TUniqueTSObj<A3DSurfPipeData>::InitializeData();
+CADINTERFACES_API void TUniqueTSObj<A3DSurfPlaneData>::InitializeData();
+CADINTERFACES_API void TUniqueTSObj<A3DSurfRevolutionData>::InitializeData();
+CADINTERFACES_API void TUniqueTSObj<A3DSurfRuledData>::InitializeData();
+CADINTERFACES_API void TUniqueTSObj<A3DSurfSphereData>::InitializeData();
+CADINTERFACES_API void TUniqueTSObj<A3DSurfTorusData>::InitializeData();
+CADINTERFACES_API void TUniqueTSObj<A3DTess3DData>::InitializeData();
+CADINTERFACES_API void TUniqueTSObj<A3DTessBaseData>::InitializeData();
+CADINTERFACES_API void TUniqueTSObj<A3DTopoBodyData>::InitializeData();
+CADINTERFACES_API void TUniqueTSObj<A3DTopoBrepDataData>::InitializeData();
+CADINTERFACES_API void TUniqueTSObj<A3DTopoCoEdgeData>::InitializeData();
+CADINTERFACES_API void TUniqueTSObj<A3DTopoConnexData>::InitializeData();
+CADINTERFACES_API void TUniqueTSObj<A3DTopoContextData>::InitializeData();
+CADINTERFACES_API void TUniqueTSObj<A3DTopoEdgeData>::InitializeData();
+CADINTERFACES_API void TUniqueTSObj<A3DTopoFaceData>::InitializeData();
+CADINTERFACES_API void TUniqueTSObj<A3DTopoLoopData>::InitializeData();
+CADINTERFACES_API void TUniqueTSObj<A3DTopoShellData>::InitializeData();
+CADINTERFACES_API void TUniqueTSObj<A3DTopoUniqueVertexData>::InitializeData();
+CADINTERFACES_API void TUniqueTSObj<A3DTopoWireEdgeData>::InitializeData();
+
+CADINTERFACES_API A3DStatus TUniqueTSObj<A3DAsmModelFileData>::GetData(const A3DEntity* InEntityPtr);
+CADINTERFACES_API A3DStatus TUniqueTSObj<A3DAsmPartDefinitionData>::GetData(const A3DEntity* InEntityPtr);
+CADINTERFACES_API A3DStatus TUniqueTSObj<A3DAsmProductOccurrenceData>::GetData(const A3DEntity* InEntityPtr);
+CADINTERFACES_API A3DStatus TUniqueTSObj<A3DAsmProductOccurrenceDataCV5>::GetData(const A3DEntity* InEntityPtr);
+CADINTERFACES_API A3DStatus TUniqueTSObj<A3DAsmProductOccurrenceDataSLW>::GetData(const A3DEntity* InEntityPtr);
+CADINTERFACES_API A3DStatus TUniqueTSObj<A3DAsmProductOccurrenceDataUg>::GetData(const A3DEntity* InEntityPtr);
+CADINTERFACES_API A3DStatus TUniqueTSObj<A3DBoundingBoxData>::GetData(const A3DEntity* InEntityPtr);
+CADINTERFACES_API A3DStatus TUniqueTSObj<A3DCrvCircleData>::GetData(const A3DEntity* InEntityPtr);
+CADINTERFACES_API A3DStatus TUniqueTSObj<A3DCrvCompositeData>::GetData(const A3DEntity* InEntityPtr);
+CADINTERFACES_API A3DStatus TUniqueTSObj<A3DCrvEllipseData>::GetData(const A3DEntity* InEntityPtr);
+CADINTERFACES_API A3DStatus TUniqueTSObj<A3DCrvHelixData>::GetData(const A3DEntity* InEntityPtr);
+CADINTERFACES_API A3DStatus TUniqueTSObj<A3DCrvHyperbolaData>::GetData(const A3DEntity* InEntityPtr);
+CADINTERFACES_API A3DStatus TUniqueTSObj<A3DCrvLineData>::GetData(const A3DEntity* InEntityPtr);
+CADINTERFACES_API A3DStatus TUniqueTSObj<A3DCrvNurbsData>::GetData(const A3DEntity* InEntityPtr);
+CADINTERFACES_API A3DStatus TUniqueTSObj<A3DCrvParabolaData>::GetData(const A3DEntity* InEntityPtr);
+CADINTERFACES_API A3DStatus TUniqueTSObj<A3DCrvPolyLineData>::GetData(const A3DEntity* InEntityPtr);
+CADINTERFACES_API A3DStatus TUniqueTSObj<A3DCrvTransformData>::GetData(const A3DEntity* InEntityPtr);
+CADINTERFACES_API A3DStatus TUniqueTSObj<A3DGlobalData>::GetData(const A3DEntity* InEntityPtr);
+CADINTERFACES_API A3DStatus TUniqueTSObj<A3DGraphicsData>::GetData(const A3DEntity* InEntityPtr);
+CADINTERFACES_API A3DStatus TUniqueTSObj<A3DIntervalData>::GetData(const A3DEntity* InEntityPtr);
+CADINTERFACES_API A3DStatus TUniqueTSObj<A3DMiscAttributeData>::GetData(const A3DEntity* InEntityPtr);
+CADINTERFACES_API A3DStatus TUniqueTSObj<A3DMiscCartesianTransformationData>::GetData(const A3DEntity* InEntityPtr);
+CADINTERFACES_API A3DStatus TUniqueTSObj<A3DMiscEntityReferenceData>::GetData(const A3DEntity* InEntityPtr);
+CADINTERFACES_API A3DStatus TUniqueTSObj<A3DMiscGeneralTransformationData>::GetData(const A3DEntity* InEntityPtr);
+CADINTERFACES_API A3DStatus TUniqueTSObj<A3DMiscMaterialPropertiesData>::GetData(const A3DEntity* InEntityPtr);
+CADINTERFACES_API A3DStatus TUniqueTSObj<A3DMiscReferenceOnCsysItemData>::GetData(const A3DEntity* InEntityPtr);
+CADINTERFACES_API A3DStatus TUniqueTSObj<A3DMiscReferenceOnTessData>::GetData(const A3DEntity* InEntityPtr);
+CADINTERFACES_API A3DStatus TUniqueTSObj<A3DMiscReferenceOnTopologyData>::GetData(const A3DEntity* InEntityPtr);
+CADINTERFACES_API A3DStatus TUniqueTSObj<A3DMiscSingleAttributeData>::GetData(const A3DEntity* InEntityPtr);
+CADINTERFACES_API A3DStatus TUniqueTSObj<A3DRWParamsExportPrcData>::GetData(const A3DEntity* InEntityPtr);
+CADINTERFACES_API A3DStatus TUniqueTSObj<A3DRiBrepModelData>::GetData(const A3DEntity* InEntityPtr);
+CADINTERFACES_API A3DStatus TUniqueTSObj<A3DRiCoordinateSystemData>::GetData(const A3DEntity* InEntityPtr);
+CADINTERFACES_API A3DStatus TUniqueTSObj<A3DRiDirectionData>::GetData(const A3DEntity* InEntityPtr);
+CADINTERFACES_API A3DStatus TUniqueTSObj<A3DRiPolyBrepModelData>::GetData(const A3DEntity* InEntityPtr);
+CADINTERFACES_API A3DStatus TUniqueTSObj<A3DRiRepresentationItemData>::GetData(const A3DEntity* InEntityPtr);
+CADINTERFACES_API A3DStatus TUniqueTSObj<A3DRiSetData>::GetData(const A3DEntity* InEntityPtr);
+CADINTERFACES_API A3DStatus TUniqueTSObj<A3DRootBaseData>::GetData(const A3DEntity* InEntityPtr);
+CADINTERFACES_API A3DStatus TUniqueTSObj<A3DRootBaseWithGraphicsData>::GetData(const A3DEntity* InEntityPtr);
+CADINTERFACES_API A3DStatus TUniqueTSObj<A3DRWParamsTessellationData>::GetData(const A3DEntity* InEntityPtr);
+CADINTERFACES_API A3DStatus TUniqueTSObj<A3DSewOptionsData>::GetData(const A3DEntity* InEntityPtr);
+CADINTERFACES_API A3DStatus TUniqueTSObj<A3DSurfBlend01Data>::GetData(const A3DEntity* InEntityPtr);
+CADINTERFACES_API A3DStatus TUniqueTSObj<A3DSurfBlend02Data>::GetData(const A3DEntity* InEntityPtr);
+CADINTERFACES_API A3DStatus TUniqueTSObj<A3DSurfBlend03Data>::GetData(const A3DEntity* InEntityPtr);
+CADINTERFACES_API A3DStatus TUniqueTSObj<A3DSurfConeData>::GetData(const A3DEntity* InEntityPtr);
+CADINTERFACES_API A3DStatus TUniqueTSObj<A3DSurfCylinderData>::GetData(const A3DEntity* InEntityPtr);
+CADINTERFACES_API A3DStatus TUniqueTSObj<A3DSurfCylindricalData>::GetData(const A3DEntity* InEntityPtr);
+CADINTERFACES_API A3DStatus TUniqueTSObj<A3DSurfExtrusionData>::GetData(const A3DEntity* InEntityPtr);
+CADINTERFACES_API A3DStatus TUniqueTSObj<A3DSurfFromCurvesData>::GetData(const A3DEntity* InEntityPtr);
+CADINTERFACES_API A3DStatus TUniqueTSObj<A3DSurfNurbsData>::GetData(const A3DEntity* InEntityPtr);
+CADINTERFACES_API A3DStatus TUniqueTSObj<A3DSurfPipeData>::GetData(const A3DEntity* InEntityPtr);
+CADINTERFACES_API A3DStatus TUniqueTSObj<A3DSurfPlaneData>::GetData(const A3DEntity* InEntityPtr);
+CADINTERFACES_API A3DStatus TUniqueTSObj<A3DSurfRevolutionData>::GetData(const A3DEntity* InEntityPtr);
+CADINTERFACES_API A3DStatus TUniqueTSObj<A3DSurfRuledData>::GetData(const A3DEntity* InEntityPtr);
+CADINTERFACES_API A3DStatus TUniqueTSObj<A3DSurfSphereData>::GetData(const A3DEntity* InEntityPtr);
+CADINTERFACES_API A3DStatus TUniqueTSObj<A3DSurfTorusData>::GetData(const A3DEntity* InEntityPtr);
+CADINTERFACES_API A3DStatus TUniqueTSObj<A3DTess3DData>::GetData(const A3DEntity* InEntityPtr);
+CADINTERFACES_API A3DStatus TUniqueTSObj<A3DTessBaseData>::GetData(const A3DEntity* InEntityPtr);
+CADINTERFACES_API A3DStatus TUniqueTSObj<A3DTopoBodyData>::GetData(const A3DEntity* InEntityPtr);
+CADINTERFACES_API A3DStatus TUniqueTSObj<A3DTopoBrepDataData>::GetData(const A3DEntity* InEntityPtr);
+CADINTERFACES_API A3DStatus TUniqueTSObj<A3DTopoCoEdgeData>::GetData(const A3DEntity* InEntityPtr);
+CADINTERFACES_API A3DStatus TUniqueTSObj<A3DTopoConnexData>::GetData(const A3DEntity* InEntityPtr);
+CADINTERFACES_API A3DStatus TUniqueTSObj<A3DTopoContextData>::GetData(const A3DEntity* InEntityPtr);
+CADINTERFACES_API A3DStatus TUniqueTSObj<A3DTopoEdgeData>::GetData(const A3DEntity* InEntityPtr);
+CADINTERFACES_API A3DStatus TUniqueTSObj<A3DTopoFaceData>::GetData(const A3DEntity* InEntityPtr);
+CADINTERFACES_API A3DStatus TUniqueTSObj<A3DTopoLoopData>::GetData(const A3DEntity* InEntityPtr);
+CADINTERFACES_API A3DStatus TUniqueTSObj<A3DTopoShellData>::GetData(const A3DEntity* InEntityPtr);
+CADINTERFACES_API A3DStatus TUniqueTSObj<A3DTopoUniqueVertexData>::GetData(const A3DEntity* InEntityPtr);
+CADINTERFACES_API A3DStatus TUniqueTSObj<A3DTopoWireEdgeData>::GetData(const A3DEntity* InEntityPtr);
+
+CADINTERFACES_API const A3DEntity* TUniqueTSObj<A3DAsmModelFileData>::GetDefaultIndexerValue() const;
+CADINTERFACES_API const A3DEntity* TUniqueTSObj<A3DAsmPartDefinitionData>::GetDefaultIndexerValue() const;
+CADINTERFACES_API const A3DEntity* TUniqueTSObj<A3DAsmProductOccurrenceData>::GetDefaultIndexerValue() const;
+CADINTERFACES_API const A3DEntity* TUniqueTSObj<A3DAsmProductOccurrenceDataCV5>::GetDefaultIndexerValue() const;
+CADINTERFACES_API const A3DEntity* TUniqueTSObj<A3DAsmProductOccurrenceDataSLW>::GetDefaultIndexerValue() const;
+CADINTERFACES_API const A3DEntity* TUniqueTSObj<A3DAsmProductOccurrenceDataUg>::GetDefaultIndexerValue() const;
+CADINTERFACES_API const A3DEntity* TUniqueTSObj<A3DBoundingBoxData>::GetDefaultIndexerValue() const;
+CADINTERFACES_API const A3DEntity* TUniqueTSObj<A3DCrvCircleData>::GetDefaultIndexerValue() const;
+CADINTERFACES_API const A3DEntity* TUniqueTSObj<A3DCrvCompositeData>::GetDefaultIndexerValue() const;
+CADINTERFACES_API const A3DEntity* TUniqueTSObj<A3DCrvEllipseData>::GetDefaultIndexerValue() const;
+CADINTERFACES_API const A3DEntity* TUniqueTSObj<A3DCrvHelixData>::GetDefaultIndexerValue() const;
+CADINTERFACES_API const A3DEntity* TUniqueTSObj<A3DCrvHyperbolaData>::GetDefaultIndexerValue() const;
+CADINTERFACES_API const A3DEntity* TUniqueTSObj<A3DCrvLineData>::GetDefaultIndexerValue() const;
+CADINTERFACES_API const A3DEntity* TUniqueTSObj<A3DCrvNurbsData>::GetDefaultIndexerValue() const;
+CADINTERFACES_API const A3DEntity* TUniqueTSObj<A3DCrvParabolaData>::GetDefaultIndexerValue() const;
+CADINTERFACES_API const A3DEntity* TUniqueTSObj<A3DCrvPolyLineData>::GetDefaultIndexerValue() const;
+CADINTERFACES_API const A3DEntity* TUniqueTSObj<A3DCrvTransformData>::GetDefaultIndexerValue() const;
+CADINTERFACES_API const A3DEntity* TUniqueTSObj<A3DGlobalData>::GetDefaultIndexerValue() const;
+CADINTERFACES_API const A3DEntity* TUniqueTSObj<A3DGraphicsData>::GetDefaultIndexerValue() const;
+CADINTERFACES_API const A3DEntity* TUniqueTSObj<A3DIntervalData>::GetDefaultIndexerValue() const;
+CADINTERFACES_API const A3DEntity* TUniqueTSObj<A3DMiscAttributeData>::GetDefaultIndexerValue() const;
+CADINTERFACES_API const A3DEntity* TUniqueTSObj<A3DMiscCartesianTransformationData>::GetDefaultIndexerValue() const;
+CADINTERFACES_API const A3DEntity* TUniqueTSObj<A3DMiscEntityReferenceData>::GetDefaultIndexerValue() const;
+CADINTERFACES_API const A3DEntity* TUniqueTSObj<A3DMiscGeneralTransformationData>::GetDefaultIndexerValue() const;
+CADINTERFACES_API const A3DEntity* TUniqueTSObj<A3DMiscMaterialPropertiesData>::GetDefaultIndexerValue() const;
+CADINTERFACES_API const A3DEntity* TUniqueTSObj<A3DMiscReferenceOnCsysItemData>::GetDefaultIndexerValue() const;
+CADINTERFACES_API const A3DEntity* TUniqueTSObj<A3DMiscReferenceOnTessData>::GetDefaultIndexerValue() const;
+CADINTERFACES_API const A3DEntity* TUniqueTSObj<A3DMiscReferenceOnTopologyData>::GetDefaultIndexerValue() const;
+CADINTERFACES_API const A3DEntity* TUniqueTSObj<A3DMiscSingleAttributeData>::GetDefaultIndexerValue() const;
+CADINTERFACES_API const A3DEntity* TUniqueTSObj<A3DRWParamsExportPrcData>::GetDefaultIndexerValue() const;
+CADINTERFACES_API const A3DEntity* TUniqueTSObj<A3DRiBrepModelData>::GetDefaultIndexerValue() const;
+CADINTERFACES_API const A3DEntity* TUniqueTSObj<A3DRiCoordinateSystemData>::GetDefaultIndexerValue() const;
+CADINTERFACES_API const A3DEntity* TUniqueTSObj<A3DRiDirectionData>::GetDefaultIndexerValue() const;
+CADINTERFACES_API const A3DEntity* TUniqueTSObj<A3DRiPolyBrepModelData>::GetDefaultIndexerValue() const;
+CADINTERFACES_API const A3DEntity* TUniqueTSObj<A3DRiRepresentationItemData>::GetDefaultIndexerValue() const;
+CADINTERFACES_API const A3DEntity* TUniqueTSObj<A3DRiSetData>::GetDefaultIndexerValue() const;
+CADINTERFACES_API const A3DEntity* TUniqueTSObj<A3DRootBaseData>::GetDefaultIndexerValue() const;
+CADINTERFACES_API const A3DEntity* TUniqueTSObj<A3DRootBaseWithGraphicsData>::GetDefaultIndexerValue() const;
+CADINTERFACES_API const A3DEntity* TUniqueTSObj<A3DRWParamsTessellationData>::GetDefaultIndexerValue() const;
+CADINTERFACES_API const A3DEntity* TUniqueTSObj<A3DSewOptionsData>::GetDefaultIndexerValue() const;
+CADINTERFACES_API const A3DEntity* TUniqueTSObj<A3DSurfBlend01Data>::GetDefaultIndexerValue() const;
+CADINTERFACES_API const A3DEntity* TUniqueTSObj<A3DSurfBlend02Data>::GetDefaultIndexerValue() const;
+CADINTERFACES_API const A3DEntity* TUniqueTSObj<A3DSurfBlend03Data>::GetDefaultIndexerValue() const;
+CADINTERFACES_API const A3DEntity* TUniqueTSObj<A3DSurfConeData>::GetDefaultIndexerValue() const;
+CADINTERFACES_API const A3DEntity* TUniqueTSObj<A3DSurfCylinderData>::GetDefaultIndexerValue() const;
+CADINTERFACES_API const A3DEntity* TUniqueTSObj<A3DSurfCylindricalData>::GetDefaultIndexerValue() const;
+CADINTERFACES_API const A3DEntity* TUniqueTSObj<A3DSurfExtrusionData>::GetDefaultIndexerValue() const;
+CADINTERFACES_API const A3DEntity* TUniqueTSObj<A3DSurfFromCurvesData>::GetDefaultIndexerValue() const;
+CADINTERFACES_API const A3DEntity* TUniqueTSObj<A3DSurfNurbsData>::GetDefaultIndexerValue() const;
+CADINTERFACES_API const A3DEntity* TUniqueTSObj<A3DSurfPipeData>::GetDefaultIndexerValue() const;
+CADINTERFACES_API const A3DEntity* TUniqueTSObj<A3DSurfPlaneData>::GetDefaultIndexerValue() const;
+CADINTERFACES_API const A3DEntity* TUniqueTSObj<A3DSurfRevolutionData>::GetDefaultIndexerValue() const;
+CADINTERFACES_API const A3DEntity* TUniqueTSObj<A3DSurfRuledData>::GetDefaultIndexerValue() const;
+CADINTERFACES_API const A3DEntity* TUniqueTSObj<A3DSurfSphereData>::GetDefaultIndexerValue() const;
+CADINTERFACES_API const A3DEntity* TUniqueTSObj<A3DSurfTorusData>::GetDefaultIndexerValue() const;
+CADINTERFACES_API const A3DEntity* TUniqueTSObj<A3DTess3DData>::GetDefaultIndexerValue() const;
+CADINTERFACES_API const A3DEntity* TUniqueTSObj<A3DTessBaseData>::GetDefaultIndexerValue() const;
+CADINTERFACES_API const A3DEntity* TUniqueTSObj<A3DTopoBodyData>::GetDefaultIndexerValue() const;
+CADINTERFACES_API const A3DEntity* TUniqueTSObj<A3DTopoBrepDataData>::GetDefaultIndexerValue() const;
+CADINTERFACES_API const A3DEntity* TUniqueTSObj<A3DTopoCoEdgeData>::GetDefaultIndexerValue() const;
+CADINTERFACES_API const A3DEntity* TUniqueTSObj<A3DTopoConnexData>::GetDefaultIndexerValue() const;
+CADINTERFACES_API const A3DEntity* TUniqueTSObj<A3DTopoContextData>::GetDefaultIndexerValue() const;
+CADINTERFACES_API const A3DEntity* TUniqueTSObj<A3DTopoEdgeData>::GetDefaultIndexerValue() const;
+CADINTERFACES_API const A3DEntity* TUniqueTSObj<A3DTopoFaceData>::GetDefaultIndexerValue() const;
+CADINTERFACES_API const A3DEntity* TUniqueTSObj<A3DTopoLoopData>::GetDefaultIndexerValue() const;
+CADINTERFACES_API const A3DEntity* TUniqueTSObj<A3DTopoShellData>::GetDefaultIndexerValue() const;
+CADINTERFACES_API const A3DEntity* TUniqueTSObj<A3DTopoUniqueVertexData>::GetDefaultIndexerValue() const;
+CADINTERFACES_API const A3DEntity* TUniqueTSObj<A3DTopoWireEdgeData>::GetDefaultIndexerValue() const;
+
+// TUniqueTSObjFromIndex -----------------------------------
+
+CADINTERFACES_API void TUniqueTSObj<A3DGraphMaterialData>::InitializeData();
+CADINTERFACES_API void TUniqueTSObj<A3DGraphPictureData>::InitializeData();
+CADINTERFACES_API void TUniqueTSObj<A3DGraphRgbColorData>::InitializeData();
+CADINTERFACES_API void TUniqueTSObj<A3DGraphStyleData>::InitializeData();
+CADINTERFACES_API void TUniqueTSObj<A3DGraphTextureApplicationData>::InitializeData();
+CADINTERFACES_API void TUniqueTSObj<A3DGraphTextureDefinitionData>::InitializeData();
+
+CADINTERFACES_API A3DStatus TUniqueTSObjFromIndex<A3DGraphMaterialData>::GetData(const uint32 InEntityIndex);
+CADINTERFACES_API A3DStatus TUniqueTSObjFromIndex<A3DGraphPictureData>::GetData(const uint32 InEntityIndex);
+CADINTERFACES_API A3DStatus TUniqueTSObjFromIndex<A3DGraphRgbColorData>::GetData(const uint32 InEntityIndex);
+CADINTERFACES_API A3DStatus TUniqueTSObjFromIndex<A3DGraphStyleData>::GetData(const uint32 InEntityIndex);
+CADINTERFACES_API A3DStatus TUniqueTSObjFromIndex<A3DGraphTextureApplicationData>::GetData(const uint32 InEntityIndex);
+CADINTERFACES_API A3DStatus TUniqueTSObjFromIndex<A3DGraphTextureDefinitionData>::GetData(const uint32 InEntityIndex);
+
+CADINTERFACES_API uint32 TUniqueTSObjFromIndex<A3DGraphMaterialData>::GetDefaultIndexerValue() const;
+CADINTERFACES_API uint32 TUniqueTSObjFromIndex<A3DGraphPictureData>::GetDefaultIndexerValue() const;
+CADINTERFACES_API uint32 TUniqueTSObjFromIndex<A3DGraphRgbColorData>::GetDefaultIndexerValue() const;
+CADINTERFACES_API uint32 TUniqueTSObjFromIndex<A3DGraphStyleData>::GetDefaultIndexerValue() const;
+CADINTERFACES_API uint32 TUniqueTSObjFromIndex<A3DGraphTextureApplicationData>::GetDefaultIndexerValue() const;
+CADINTERFACES_API uint32 TUniqueTSObjFromIndex<A3DGraphTextureDefinitionData>::GetDefaultIndexerValue() const;
+
+// A3DUTF8Char* -----------------------------------
+
+CADINTERFACES_API void TUniqueTSObj<A3DUTF8Char*>::InitializeData();
+CADINTERFACES_API A3DStatus TUniqueTSObj<A3DUTF8Char*>::GetData(const A3DEntity* InEntityPtr);
+CADINTERFACES_API const A3DEntity* TUniqueTSObj<A3DUTF8Char*>::GetDefaultIndexerValue() const;
+
+
+}
+
+#endif
\ No newline at end of file
diff --git a/Engine/Plugins/Enterprise/DatasmithCADImporter/Source/CADInterfaces/Public/TechSoftInterface.h b/Engine/Plugins/Enterprise/DatasmithCADImporter/Source/CADInterfaces/Public/TechSoftInterface.h
--- a/Engine/Plugins/Enterprise/DatasmithCADImporter/Source/CADInterfaces/Public/TechSoftInterface.h
+++ b/Engine/Plugins/Enterprise/DatasmithCADImporter/Source/CADInterfaces/Public/TechSoftInterface.h
@@ -14,7 +14,10 @@
 THIRD_PARTY_INCLUDES_START
 #ifdef USE_TECHSOFT_SDK
 #pragma push_macro("TEXT")
+#pragma warning(push)
+#pragma warning(disable:4191) // unsafe sprintf
 #include "A3DSDKIncludes.h"
+#pragma warning(pop)
 #pragma pop_macro("TEXT")
 #endif
 THIRD_PARTY_INCLUDES_END
@@ -27,64 +30,67 @@ THIRD_PARTY_INCLUDES_END
 #define JSON_ENTRY_COLOR_NAME		TEXT("ColorName")
 #define JSON_ENTRY_MATERIAL_NAME	TEXT("MaterialName")
 
+typedef void A3DAsmModelFile;
+
 namespace CADLibrary
 {
-	class FImportParameters;
-	class FBodyMesh;
 
-CADINTERFACES_API class FTechSoftInterface
+class FBodyMesh;
+class FImportParameters;
+class FUniqueTechSoftModelFile;
+
+class CADINTERFACES_API FTechSoftInterface
 {
 public:
+
+	static FTechSoftInterface& Get();
+
 	/*
 	* Returns true if the object has been created outside of the memory pool of the running process
 	* This is the case when the object has been created by the DatasmithRuntime plugin
 	*/
-	CADINTERFACES_API bool IsExternal()
+	bool IsExternal()
 	{
 		return bIsExternal;
 	}
 
-	CADINTERFACES_API void SetExternal(bool Value)
+	void SetExternal(bool Value)
 	{
 		bIsExternal = Value;
 	}
 
-	CADINTERFACES_API bool InitializeKernel(const TCHAR* = TEXT(""));
-
-#ifdef USE_TECHSOFT_SDK
-	A3DStatus Import(const A3DImport& Import);
-	A3DAsmModelFile* GetModelFile();
-	A3DStatus UnloadModel();
-
-	static A3DUns32 InvalidScriptIndex;
-#endif
-
-	CADINTERFACES_API bool GetBodyFromHsfFile(const FString& Filename, const FImportParameters& ImportParameters, FBodyMesh& BodyMesh);
+	bool InitializeKernel(const TCHAR* = TEXT(""));
 
-	bool FillBodyMesh(void* BodyPtr, const FImportParameters& ImportParameters, double FileUnit, FBodyMesh& BodyMesh);
-
-	void SaveBodyToHsfFile(void* BodyPtr, const FString& Filename, const FString& JsonString);
+	static uint32 InvalidScriptIndex;
 
 private:
 
 	bool bIsExternal = false;
 	bool bIsInitialize = false;
-
-#ifdef USE_TECHSOFT_SDK
-	TUniquePtr<A3DSDKHOOPSExchangeLoader> ExchangeLoader;
-#endif
 };
 
-namespace TechSoftUtils
+namespace TechSoftInterface
 {
+
+/**
+ * Single-ownership smart TeshSoft file object
+ * This object ensure to delete the FileModel when it no more used
+ */
+
 CADINTERFACES_API bool TECHSOFT_InitializeKernel(const TCHAR* = TEXT(""));
-CADINTERFACES_API FTechSoftInterface& GetTechSoftInterface();
+
 CADINTERFACES_API FString GetTechSoftVersion();
 
 #ifdef USE_TECHSOFT_SDK
-FColor GetColorAt(uint32 ColorIndex);
 
-CADINTERFACES_API A3DStatus GetGlobalPointer(A3DGlobal** GlobalPtr);
+CADINTERFACES_API FUniqueTechSoftModelFile LoadModelFileFromFile(const A3DImport& Import);
+CADINTERFACES_API FUniqueTechSoftModelFile LoadModelFileFromPrcFile(const A3DUTF8Char* CADFileName, A3DRWParamsPrcReadHelper** ReadHelper);
+
+CADINTERFACES_API A3DStatus DeleteModelFile(A3DAsmModelFile* ModelFile);
+CADINTERFACES_API A3DStatus DeleteEntity(A3DEntity* pEntity);
+
+CADINTERFACES_API A3DGlobal* GetGlobalPointer();
+CADINTERFACES_API A3DEntity* GetPointerFromIndex(const uint32 Index, const A3DEEntityType Type);
 
 CADINTERFACES_API A3DStatus GetSurfaceAsNurbs(const A3DSurfBase* SurfacePtr, A3DSurfNurbsData* DataPtr, A3DDouble Tolerance, A3DBool bUseSameParameterization);
 CADINTERFACES_API A3DStatus GetCurveAsNurbs(const A3DCrvBase* A3DCurve, A3DCrvNurbsData* DataPtr, A3DDouble Tolerance, A3DBool bUseSameParameterization);
@@ -97,15 +103,120 @@ CADINTERFACES_API bool IsEntityBaseWithGraphicsType(const A3DEntity* pEntity);
 CADINTERFACES_API bool IsEntityBaseType(const A3DEntity* EntityPtr);
 CADINTERFACES_API bool IsMaterialTexture(const uint32 MaterialIndex);
 
-CADINTERFACES_API A3DStatus HealBRep(A3DRiBrepModel** BRepToHeal, double Tolerance, A3DSewOptionsData const* SewOptions, A3DRiBrepModel*** OutNewBReps, uint32& OutNewBRepCount);
-CADINTERFACES_API A3DStatus SewBReps(A3DRiBrepModel*** BRepsToSew, uint32 const BRepCount, double Tolerance, A3DSewOptionsData const* SewOptions, A3DRiBrepModel*** OutNewBReps, uint32& OutNewBRepCount);
-CADINTERFACES_API A3DStatus SewModel(A3DAsmModelFile** ModelPtr, double Tolerance, A3DSewOptionsData const* SewOptions);
 
-CADINTERFACES_API A3DEntity* GetPointerFromIndex(const uint32 Index, const A3DEEntityType Type);
+CADINTERFACES_API FUniqueTechSoftModelFile CreateAsmModelFile(const A3DAsmModelFileData& ModelFileData);
+
+CADINTERFACES_API A3DRiBrepModel* CreateRiBRepModel(const A3DRiBrepModelData& RiBRepModelData);
+CADINTERFACES_API A3DAsmPartDefinition* CreateAsmPartDefinition(const A3DAsmPartDefinitionData& PartDefinitionData);
+CADINTERFACES_API A3DAsmProductOccurrence* CreateAsmProductOccurrence(const A3DAsmProductOccurrenceData& ProductOccurrenceData);
+CADINTERFACES_API A3DMiscAttribute* CreateMiscAttribute(const A3DMiscAttributeData& AttributeData);
+CADINTERFACES_API A3DAsmModelFile* CreateModelFile(const A3DAsmModelFileData& ModelFileData);
+CADINTERFACES_API A3DTopoBrepData* CreateTopoBRep(const A3DTopoBrepDataData& TopoBRepDataData);
+CADINTERFACES_API A3DTopoCoEdge* CreateTopoCoEdge(const A3DTopoCoEdgeData& TopoCoEdgeData);
+CADINTERFACES_API A3DTopoConnex* CreateTopoConnex(const A3DTopoConnexData& TopoConnexData);
+CADINTERFACES_API A3DTopoEdge* CreateTopoEdge(const A3DTopoEdgeData& TopoEdgeData);
+CADINTERFACES_API A3DTopoFace* CreateTopoFace(const A3DTopoFaceData& TopoFaceData);
+CADINTERFACES_API A3DTopoLoop* CreateTopoLoop(const A3DTopoLoopData& TopoLoopData);
+CADINTERFACES_API A3DTopoShell* CreateTopoShell(const A3DTopoShellData& TopoShellData);
+
+CADINTERFACES_API A3DCrvNurbs* CreateCurveTransform(const A3DCrvTransformData& CurveTransformData);
+CADINTERFACES_API A3DCrvNurbs* CreateCurveNurbs(const A3DCrvNurbsData& CurveNurbsData);
+
+CADINTERFACES_API A3DSurfNurbs* CreateSurfaceCylinder(const A3DSurfCylinderData& SurfaceCylinderData);
+CADINTERFACES_API A3DSurfNurbs* CreateSurfaceNurbs(const A3DSurfNurbsData& SurfaceNurbsData);
+
+CADINTERFACES_API A3DGraphics* CreateGraphics(const A3DGraphicsData& GraphicsData);
+
+CADINTERFACES_API A3DStatus LinkCoEdges(A3DTopoCoEdge* CoEdgePtr, A3DTopoCoEdge* NeighbourCoEdgePtr);
+
+CADINTERFACES_API A3DStatus SetRootBase(A3DEntity* EntityPtr, const A3DRootBaseData& RootBaseData);
+CADINTERFACES_API A3DStatus SetRootBaseWithGraphics(const A3DRootBaseWithGraphicsData& RootBaseWithGraphicsData, A3DRootBaseWithGraphics* RootPtr);
+
+CADINTERFACES_API A3DStatus ExportModelFileToPrcFile(const A3DAsmModelFile* ModelFile, const A3DRWParamsExportPrcData* ParamsExportData, const A3DUTF8Char* CADFileName, A3DRWParamsPrcWriteHelper** PrcWriteHelper);
+
+CADINTERFACES_API A3DUns32 InsertGraphRgbColor(const A3DGraphRgbColorData& InRgbColorData);
+CADINTERFACES_API A3DUns32 InsertGraphMaterial(const A3DGraphMaterialData& InMaterialData);
+CADINTERFACES_API A3DUns32 InsertGraphStyle(const A3DGraphStyleData& InStyleData);
+
+CADINTERFACES_API A3DStatus SewModel(A3DAsmModelFile* ModelPtr, double Tolerance, A3DSewOptionsData const* SewOptions);
+CADINTERFACES_API A3DStatus SewBReps(A3DRiBrepModel** BRepsToSew, uint32 const BRepCount, double Tolerance, A3DSewOptionsData const* SewOptions, A3DRiBrepModel*** OutNewBReps, uint32& OutNewBRepCount);
 
 #endif
 
 } // NS TechSoftUtils
 
-}
+/**
+ * Class to manage model file lifetime
+ */
+class CADINTERFACES_API FUniqueTechSoftModelFile
+{
+public:
+
+	FUniqueTechSoftModelFile()
+		: ModelFile(nullptr)
+	{
+	}
+
+	explicit FUniqueTechSoftModelFile(A3DAsmModelFile* InModelFile)
+		: ModelFile(InModelFile)
+	{
+	}
+
+	~FUniqueTechSoftModelFile()
+	{
+		Reset();
+	}
+
+	/**
+	 * Return true if the model is loaded
+	 */
+	const bool IsValid() const
+	{
+		return ModelFile != nullptr;
+	}
+
+	// Non-copyable
+	FUniqueTechSoftModelFile(const FUniqueTechSoftModelFile&) = delete;
+	FUniqueTechSoftModelFile& operator=(const FUniqueTechSoftModelFile&) = delete;
+
+	FUniqueTechSoftModelFile(FUniqueTechSoftModelFile&& Old)
+		: ModelFile(Old.ModelFile)
+	{
+		Old.ModelFile = nullptr;
+	}
+
+	FUniqueTechSoftModelFile& operator=(FUniqueTechSoftModelFile&& Old)
+	{
+		Reset();
+		ModelFile = Old.ModelFile;
+		Old.ModelFile = nullptr;
+		return *this;
+	}
+
+	A3DAsmModelFile* Get()
+	{
+		return ModelFile;
+	}
+
+	A3DAsmModelFile** GetPtr()
+	{
+		return &ModelFile;
+	}
+
+	void Reset()
+	{
+		if (ModelFile != nullptr)
+		{
+#ifdef USE_TECHSOFT_SDK
+			TechSoftInterface::DeleteModelFile(ModelFile);
+#endif
+			ModelFile = nullptr;
+		}
+	}
+
+private:
+	A3DAsmModelFile* ModelFile;
+};
+
+} // CADLibrary
 
diff --git a/Engine/Plugins/Enterprise/DatasmithCADImporter/Source/CADInterfaces/Public/TechSoftUtils.h b/Engine/Plugins/Enterprise/DatasmithCADImporter/Source/CADInterfaces/Public/TechSoftUtils.h
new file mode 100644
--- /dev/null
+++ b/Engine/Plugins/Enterprise/DatasmithCADImporter/Source/CADInterfaces/Public/TechSoftUtils.h
@@ -0,0 +1,57 @@
+// Copyright Epic Games, Inc. All Rights Reserved.
+#pragma once
+
+#include "CoreMinimal.h"
+
+#include "TechSoftInterface.h"
+
+#define JSON_ENTRY_FILE_UNIT		TEXT("FileUnit")
+#define JSON_ENTRY_COLOR_NAME		TEXT("ColorName")
+#define JSON_ENTRY_MATERIAL_NAME	TEXT("MaterialName")
+
+struct FColor;
+
+typedef void A3DAsmPartDefinition;
+typedef void A3DCrvNurbs;
+typedef void A3DEntity;
+typedef void A3DRiRepresentationItem;
+typedef void A3DSurfBase;
+typedef void A3DTopoEdge;
+typedef void A3DTopoFace;
+typedef void A3DTopoShell;
+
+namespace CADLibrary
+{
+
+class FImportParameters;
+class FBodyMesh;
+
+namespace TechSoftInterface
+{
+class FTechSoftInterface;
+}
+
+
+namespace TechSoftUtils
+{
+
+CADINTERFACES_API bool GetBodyFromPcrFile(const FString& Filename, const FImportParameters& ImportParameters, FBodyMesh& BodyMesh);
+CADINTERFACES_API bool FillBodyMesh(void* BodyPtr, const FImportParameters& ImportParameters, double FileUnit, FBodyMesh& BodyMesh);
+
+CADINTERFACES_API FUniqueTechSoftModelFile SaveBodiesToPcrFile(void** Bodies, uint32 BodyCount, const FString& Filename, const FString& JsonString);
+CADINTERFACES_API void SaveModelFileToPcrFile(void* ModelFile, const FString& Filename);
+
+CADINTERFACES_API int32 SetEntityGraphicsColor(A3DEntity* InEntity, FColor Color);
+
+CADINTERFACES_API A3DAsmPartDefinition* CreatePart(TArray<A3DRiRepresentationItem*>& RepresentationItems);
+CADINTERFACES_API A3DRiRepresentationItem* CreateRIBRep(A3DTopoShell* TopoShellPtr);
+CADINTERFACES_API A3DTopoEdge* CreateTopoEdge();
+CADINTERFACES_API A3DTopoFace* CreateTopoFaceWithNaturalLoop(A3DSurfBase* CarrierSurface);
+CADINTERFACES_API A3DCrvNurbs* CreateTrimNurbsCurve(A3DCrvNurbs* CurveNurbsPtr, double UMin, double UMax, bool bIs2D);
+
+CADINTERFACES_API FColor GetColorAt(uint32 ColorIndex);
+
+} // NS TechSoftUtils
+
+} // CADLibrary
+
diff --git a/Engine/Plugins/Enterprise/DatasmithCADImporter/Source/CADKernelSurface/Public/CADModelToCADKernelConverterBase.h b/Engine/Plugins/Enterprise/DatasmithCADImporter/Source/CADKernelSurface/Public/CADModelToCADKernelConverterBase.h
--- a/Engine/Plugins/Enterprise/DatasmithCADImporter/Source/CADKernelSurface/Public/CADModelToCADKernelConverterBase.h
+++ b/Engine/Plugins/Enterprise/DatasmithCADImporter/Source/CADKernelSurface/Public/CADModelToCADKernelConverterBase.h
@@ -43,12 +43,13 @@ public:
 			double JoiningTolerance = ImportParameters.ConvertMMToImportUnit(0.1);
 			CADKernel::FTopomaker Topomaker(CADKernelSession, JoiningTolerance);
 			Topomaker.Sew();
+			Topomaker.OrientShells();
 		}
 
 		return true;
 	}
 
-	virtual bool SaveBRep(const TCHAR* InFolderPath, TSharedRef<IDatasmithMeshElement>& MeshElement) override
+	virtual bool SaveModel(const TCHAR* InFolderPath, TSharedRef<IDatasmithMeshElement>& MeshElement) override
 	{
 		FString FilePath = FPaths::Combine(InFolderPath, MeshElement->GetName()) + TEXT(".ugeom");
 		CADKernelSession.SaveDatabase(*FilePath);
diff --git a/Engine/Plugins/Enterprise/DatasmithCADImporter/Source/CoreTechSurface/Private/CoreTechSurfaceExtension.cpp b/Engine/Plugins/Enterprise/DatasmithCADImporter/Source/CoreTechSurface/Private/CoreTechSurfaceExtension.cpp
--- a/Engine/Plugins/Enterprise/DatasmithCADImporter/Source/CoreTechSurface/Private/CoreTechSurfaceExtension.cpp
+++ b/Engine/Plugins/Enterprise/DatasmithCADImporter/Source/CoreTechSurface/Private/CoreTechSurfaceExtension.cpp
@@ -41,8 +41,11 @@ bool UTempCoreTechParametricSurfaceData::Tessellate(UStaticMesh& StaticMesh, con
 #if WITH_EDITOR
 	// make a temporary file as CoreTech can only deal with files.
 	int32 Hash = GetTypeHash(StaticMesh.GetPathName());
-	FString ResourceFile = FPaths::ConvertRelativePathToFull(FPaths::ProjectIntermediateDir() / FString::Printf(TEXT("0x%08x.ct"), Hash));
 
+	FString CachePath = FPaths::ConvertRelativePathToFull(FPaths::Combine(FPaths::ProjectIntermediateDir(), TEXT("Retessellate")));
+	IFileManager::Get().MakeDirectory(*CachePath, true);
+
+	FString ResourceFile = FPaths::ConvertRelativePathToFull(FPaths::Combine(CachePath, FString::Printf(TEXT("0x%08x.ct"), Hash)));
 	FFileHelper::SaveArrayToFile(RawData, *ResourceFile);
 
 	CADLibrary::FCTMesh Mesh;
diff --git a/Engine/Plugins/Enterprise/DatasmithCADImporter/Source/CoreTechSurface/Private/CoreTechSurfaceHelper.cpp b/Engine/Plugins/Enterprise/DatasmithCADImporter/Source/CoreTechSurface/Private/CoreTechSurfaceHelper.cpp
--- a/Engine/Plugins/Enterprise/DatasmithCADImporter/Source/CoreTechSurface/Private/CoreTechSurfaceHelper.cpp
+++ b/Engine/Plugins/Enterprise/DatasmithCADImporter/Source/CoreTechSurface/Private/CoreTechSurfaceHelper.cpp
@@ -62,22 +62,4 @@ namespace CoreTechSurface
 
 		return Tessellate(MainObjectID, ImportParameters, MeshParameters, MeshDescription);
 	}
-
-	void AddSurfaceDataForMesh(const TCHAR* InFilePath, const CADLibrary::FImportParameters& InSceneParameters, const CADLibrary::FMeshParameters& InMeshParameters, const FDatasmithTessellationOptions& InTessellationOptions, FDatasmithMeshElementPayload& OutMeshPayload)
-	{
-		if (ICADInterfacesModule::GetAvailability() == ECADInterfaceAvailability::Available)
-		{
-			UTempCoreTechParametricSurfaceData* CoreTechData = Datasmith::MakeAdditionalData<UTempCoreTechParametricSurfaceData>();
-
-			if (CoreTechData->SetFile(InFilePath))
-			{
-				CoreTechData->SetImportParameters(InSceneParameters);
-				CoreTechData->SetMeshParameters(InMeshParameters);
-				CoreTechData->SetLastTessellationOptions(InTessellationOptions);
-
-				OutMeshPayload.AdditionalData.Add(CoreTechData);
-			}
-		}
-	}
-
 }
diff --git a/Engine/Plugins/Enterprise/DatasmithCADImporter/Source/CoreTechSurface/Public/CADModelToCoretechConverterBase.h b/Engine/Plugins/Enterprise/DatasmithCADImporter/Source/CoreTechSurface/Public/CADModelToCoretechConverterBase.h
--- a/Engine/Plugins/Enterprise/DatasmithCADImporter/Source/CoreTechSurface/Public/CADModelToCoretechConverterBase.h
+++ b/Engine/Plugins/Enterprise/DatasmithCADImporter/Source/CoreTechSurface/Public/CADModelToCoretechConverterBase.h
@@ -9,6 +9,7 @@
 #include "CoreTechSurfaceHelper.h"
 #include "IDatasmithSceneElements.h"
 #include "Utility/DatasmithMeshHelper.h"
+#include "ParametricSurfaceTranslator.h"
 
 class FCADModelToCoretechConverterBase : public CADLibrary::FCTSession, public CADLibrary::ICADModelConverter
 {
@@ -38,7 +39,7 @@ public:
 		SetSceneUnit(NewSceneunit);
 	}
 
-	virtual bool SaveBRep(const TCHAR* InFolderPath, TSharedRef<IDatasmithMeshElement>& MeshElement) override
+	virtual bool SaveModel(const TCHAR* InFolderPath, TSharedRef<IDatasmithMeshElement>& MeshElement) override
 	{
 		FString FilePath = FPaths::Combine(InFolderPath, MeshElement->GetName()) + TEXT(".ct");
 		if (SaveBrep(FilePath))
@@ -77,7 +78,7 @@ public:
 	void AddSurfaceDataForMesh(const TCHAR* InFilePath, const CADLibrary::FMeshParameters& InMeshParameters, const FDatasmithTessellationOptions& InTessellationOptions, FDatasmithMeshElementPayload& OutMeshPayload) const
 	{
 		const CADLibrary::FImportParameters& InImportParameters = GetImportParameters();
-		CoreTechSurface::AddSurfaceDataForMesh(InFilePath, InImportParameters, InMeshParameters, InTessellationOptions, OutMeshPayload);
+		ParametricSurfaceUtils::AddSurfaceData(InFilePath, InImportParameters, InMeshParameters, InTessellationOptions, OutMeshPayload);
 	}
 
 };
diff --git a/Engine/Plugins/Enterprise/DatasmithCADImporter/Source/CoreTechSurface/Public/CoreTechSurfaceHelper.h b/Engine/Plugins/Enterprise/DatasmithCADImporter/Source/CoreTechSurface/Public/CoreTechSurfaceHelper.h
--- a/Engine/Plugins/Enterprise/DatasmithCADImporter/Source/CoreTechSurface/Public/CoreTechSurfaceHelper.h
+++ b/Engine/Plugins/Enterprise/DatasmithCADImporter/Source/CoreTechSurface/Public/CoreTechSurfaceHelper.h
@@ -15,7 +15,6 @@ struct FMeshDescription;
 
 namespace CoreTechSurface
 {
-	CORETECHSURFACE_API void AddSurfaceDataForMesh(const TCHAR* CoretechFile, const CADLibrary::FImportParameters& InSceneParameters, const CADLibrary::FMeshParameters&, const FDatasmithTessellationOptions& InTessellationOptions, FDatasmithMeshElementPayload& OutMeshPayload);
 	CORETECHSURFACE_API bool LoadFile(const FString& FileName, const CADLibrary::FImportParameters& InImportParameters, const CADLibrary::FMeshParameters& InMeshParameters, FMeshDescription& OutMeshDescription);
 	CORETECHSURFACE_API bool Tessellate(uint64 MainObjectId, const CADLibrary::FImportParameters& InImportParameters, const CADLibrary::FMeshParameters& InMeshParameters, FMeshDescription& OutMesh);
 }
\ No newline at end of file
diff --git a/Engine/Plugins/Enterprise/DatasmithCADImporter/Source/DatasmithCADTranslator/Private/DatasmithCADTranslator.cpp b/Engine/Plugins/Enterprise/DatasmithCADImporter/Source/DatasmithCADTranslator/Private/DatasmithCADTranslator.cpp
--- a/Engine/Plugins/Enterprise/DatasmithCADImporter/Source/DatasmithCADTranslator/Private/DatasmithCADTranslator.cpp
+++ b/Engine/Plugins/Enterprise/DatasmithCADImporter/Source/DatasmithCADTranslator/Private/DatasmithCADTranslator.cpp
@@ -209,7 +209,7 @@ bool FDatasmithCADTranslator::LoadStaticMesh(const TSharedRef<IDatasmithMeshElem
 
 		if (CADLibrary::FImportParameters::bGDisableCADKernelTessellation)
 		{
-			AddSurfaceData(MeshElement->GetFile(), ImportParameters, MeshParameters, OutMeshPayload);
+			ParametricSurfaceUtils::AddSurfaceData(MeshElement->GetFile(), ImportParameters, MeshParameters, GetCommonTessellationOptions(), OutMeshPayload);
 		}
 		else
 		{
diff --git a/Engine/Plugins/Enterprise/DatasmithCADImporter/Source/DatasmithOpenNurbsTranslator/Private/DatasmithOpenNurbsTranslator.cpp b/Engine/Plugins/Enterprise/DatasmithCADImporter/Source/DatasmithOpenNurbsTranslator/Private/DatasmithOpenNurbsTranslator.cpp
--- a/Engine/Plugins/Enterprise/DatasmithCADImporter/Source/DatasmithOpenNurbsTranslator/Private/DatasmithOpenNurbsTranslator.cpp
+++ b/Engine/Plugins/Enterprise/DatasmithCADImporter/Source/DatasmithOpenNurbsTranslator/Private/DatasmithOpenNurbsTranslator.cpp
@@ -21,6 +21,7 @@
 #include "OpenNurbsBRepConverter.h"
 #include "OpenNurbsBRepToCADKernelConverter.h"
 #include "OpenNurbsBRepToCoretechConverter.h" // requires CoreTech as public dependency
+#include "OpenNurbsBRepToTechSoftConverter.h" // requires Techsoft as public dependency
 #include "Utility/DatasmithMeshHelper.h"
 
 #if WITH_EDITOR
@@ -2970,7 +2971,7 @@ bool FOpenNurbsTranslatorImpl::TranslateBRep(ON_Brep* Brep, const ON_3dmObjectAt
 		OpenNurbsBRepConverter->AddBRep(*Brep, Offset);
 		CADModelConverter->RepairTopology();
 
-		CADModelConverter->SaveBRep(*OutputPath, MeshElement);
+		CADModelConverter->SaveModel(*OutputPath, MeshElement);
 
 		CADLibrary::FMeshParameters MeshParameters;
 		return CADModelConverter->Tessellate(MeshParameters, OutMesh);
diff --git a/Engine/Plugins/Enterprise/DatasmithCADImporter/Source/DatasmithOpenNurbsTranslator/Private/OpenNurbsBRepToCADKernelConverter.cpp b/Engine/Plugins/Enterprise/DatasmithCADImporter/Source/DatasmithOpenNurbsTranslator/Private/OpenNurbsBRepToCADKernelConverter.cpp
--- a/Engine/Plugins/Enterprise/DatasmithCADImporter/Source/DatasmithOpenNurbsTranslator/Private/OpenNurbsBRepToCADKernelConverter.cpp
+++ b/Engine/Plugins/Enterprise/DatasmithCADImporter/Source/DatasmithOpenNurbsTranslator/Private/OpenNurbsBRepToCADKernelConverter.cpp
@@ -11,8 +11,6 @@
 #include "opennurbs.h"
 #pragma warning(pop)
 
-#include <vector>
-
 #include "CADKernelTools.h"
 
 #include "CADKernel/Core/Session.h"
@@ -198,7 +196,7 @@ TSharedPtr<CADKernel::FTopologicalEdge> FOpenNurbsBRepToCADKernelConverter::AddE
 	}
 
 	ON_NurbsCurve OpenNurbsCurve;
-	int NurbFormSuccess = OpenNurbsTrim.GetNurbForm(OpenNurbsCurve); // 0:Nok 1:Ok 2:OkBut
+	int32 NurbFormSuccess = OpenNurbsTrim.GetNurbForm(OpenNurbsCurve); // 0:Nok 1:Ok 2:OkBut
 	if (NurbFormSuccess == 0)
 	{
 		return TSharedPtr<FTopologicalEdge>();
@@ -214,7 +212,7 @@ TSharedPtr<CADKernel::FTopologicalEdge> FOpenNurbsBRepToCADKernelConverter::AddE
 
 	NurbsCurveData.NodalVector.Reserve(KnotCount + 2);
 	NurbsCurveData.NodalVector.Emplace(OpenNurbsCurve.SuperfluousKnot(0));
-	for (int Index = 0; Index < KnotCount; ++Index)
+	for (int32 Index = 0; Index < KnotCount; ++Index)
 	{
 		NurbsCurveData.NodalVector.Emplace(OpenNurbsCurve.Knot(Index));
 	}
@@ -284,8 +282,8 @@ TSharedPtr<CADKernel::FTopologicalFace> FOpenNurbsBRepToCADKernelConverter::AddF
 		return Face;
 	}
 
-	int LoopCount = OpenNurbsFace.LoopCount();
-	for (int LoopIndex = 0; LoopIndex < LoopCount; ++LoopIndex)
+	int32 LoopCount = OpenNurbsFace.LoopCount();
+	for (int32 LoopIndex = 0; LoopIndex < LoopCount; ++LoopIndex)
 	{
 		const ON_BrepLoop& OpenNurbsLoop = *OpenNurbsFace.Loop(LoopIndex);
 		TSharedPtr<FTopologicalLoop> Loop = AddLoop(OpenNurbsLoop, Surface);
@@ -314,8 +312,8 @@ bool FOpenNurbsBRepToCADKernelConverter::AddBRep(ON_Brep& BRep, const ON_3dVecto
 	BRep.FlipReversedSurfaces();
 
 	// Create faces
-	int FaceCount = BRep.m_F.Count();
-	for (int index = 0; index < FaceCount; index++)
+	int32 FaceCount = BRep.m_F.Count();
+	for (int32 index = 0; index < FaceCount; index++)
 	{
 		const ON_BrepFace& OpenNurbsFace = BRep.m_F[index];
 		TSharedPtr<FTopologicalFace> Face = AddFace(OpenNurbsFace);
diff --git a/Engine/Plugins/Enterprise/DatasmithCADImporter/Source/DatasmithOpenNurbsTranslator/Private/OpenNurbsBRepToCoretechConverter.cpp b/Engine/Plugins/Enterprise/DatasmithCADImporter/Source/DatasmithOpenNurbsTranslator/Private/OpenNurbsBRepToCoretechConverter.cpp
--- a/Engine/Plugins/Enterprise/DatasmithCADImporter/Source/DatasmithOpenNurbsTranslator/Private/OpenNurbsBRepToCoretechConverter.cpp
+++ b/Engine/Plugins/Enterprise/DatasmithCADImporter/Source/DatasmithOpenNurbsTranslator/Private/OpenNurbsBRepToCoretechConverter.cpp
@@ -11,8 +11,6 @@
 #include "opennurbs.h"
 #pragma warning(pop)
 
-#include <vector>
-
 // As of 2018 SP2, CoreTech tessellation for a face is broken when the outer loop is the whole uv range, and some edges are singularities.
 // In those cases, inner loops behave like outers, outer is ignored.
 // current fix will split the uv plane in two parts.
@@ -153,7 +151,7 @@ struct SurfaceInfo
 		if (bUOrderIncreased || bVOrderIncreased)
 		{
 			BuildHull();
-	}
+		}
 	}
 
 	void FixUnsupportedParameters()
@@ -285,7 +283,7 @@ void FBRepToKernelIOBodyTranslator::CreateCTFace_internal(const ON_BrepFace& Fac
 			{
 				for (int32 j = 0; j < CTCurve.ControlPoints.Num(); j += CTCurve.ControlPointSize)
 				{
-					CTCurve.ControlPoints[j] = min(max(CTCurve.ControlPoints[j], outerBBox.m_min.x), outerBBox.m_max.x);
+					CTCurve.ControlPoints[j] = FMath::Min(FMath::Max(CTCurve.ControlPoints[j], outerBBox.m_min.x), outerBBox.m_max.x);
 				}
 			}
 
diff --git a/Engine/Plugins/Enterprise/DatasmithCADImporter/Source/DatasmithOpenNurbsTranslator/Private/OpenNurbsBRepToTechSoftConverter.cpp b/Engine/Plugins/Enterprise/DatasmithCADImporter/Source/DatasmithOpenNurbsTranslator/Private/OpenNurbsBRepToTechSoftConverter.cpp
new file mode 100644
--- /dev/null
+++ b/Engine/Plugins/Enterprise/DatasmithCADImporter/Source/DatasmithOpenNurbsTranslator/Private/OpenNurbsBRepToTechSoftConverter.cpp
@@ -0,0 +1,451 @@
+// Copyright Epic Games, Inc. All Rights Reserved.
+
+#include "OpenNurbsBRepToTechSoftConverter.h"
+
+#ifdef USE_OPENNURBS
+
+#pragma warning(push)
+#pragma warning(disable:4265)
+#pragma warning(disable:4005) // TEXT macro redefinition
+#include "opennurbs.h"
+#pragma warning(pop)
+
+#include "TechSoftInterface.h"
+#include "TechSoftUtils.h"
+#include "TUniqueTechSoftObj.h"
+
+#ifdef USE_TECHSOFT_SDK
+namespace OpenNurbsUtils
+{
+enum EAxis { U, V };
+
+template<typename PointType>
+void SetA3DVector(const PointType& OpenNurbsPoint, A3DVector3dData& OutTSPoint)
+{
+	OutTSPoint.m_dX = OpenNurbsPoint.x;
+	OutTSPoint.m_dY = OpenNurbsPoint.y;
+	OutTSPoint.m_dZ = OpenNurbsPoint.z;
+};
+
+template<typename PointType>
+void SetA3DVectorAndWeight(const PointType& OpenNurbsPoint, A3DVector3dData& OutTSPoint, A3DDouble& OutWeight)
+{
+	SetA3DVector(OpenNurbsPoint, OutTSPoint);
+	OutWeight = OpenNurbsPoint.w;
+};
+
+
+struct FOpenNurbsSurfaceInfo
+{
+	FOpenNurbsSurfaceInfo(const ON_NurbsSurface& InOpenNurbsSurface, A3DSurfNurbsData& InData)
+		: OpenNurbsSurface(InOpenNurbsSurface)
+		, Data(InData)
+	{
+	}
+
+	const ON_NurbsSurface& OpenNurbsSurface;
+	A3DSurfNurbsData& Data;
+
+	TArray<A3DDouble> UNodalVector;
+	TArray<A3DDouble> VNodalVector;
+
+	TArray<A3DDouble> Weights;
+	TArray<A3DVector3dData> ControlPoints;
+
+	void FillPerAxisInfo(EAxis Axis)
+	{
+		A3DUns32& Degree = Axis == EAxis::U ? Data.m_uiUDegree : Data.m_uiVDegree;
+		Degree = OpenNurbsSurface.Order(Axis) - 1;
+
+		A3DUns32& ControlPointCount = Axis == EAxis::U ? Data.m_uiUCtrlSize : Data.m_uiVCtrlSize;
+		ControlPointCount = OpenNurbsSurface.CVCount(Axis);
+
+		A3DUns32& KnotSize = Axis == EAxis::U ? Data.m_uiUKnotSize : Data.m_uiVKnotSize;
+		KnotSize = Degree + ControlPointCount + 1;
+
+		TArray<A3DDouble>& NodalVector = Axis == EAxis::U ? UNodalVector : VNodalVector;
+		NodalVector.Reserve(KnotSize);
+		NodalVector.Add(OpenNurbsSurface.SuperfluousKnot(Axis, 0));
+		uint32 KnotCount = OpenNurbsSurface.KnotCount(Axis);
+		for (uint32 i = 0; i < KnotCount; ++i)
+		{
+			NodalVector.Add(OpenNurbsSurface.Knot(Axis, i));
+		}
+		NodalVector.Add(OpenNurbsSurface.SuperfluousKnot(Axis, 1));
+
+		A3DDouble*& UKnots = Axis == EAxis::U ? Data.m_pdUKnots : Data.m_pdVKnots;
+		UKnots = NodalVector.GetData();
+	}
+
+	void BuildControlPoints()
+	{
+		bool bIsRational = OpenNurbsSurface.IsRational();
+
+		ControlPoints.SetNum(Data.m_uiUCtrlSize * Data.m_uiVCtrlSize);
+		A3DVector3dData* ControlPointPtr = ControlPoints.GetData();
+		Data.m_pCtrlPts = ControlPointPtr;
+
+		if(bIsRational)
+		{
+			Weights.SetNum(Data.m_uiUCtrlSize * Data.m_uiVCtrlSize);
+			A3DDouble* WeightPtr = Weights.GetData();
+			Data.m_pdWeights = WeightPtr;
+
+			ON_4dPoint OpenNurbsPoint;
+			for (A3DUns32 UIndex = 0; UIndex < Data.m_uiUCtrlSize; ++UIndex)
+			{
+				for (A3DUns32 VIndex = 0; VIndex < Data.m_uiVCtrlSize; ++VIndex, ++ControlPointPtr, ++WeightPtr)
+				{
+					OpenNurbsSurface.GetCV(UIndex, VIndex, ON::point_style::euclidean_rational, OpenNurbsPoint);
+					SetA3DVectorAndWeight(OpenNurbsPoint, *ControlPointPtr, *WeightPtr);
+				}
+			}
+		}
+		else
+		{
+			ON_3dPoint OpenNurbsPoint;
+			for (A3DUns32 UIndex = 0; UIndex < Data.m_uiUCtrlSize; ++UIndex)
+			{
+				for (A3DUns32 VIndex = 0; VIndex < Data.m_uiVCtrlSize; ++VIndex, ++ControlPointPtr)
+				{
+					OpenNurbsSurface.GetCV(UIndex, VIndex, ON::point_style::not_rational, OpenNurbsPoint);
+					SetA3DVector(OpenNurbsPoint, *ControlPointPtr);
+				}
+			}
+		}
+	}
+
+	void Populate()
+	{
+		Data.m_eKnotType = A3DEKnotType::kA3DKnotTypeUnspecified;
+		Data.m_eSurfaceForm = A3DEBSplineSurfaceForm::kA3DBSplineSurfaceFormUnspecified;
+
+		FillPerAxisInfo(U);
+		FillPerAxisInfo(V);
+
+		BuildControlPoints();
+	}
+};
+
+} // NS
+
+A3DSurfBase* FOpenNurbsBRepToTechSoftConverter::CreateSurface(const ON_NurbsSurface& OpenNurbsSurface)
+{
+	CADLibrary::TUniqueTSObj<A3DSurfNurbsData> NurbsSurfaceData;
+
+	OpenNurbsUtils::FOpenNurbsSurfaceInfo NurbsInfo(OpenNurbsSurface, *NurbsSurfaceData);
+	NurbsInfo.Populate();
+
+	return CADLibrary::TechSoftInterface::CreateSurfaceNurbs(*NurbsSurfaceData);
+}
+
+A3DCrvBase* FOpenNurbsBRepToTechSoftConverter::CreateCurve(const ON_NurbsCurve& OpenNurbsCurve)
+{
+	CADLibrary::TUniqueTSObj<A3DCrvNurbsData> NurbsCurveData;
+
+	NurbsCurveData->m_bIs2D = true;
+	NurbsCurveData->m_bRational = OpenNurbsCurve.IsRational();
+	NurbsCurveData->m_uiDegree = OpenNurbsCurve.Order() - 1;
+
+
+	NurbsCurveData->m_eKnotType = kA3DKnotTypeUnspecified;
+	NurbsCurveData->m_eCurveForm = kA3DBSplineCurveFormUnspecified;
+
+	int32 KnotCount = OpenNurbsCurve.KnotCount();
+	TArray<A3DDouble> NodalVector;
+	NodalVector.Reserve(KnotCount + 2);
+	NodalVector.Emplace(OpenNurbsCurve.SuperfluousKnot(0));
+	for (int Index = 0; Index < KnotCount; ++Index)
+	{
+		NodalVector.Emplace(OpenNurbsCurve.Knot(Index));
+	}
+	NodalVector.Emplace(OpenNurbsCurve.SuperfluousKnot(1));
+
+	TArray<A3DVector3dData> ControlPointArray;
+	TArray<A3DDouble> WeightArray;
+
+	int32 ControlPointCount = OpenNurbsCurve.CVCount();
+
+	ControlPointArray.SetNumUninitialized(ControlPointCount);
+	A3DVector3dData* ControlPoints = ControlPointArray.GetData();
+	ON_3dPoint OpenNurbsPoint;
+	if (OpenNurbsCurve.IsRational())
+	{
+		WeightArray.SetNumUninitialized(ControlPointCount);
+		A3DDouble* Weights = WeightArray.GetData();
+		for (int32 Index = 0; Index < ControlPointCount; ++Index, ++ControlPoints, ++Weights)
+		{
+			OpenNurbsCurve.GetCV(Index, ON::point_style::euclidean_rational, OpenNurbsPoint);
+			ControlPoints->m_dX = OpenNurbsPoint.x;
+			ControlPoints->m_dY = OpenNurbsPoint.y;
+			ControlPoints->m_dZ = 0;
+			*Weights = OpenNurbsPoint.z;
+		}
+	}
+	else
+	{
+		for (int32 Index = 0; Index < ControlPointCount; ++Index, ++ControlPoints)
+		{
+			OpenNurbsCurve.GetCV(Index, ON::point_style::not_rational, OpenNurbsPoint);
+			ControlPoints->m_dX = OpenNurbsPoint.x;
+			ControlPoints->m_dY = OpenNurbsPoint.y;
+			ControlPoints->m_dZ = 0;
+		}
+	}
+
+	NurbsCurveData->m_pCtrlPts = ControlPointArray.GetData();
+	NurbsCurveData->m_uiCtrlSize = ControlPointArray.Num();
+
+	NurbsCurveData->m_pdWeights = NurbsCurveData->m_bRational ? WeightArray.GetData() : nullptr;
+	NurbsCurveData->m_uiWeightSize = NurbsCurveData->m_bRational ? WeightArray.Num() : 0;
+
+	NurbsCurveData->m_pdKnots = NodalVector.GetData();
+	NurbsCurveData->m_uiKnotSize = NodalVector.Num();
+	
+	return CADLibrary::TechSoftInterface::CreateCurveNurbs(*NurbsCurveData);
+}
+
+A3DTopoLoop* FOpenNurbsBRepToTechSoftConverter::CreateTopoLoop(const ON_BrepLoop& OpenNurbsLoop)
+{
+	if (!OpenNurbsLoop.IsValid())
+	{
+		return nullptr;
+	}
+
+	int32 EdgeCount = OpenNurbsLoop.TrimCount();
+	TArray<A3DTopoCoEdge*> CoEdges;
+	CoEdges.Reserve(EdgeCount);
+
+	for (int32 Index = 0; Index < EdgeCount; ++Index)
+	{
+		ON_BrepTrim& Trim = *OpenNurbsLoop.Trim(Index);
+
+		A3DTopoCoEdge* CoEdge = CreateTopoCoEdge(Trim);
+		if (CoEdge != nullptr)
+		{
+			CoEdges.Add(CoEdge);
+		}
+	}
+
+	if (CoEdges.Num() == 0)
+	{
+		return nullptr;
+	}
+
+	CADLibrary::TUniqueTSObj<A3DTopoLoopData> LoopData;
+
+	LoopData->m_ppCoEdges = CoEdges.GetData();
+	LoopData->m_uiCoEdgeSize = CoEdges.Num();
+	LoopData->m_ucOrientationWithSurface = 1;
+
+	return CADLibrary::TechSoftInterface::CreateTopoLoop(*LoopData);
+}
+
+void FOpenNurbsBRepToTechSoftConverter::LinkEdgesLoop(const ON_BrepLoop& OpenNurbsLoop)
+{
+	int32 EdgeCount = OpenNurbsLoop.TrimCount();
+	for (int32 Index = 0; Index < EdgeCount; ++Index)
+	{
+		ON_BrepTrim& OpenNurbsTrim = *OpenNurbsLoop.Trim(Index);
+		ON_BrepEdge* OpenNurbsEdge = OpenNurbsTrim.Edge();
+		if (OpenNurbsEdge == nullptr)
+		{
+			continue;
+		}
+
+		A3DTopoCoEdge** CoEdge = OpenNurbsTrimId2TechSoftCoEdge.Find(OpenNurbsTrim.m_trim_index);
+		if (CoEdge == nullptr)
+		{
+			continue;
+		}
+
+		for (int32 Endex = 0; Endex < OpenNurbsEdge->m_ti.Count(); ++Endex)
+		{
+			int32 LinkedEdgeId = OpenNurbsEdge->m_ti[Endex];
+			if (LinkedEdgeId == OpenNurbsTrim.m_trim_index)
+			{
+				continue;
+			}
+
+			A3DTopoCoEdge** TwinCoEdge = OpenNurbsTrimId2TechSoftCoEdge.Find(LinkedEdgeId);
+			if (TwinCoEdge != nullptr)
+			{
+				CADLibrary::TechSoftInterface::LinkCoEdges(*CoEdge, *TwinCoEdge);
+				break;
+			}
+		}
+	}
+}
+
+A3DTopoCoEdge* FOpenNurbsBRepToTechSoftConverter::CreateTopoCoEdge(ON_BrepTrim& OpenNurbsTrim)
+{
+	ON_NurbsCurve OpenNurbsCurve;
+	int NurbFormSuccess = OpenNurbsTrim.GetNurbForm(OpenNurbsCurve); // 0:Nok 1:Ok 2:OkBut
+	if (NurbFormSuccess == 0)
+	{
+		return nullptr;
+	}
+
+	A3DCrvBase* NurbsCurvePtr = CreateCurve(OpenNurbsCurve);
+	if (NurbsCurvePtr == nullptr)
+	{
+		return nullptr;
+	}
+
+	ON_Interval dom = OpenNurbsCurve.Domain();
+	if (!FMath::IsNearlyEqual(OpenNurbsCurve.SuperfluousKnot(0), dom.m_t[0]) ||
+		!FMath::IsNearlyEqual(OpenNurbsCurve.SuperfluousKnot(1), dom.m_t[1]))
+	{
+		NurbsCurvePtr = CADLibrary::TechSoftUtils::CreateTrimNurbsCurve(NurbsCurvePtr, dom.m_t[0], dom.m_t[1], true);
+		if (NurbsCurvePtr == nullptr)
+		{
+			return nullptr;
+		}
+	}
+
+	A3DTopoEdge* EdgePtr = CADLibrary::TechSoftUtils::CreateTopoEdge();
+	if (EdgePtr == nullptr)
+	{
+		return nullptr;
+	}
+
+	CADLibrary::TUniqueTSObj<A3DTopoCoEdgeData> CoEdgeData;
+	CoEdgeData->m_pUVCurve = NurbsCurvePtr;
+	CoEdgeData->m_pEdge = EdgePtr;
+	CoEdgeData->m_ucOrientationWithLoop = OpenNurbsTrim.m_bRev3d ? 0 : 1;
+	CoEdgeData->m_ucOrientationUVWithLoop = 1;
+
+	A3DTopoCoEdge* CoEdgePtr = CADLibrary::TechSoftInterface::CreateTopoCoEdge(*CoEdgeData);
+
+	// Only Edge with twin need to be in the map used in LinkEdgesLoop 
+	ON_BrepEdge* OpenNurbsEdge = OpenNurbsTrim.Edge();
+	if (CoEdgePtr != nullptr && OpenNurbsEdge != nullptr && OpenNurbsEdge->m_ti.Count() > 1)
+	{
+		OpenNurbsTrimId2TechSoftCoEdge.Add(OpenNurbsTrim.m_trim_index, CoEdgePtr);
+	}
+	return CoEdgePtr;
+}
+
+A3DTopoFace* FOpenNurbsBRepToTechSoftConverter::CreateTopoFace(const ON_BrepFace& OpenNurbsFace)
+{
+	ON_NurbsSurface OpenNurbsSurface;
+	OpenNurbsFace.NurbsSurface(&OpenNurbsSurface);
+
+	A3DSurfBase* CarrierSurface = CreateSurface(OpenNurbsSurface);
+	if (CarrierSurface == nullptr)
+	{
+		return nullptr;
+	}
+
+	const ON_BrepLoop* OuterLoop = OpenNurbsFace.OuterLoop();
+	if (OuterLoop == nullptr)
+	{
+		return CADLibrary::TechSoftUtils::CreateTopoFaceWithNaturalLoop(CarrierSurface);
+	}
+
+	int32 LoopCount = OpenNurbsFace.LoopCount();
+
+	TArray<A3DTopoLoop*> Loops;
+	Loops.Reserve(LoopCount);
+
+	int32 OuterLoopIndex = 0;
+	for (int32 LoopIndex = 0; LoopIndex < LoopCount; ++LoopIndex)
+	{
+		const ON_BrepLoop& OpenNurbsLoop = *OpenNurbsFace.Loop(LoopIndex);
+		ON_BrepLoop::TYPE LoopType = OpenNurbsLoop.m_type;
+		if (LoopType == ON_BrepLoop::TYPE::outer)
+		{
+			OuterLoopIndex = Loops.Num();
+		}
+
+		A3DTopoLoop* Loop = CreateTopoLoop(OpenNurbsLoop);
+		if (Loop != nullptr)
+		{
+			Loops.Add(Loop);
+			LinkEdgesLoop(OpenNurbsLoop);
+		}
+	}
+
+	if (Loops.Num() == 0)
+	{
+		return nullptr;
+	}
+
+	CADLibrary::TUniqueTSObj<A3DTopoFaceData> Face;
+	Face->m_pSurface = CarrierSurface;
+	Face->m_bHasTrimDomain = false;
+	Face->m_ppLoops = Loops.GetData();
+	Face->m_uiLoopSize = Loops.Num();
+	Face->m_uiOuterLoopIndex = OuterLoopIndex;
+	Face->m_dTolerance = 0.01; //mm
+
+	A3DTopoFace* FacePtr = CADLibrary::TechSoftInterface::CreateTopoFace(*Face);
+
+	CADLibrary::TechSoftUtils::SetEntityGraphicsColor(FacePtr, FColor(200, 200, 200));
+
+	return FacePtr;
+}
+#endif
+
+bool FOpenNurbsBRepToTechSoftConverter::AddBRep(ON_Brep& BRep, const ON_3dVector& Offset)
+{
+#ifdef USE_TECHSOFT_SDK
+	OpenNurbsTrimId2TechSoftCoEdge.Empty();
+
+	BRep.Translate(Offset);
+	BRep.FlipReversedSurfaces();
+
+	// Create faces
+	int32 FaceCount = BRep.m_F.Count();
+	TArray<A3DTopoFace*> TSFaces;
+	TArray<A3DUns8> FaceOrientations;
+
+	TSFaces.Reserve(FaceCount);
+	FaceOrientations.Reserve(FaceCount);
+
+	for (int32 index = 0; index < FaceCount; index++)
+	{
+		const ON_BrepFace& OpenNurbsFace = BRep.m_F[index];
+		A3DTopoFace* TSFace = CreateTopoFace(OpenNurbsFace);
+		if (TSFace != nullptr)
+		{
+			TSFaces.Add(TSFace);
+			FaceOrientations.Add(1);
+		}
+	}
+
+	BRep.Translate(-Offset);
+
+	if (TSFaces.IsEmpty())
+	{
+		return false;
+	}
+
+	A3DTopoShell* TopoShellPtr = nullptr;
+	{
+		CADLibrary::TUniqueTSObj<A3DTopoShellData> TopoShellData;
+		TopoShellData->m_bClosed = false;
+		TopoShellData->m_ppFaces = TSFaces.GetData();
+		TopoShellData->m_uiFaceSize = TSFaces.Num();
+		TopoShellData->m_pucOrientationWithShell = FaceOrientations.GetData();
+
+		TopoShellPtr = CADLibrary::TechSoftInterface::CreateTopoShell(*TopoShellData);
+
+		if (TopoShellPtr == nullptr)
+		{
+			return false;
+		}
+	}
+
+	A3DRiRepresentationItem* RiRepresentationItem = CADLibrary::TechSoftUtils::CreateRIBRep(TopoShellPtr);
+
+	if (RiRepresentationItem != nullptr)
+	{
+		RiRepresentationItems.Add(RiRepresentationItem);
+		return true;
+	}
+#endif
+	return false;
+}
+
+#endif // defined(USE_OPENNURBS)
diff --git a/Engine/Plugins/Enterprise/DatasmithCADImporter/Source/DatasmithOpenNurbsTranslator/Private/OpenNurbsBRepToTechSoftConverter.h b/Engine/Plugins/Enterprise/DatasmithCADImporter/Source/DatasmithOpenNurbsTranslator/Private/OpenNurbsBRepToTechSoftConverter.h
new file mode 100644
--- /dev/null
+++ b/Engine/Plugins/Enterprise/DatasmithCADImporter/Source/DatasmithOpenNurbsTranslator/Private/OpenNurbsBRepToTechSoftConverter.h
@@ -0,0 +1,66 @@
+// Copyright Epic Games, Inc. All Rights Reserved.
+
+#pragma once
+
+#include "CoreMinimal.h"
+
+#ifdef USE_OPENNURBS
+
+#include "OpenNurbsBRepConverter.h"
+#include "CADModelToTechSoftConverterBase.h"
+#include "CTSession.h"
+
+#ifdef USE_TECHSOFT_SDK
+#include "TechSoftInterface.h"
+#endif
+
+
+class ON_BoundingBox;
+class ON_Brep;
+class ON_BrepFace;
+class ON_BrepLoop;
+class ON_BrepTrim;
+class ON_NurbsCurve;
+class ON_NurbsSurface;
+class ON_3dVector;
+
+class FOpenNurbsBRepToTechSoftConverter : public FCADModelToTechSoftConverterBase, public IOpenNurbsBRepConverter
+{
+public:
+
+	/**
+	 * Make sure CT is initialized, and a main object is ready.
+	 * Handle input file unit and an output unit
+	 * @param InOwner
+	 */
+	explicit FOpenNurbsBRepToTechSoftConverter(const CADLibrary::FImportParameters& ImportParameters)
+		: FCADModelToTechSoftConverterBase(ImportParameters)
+	{
+	}
+
+	virtual ~FOpenNurbsBRepToTechSoftConverter() = default;
+
+	/**
+	 * Set BRep to tessellate, offsetting it prior to tessellation(used to set mesh pivot at the center of the surface bounding box)
+	 *
+	 * @param  Brep	a BRep to tessellate
+	 * @param  Offset translate Brep by this value before tessellating 
+	 */
+	virtual bool AddBRep(ON_Brep& Brep, const ON_3dVector& Offset) override;
+
+private:
+#ifdef USE_TECHSOFT_SDK
+
+	TMap<int32, A3DTopoCoEdge*> OpenNurbsTrimId2TechSoftCoEdge;
+
+	A3DTopoCoEdge* CreateTopoCoEdge(ON_BrepTrim& Trim);
+	A3DTopoFace* CreateTopoFace(const ON_BrepFace& OpenNurbsFace);
+	A3DTopoLoop* CreateTopoLoop(const ON_BrepLoop& OpenNurbsLoop);
+	A3DSurfBase* CreateSurface(const ON_NurbsSurface& OpenNurbsSurface);
+	A3DCrvBase* CreateCurve(const ON_NurbsCurve& OpenNurbsCurve);
+
+	void LinkEdgesLoop(const ON_BrepLoop& OpenNurbsLoop);
+#endif
+};
+
+#endif // USE_OPENNURBS
\ No newline at end of file
diff --git a/Engine/Plugins/Enterprise/DatasmithCADImporter/Source/DatasmithPLMXMLTranslator/Private/DatasmithPlmXmlImporter.cpp b/Engine/Plugins/Enterprise/DatasmithCADImporter/Source/DatasmithPLMXMLTranslator/Private/DatasmithPlmXmlImporter.cpp
--- a/Engine/Plugins/Enterprise/DatasmithCADImporter/Source/DatasmithPLMXMLTranslator/Private/DatasmithPlmXmlImporter.cpp
+++ b/Engine/Plugins/Enterprise/DatasmithCADImporter/Source/DatasmithPLMXMLTranslator/Private/DatasmithPlmXmlImporter.cpp
@@ -27,6 +27,8 @@
 #include "DatasmithMeshBuilder.h"
 #include "DatasmithSceneGraphBuilder.h"
 #include "DatasmithUtils.h"
+#include "ParametricSurfaceTranslator.h"
+
 
 
 DEFINE_LOG_CATEGORY_STATIC(LogDatasmithPlmXmlImport, Log, All);
@@ -299,7 +301,7 @@ namespace PlmXml
 
 				if (CADLibrary::FImportParameters::bGDisableCADKernelTessellation)
 				{
-					CoreTechSurface::AddSurfaceDataForMesh(MeshElement->GetFile(), ImportParameters, MeshParameters, TessellationOptions, OutMeshPayload);
+					ParametricSurfaceUtils::AddSurfaceData(MeshElement->GetFile(), ImportParameters, MeshParameters, TessellationOptions, OutMeshPayload);
 				}
 				else
 				{
diff --git a/Engine/Plugins/Enterprise/DatasmithCADImporter/Source/DatasmithWireTranslator/Private/AliasModelToCADKernelConverter.cpp b/Engine/Plugins/Enterprise/DatasmithCADImporter/Source/DatasmithWireTranslator/Private/AliasModelToCADKernelConverter.cpp
--- a/Engine/Plugins/Enterprise/DatasmithCADImporter/Source/DatasmithWireTranslator/Private/AliasModelToCADKernelConverter.cpp
+++ b/Engine/Plugins/Enterprise/DatasmithCADImporter/Source/DatasmithWireTranslator/Private/AliasModelToCADKernelConverter.cpp
@@ -98,7 +98,7 @@ TSharedPtr<CADKernel::FTopologicalEdge> FAliasModelToCADKernelConverter::AddEdge
 	CADKernel::FNurbsCurveData NurbsCurveData;
 
 	NurbsCurveData.Degree = AliasTrimCurve.degree();
-	int ControlPointCount = AliasTrimCurve.numberOfCVs();
+	int32 ControlPointCount = AliasTrimCurve.numberOfCVs();
 
 	NurbsCurveData.Dimension = 2;
 	NurbsCurveData.bIsRational = true;
diff --git a/Engine/Plugins/Enterprise/DatasmithCADImporter/Source/DatasmithWireTranslator/Private/AliasModelToTechSoftConverter.cpp b/Engine/Plugins/Enterprise/DatasmithCADImporter/Source/DatasmithWireTranslator/Private/AliasModelToTechSoftConverter.cpp
new file mode 100644
--- /dev/null
+++ b/Engine/Plugins/Enterprise/DatasmithCADImporter/Source/DatasmithWireTranslator/Private/AliasModelToTechSoftConverter.cpp
@@ -0,0 +1,464 @@
+// Copyright Epic Games, Inc. All Rights Reserved.
+
+#include "AliasModelToTechSoftConverter.h"
+
+#include "Hal/PlatformMemory.h"
+#include "Math/Color.h"
+#include "MeshDescriptionHelper.h"
+#include "TechSoftInterface.h"
+#include "TechSoftUtils.h"
+#include "TUniqueTechSoftObj.h"
+
+
+#ifdef USE_OPENMODEL
+
+// Alias API wrappes object in AlObjects. This is an abstract base class which holds a reference to an anonymous data structure.
+// The only way to compare two AlObjects is to compare their data structure. That is the reason why private fields are made public. 
+#define private public
+#include "AlCurve.h"
+#include "AlDagNode.h"
+#include "AlShell.h"
+#include "AlShellNode.h"
+#include "AlSurface.h"
+#include "AlSurfaceNode.h"
+#include "AlTrimBoundary.h"
+#include "AlTrimCurve.h"
+#include "AlTrimRegion.h"
+#include "AlTM.h"
+
+namespace UE_DATASMITHWIRETRANSLATOR_NAMESPACE
+{
+#ifdef USE_TECHSOFT_SDK
+
+namespace AliasToTechSoftUtils
+{
+
+enum EAxis { U, V };
+
+template<typename Surface_T>
+A3DSurfBase* AddNURBSSurface(const Surface_T& AliasSurface, EAliasObjectReference InObjectReference, const AlMatrix4x4& InAlMatrix)
+{
+	CADLibrary::TUniqueTSObj<A3DSurfNurbsData> NurbsSurfaceData;
+
+	NurbsSurfaceData->m_eKnotType = A3DEKnotType::kA3DKnotTypeUnspecified;
+	NurbsSurfaceData->m_eSurfaceForm = A3DEBSplineSurfaceForm::kA3DBSplineSurfaceFormUnspecified;
+
+	NurbsSurfaceData->m_uiUDegree = AliasSurface.uDegree();
+	NurbsSurfaceData->m_uiVDegree = AliasSurface.vDegree();
+
+	NurbsSurfaceData->m_uiUCtrlSize = AliasSurface.uNumberOfCVsInclMultiples();
+	NurbsSurfaceData->m_uiVCtrlSize = AliasSurface.vNumberOfCVsInclMultiples();
+	NurbsSurfaceData->m_uiUKnotSize = AliasSurface.realuNumberOfKnots() + 2;
+	NurbsSurfaceData->m_uiVKnotSize = AliasSurface.realvNumberOfKnots() + 2;
+
+	TArray<A3DDouble> UNodalVector;
+	TArray<A3DDouble> VNodalVector;
+
+	TFunction<void(EAxis, TArray<A3DDouble>&)> FillNodalVector = [&](EAxis Axis, TArray<A3DDouble>& OutNodalVector)
+	{
+		TArray<double> NodalVector;
+		if(Axis == EAxis::U)
+		{
+			OutNodalVector.Reserve(NurbsSurfaceData->m_uiUKnotSize);
+			NodalVector.SetNumUninitialized(NurbsSurfaceData->m_uiUKnotSize - 2);
+			AliasSurface.realuKnotVector(NodalVector.GetData());
+		}
+		else
+		{
+			OutNodalVector.Reserve(NurbsSurfaceData->m_uiVKnotSize);
+			NodalVector.SetNumUninitialized(NurbsSurfaceData->m_uiVKnotSize - 2);
+			AliasSurface.realvKnotVector(NodalVector.GetData());
+		}
+
+		OutNodalVector.Add(NodalVector[0]);
+		for(double Value : NodalVector)
+		{
+			OutNodalVector.Add(Value);
+		}
+		OutNodalVector.Add(NodalVector.Last());
+
+		if (Axis == EAxis::U)
+		{
+			NurbsSurfaceData->m_pdUKnots = OutNodalVector.GetData();
+		}
+		else
+		{
+			NurbsSurfaceData->m_pdVKnots = OutNodalVector.GetData();
+		}
+	};
+
+	FillNodalVector(EAxis::U, UNodalVector);
+	FillNodalVector(EAxis::V, VNodalVector);
+
+	const int32 CoordinateCount = NurbsSurfaceData->m_uiUCtrlSize * NurbsSurfaceData->m_uiVCtrlSize * 4;
+	TArray<double> HomogeneousPoles;
+	HomogeneousPoles.SetNumUninitialized(CoordinateCount);
+
+	if (InObjectReference == EAliasObjectReference::WorldReference)
+	{
+		AliasSurface.CVsWorldPositionInclMultiples(HomogeneousPoles.GetData());
+	}
+	else if (InObjectReference == EAliasObjectReference::ParentReference)
+	{
+		AlTM TranformMatrix(InAlMatrix);
+		AliasSurface.CVsAffectedPositionInclMultiples(TranformMatrix, HomogeneousPoles.GetData());
+	}
+	else  // EAliasObjectReference::LocalReference
+	{
+		AliasSurface.CVsUnaffectedPositionInclMultiples(HomogeneousPoles.GetData());
+	}
+
+	int32 PoleCount = NurbsSurfaceData->m_uiUCtrlSize * NurbsSurfaceData->m_uiVCtrlSize;
+	TArray<A3DDouble> Weights;
+	TArray<A3DVector3dData> ControlPoints;
+	ControlPoints.SetNum(PoleCount);
+	Weights.SetNum(PoleCount);
+	A3DVector3dData* ControlPointPtr = ControlPoints.GetData();
+	A3DDouble* WeightPtr = Weights.GetData();
+	
+	NurbsSurfaceData->m_pdWeights = WeightPtr;
+	NurbsSurfaceData->m_pCtrlPts = ControlPointPtr;
+	
+	TFunction<void(const double*, A3DVector3dData&, A3DDouble&)> SetA3DPole = [](const double* HomogeneousPole, A3DVector3dData & OutTSPoint, A3DDouble & OutWeight)
+	{
+		OutTSPoint.m_dX = HomogeneousPole[0];
+		OutTSPoint.m_dY = HomogeneousPole[1];
+		OutTSPoint.m_dZ = HomogeneousPole[2];
+		OutWeight = HomogeneousPole[3];
+	};
+
+	double* AliasPoles = HomogeneousPoles.GetData();
+	for (int32 Index = 0; Index < PoleCount; ++Index, AliasPoles += 4, ++ControlPointPtr, ++WeightPtr)
+	{
+		SetA3DPole(AliasPoles, *ControlPointPtr, *WeightPtr);
+	}
+
+	return CADLibrary::TechSoftInterface::CreateSurfaceNurbs(*NurbsSurfaceData);
+}
+
+} // ns AliasToTechsoftUtils
+
+A3DCrvBase* FAliasModelToTechSoftConverter::CreateCurve(const AlTrimCurve& AliasTrimCurve)
+{
+	CADLibrary::TUniqueTSObj<A3DCrvNurbsData> NurbsCurveData;
+
+	NurbsCurveData->m_bIs2D = true;
+	NurbsCurveData->m_bRational = true;
+	NurbsCurveData->m_uiDegree = AliasTrimCurve.degree();
+
+	int ControlPointCount = AliasTrimCurve.numberOfCVs();
+	int32 KnotCount = AliasTrimCurve.realNumberOfKnots();
+
+	using AlPoint = double[3];
+
+	TArray<AlPoint> AliasPoles;
+	TArray<double> AliasNodalVector;
+	AliasNodalVector.SetNum(ControlPointCount);
+	AliasPoles.SetNumUninitialized(ControlPointCount*4);
+
+	// Notice that each CV has three coordinates - the three coordinates describe 2D parameter space, with a homogeneous coordinate.
+	// Each control point is u, v and w, where u and v are parameter space and w is the homogeneous coordinate.
+	AliasTrimCurve.CVsUVPosition(AliasNodalVector.GetData(), AliasPoles.GetData());
+
+	AliasNodalVector.SetNum(KnotCount);
+	AliasTrimCurve.realKnotVector(AliasNodalVector.GetData());
+
+	TArray<A3DDouble> NodalVector;
+	NodalVector.Reserve(KnotCount+2);
+	NodalVector.Add(AliasNodalVector[0]);
+	for (double Value : AliasNodalVector)
+	{
+		NodalVector.Add(Value);
+	}
+	NodalVector.Add(AliasNodalVector.Last());
+
+	TArray<double> Weights;
+	Weights.SetNumUninitialized(ControlPointCount);
+
+	AliasNodalVector.SetNumUninitialized(KnotCount);
+
+	TArray<A3DVector3dData> ControlPointArray;
+	TArray<A3DDouble> WeightArray;
+
+	ControlPointArray.SetNumUninitialized(ControlPointCount);
+	WeightArray.SetNumUninitialized(ControlPointCount);
+
+	A3DVector3dData* ControlPointPtr = ControlPointArray.GetData();
+	A3DDouble* WeightPtr = WeightArray.GetData();
+	AlPoint* AliasPolePtr = AliasPoles.GetData();
+
+	TFunction<void(const AlPoint&, A3DVector3dData&, A3DDouble&)> SetA3DPole = [](const AlPoint& AliasPole, A3DVector3dData& OutTSPoint, A3DDouble& OutWeight)
+	{
+		OutTSPoint.m_dX = AliasPole[0];
+		OutTSPoint.m_dY = AliasPole[1];
+		OutTSPoint.m_dZ = 0;
+		OutWeight = AliasPole[2];
+	};
+
+	for (int32 Index = 0; Index < ControlPointCount; ++Index, ++AliasPolePtr, ++ControlPointPtr, ++WeightPtr)
+	{
+		SetA3DPole(*AliasPolePtr, *ControlPointPtr, *WeightPtr);
+	}
+
+	NurbsCurveData->m_eKnotType = kA3DKnotTypeUnspecified;
+	NurbsCurveData->m_eCurveForm = kA3DBSplineCurveFormUnspecified;
+
+	NurbsCurveData->m_pCtrlPts = ControlPointArray.GetData();
+	NurbsCurveData->m_uiCtrlSize = ControlPointArray.Num();
+
+	NurbsCurveData->m_pdWeights = WeightArray.GetData();
+	NurbsCurveData->m_uiWeightSize = WeightArray.Num();
+
+	NurbsCurveData->m_pdKnots = NodalVector.GetData();
+	NurbsCurveData->m_uiKnotSize = NodalVector.Num();
+
+	return CADLibrary::TechSoftInterface::CreateCurveNurbs(*NurbsCurveData);
+}
+
+A3DTopoCoEdge* FAliasModelToTechSoftConverter::CreateEdge(const AlTrimCurve& TrimCurve)
+{
+	A3DCrvBase* NurbsCurvePtr = CreateCurve(TrimCurve);
+	if (NurbsCurvePtr == nullptr)
+	{
+		return nullptr;
+	}
+
+	A3DTopoEdge* EdgePtr = CADLibrary::TechSoftUtils::CreateTopoEdge();
+	if (EdgePtr == nullptr)
+	{
+		return nullptr;
+	}
+
+	CADLibrary::TUniqueTSObj<A3DTopoCoEdgeData> CoEdgeData;
+
+	CoEdgeData->m_pUVCurve = NurbsCurvePtr;
+	CoEdgeData->m_pEdge = EdgePtr;
+	CoEdgeData->m_ucOrientationWithLoop = TrimCurve.isReversed();
+	CoEdgeData->m_ucOrientationUVWithLoop = 1;
+
+	A3DTopoCoEdge* CoEdgePtr = CADLibrary::TechSoftInterface::CreateTopoCoEdge(*CoEdgeData);
+
+	// Only TrimCurve with twin need to be in the map used in LinkEdgesLoop 
+	TUniquePtr<AlTrimCurve> TwinCurve(TrimCurve.getTwinCurve());
+	if (TwinCurve.IsValid())
+	{
+		AlEdgeToTSCoEdge.Add(TrimCurve.fSpline, CoEdgePtr);
+	}
+	return CoEdgePtr;
+}
+
+A3DTopoLoop* FAliasModelToTechSoftConverter::CreateTopoLoop(const AlTrimBoundary& TrimBoundary)
+{
+	TArray<A3DTopoCoEdge*> Edges;
+	Edges.Reserve(20);
+
+	for (TUniquePtr<AlTrimCurve> TrimCurve(TrimBoundary.firstCurve()); TrimCurve.IsValid(); TrimCurve = TUniquePtr<AlTrimCurve>(TrimCurve->nextCurve()))
+	{
+		A3DTopoCoEdge* Edge = CreateEdge(*TrimCurve);
+		if (Edge != nullptr)
+		{
+			Edges.Add(Edge);
+		}
+	}
+
+	if (Edges.Num() == 0)
+	{
+		return nullptr;
+	}
+
+	CADLibrary::TUniqueTSObj<A3DTopoLoopData> LoopData;
+
+	LoopData->m_ppCoEdges = Edges.GetData();
+	LoopData->m_uiCoEdgeSize = Edges.Num();
+	LoopData->m_ucOrientationWithSurface = 1;
+
+	return CADLibrary::TechSoftInterface::CreateTopoLoop(*LoopData);
+}
+
+void FAliasModelToTechSoftConverter::LinkEdgesLoop(const AlTrimBoundary& TrimBoundary)
+{
+	for (TUniquePtr<AlTrimCurve> TrimCurve(TrimBoundary.firstCurve()); TrimCurve.IsValid(); TrimCurve = TUniquePtr<AlTrimCurve>(TrimCurve->nextCurve()))
+	{
+		A3DTopoCoEdge** Edge = AlEdgeToTSCoEdge.Find(TrimCurve->fSpline);
+		if (Edge == nullptr)
+		{
+			continue;
+		}
+
+		// Link edges
+		TUniquePtr<AlTrimCurve> TwinCurve(TrimCurve->getTwinCurve());
+		if (TwinCurve.IsValid())
+		{
+			if (A3DTopoCoEdge** TwinEdge = AlEdgeToTSCoEdge.Find(TwinCurve->fSpline))
+			{
+				if (TwinEdge != nullptr)
+				{
+					CADLibrary::TechSoftInterface::LinkCoEdges(*Edge, *TwinEdge);
+					break;
+				}
+			}
+		}
+	}
+}
+
+A3DTopoFace* FAliasModelToTechSoftConverter::AddTrimRegion(const AlTrimRegion& InTrimRegion, EAliasObjectReference InObjectReference, const AlMatrix4x4& InAlMatrix)
+{
+	A3DSurfBase* CarrierSurface = AliasToTechSoftUtils::AddNURBSSurface(InTrimRegion, InObjectReference, InAlMatrix);
+	if (CarrierSurface == nullptr)
+	{
+		return nullptr;
+	}
+
+	TArray<A3DTopoLoop*> Loops;
+	Loops.Reserve(5);
+
+	for (TUniquePtr<AlTrimBoundary> TrimBoundary(InTrimRegion.firstBoundary()); TrimBoundary.IsValid(); TrimBoundary = TUniquePtr<AlTrimBoundary>(TrimBoundary->nextBoundary()))
+	{
+		A3DTopoLoop* Loop = CreateTopoLoop(*TrimBoundary);
+		if (Loop != nullptr)
+		{
+			Loops.Add(Loop);
+			LinkEdgesLoop(*TrimBoundary);
+		}
+	}
+
+	if (Loops.Num() == 0)
+	{
+		return nullptr;
+	}
+
+	CADLibrary::TUniqueTSObj<A3DTopoFaceData> Face;
+	Face->m_pSurface = CarrierSurface;
+	Face->m_bHasTrimDomain = false;
+	Face->m_ppLoops = Loops.GetData();
+	Face->m_uiLoopSize = Loops.Num();
+	Face->m_uiOuterLoopIndex = 0;
+	Face->m_dTolerance = 0.01; //mm
+
+	A3DTopoFace* FacePtr = CADLibrary::TechSoftInterface::CreateTopoFace(*Face);
+
+	CADLibrary::TechSoftUtils::SetEntityGraphicsColor(FacePtr, FColor(200,200,200));
+
+	return FacePtr;
+}
+#endif
+
+bool FAliasModelToTechSoftConverter::AddBRep(AlDagNode& DagNode, EAliasObjectReference InObjectReference)
+{
+#ifdef USE_TECHSOFT_SDK
+	AlEdgeToTSCoEdge.Empty();
+
+	boolean bAlOrientation;
+	DagNode.getSurfaceOrientation(bAlOrientation);
+	bool bOrientation = (bool)bAlOrientation;
+
+	AlMatrix4x4 AlMatrix;
+	if (InObjectReference == EAliasObjectReference::ParentReference)
+	{
+		DagNode.localTransformationMatrix(AlMatrix);
+	}
+
+	TArray<A3DTopoFace*> TSFaces;
+	TArray<A3DUns8> FaceOrientations;
+	TSFaces.Reserve(100);
+	FaceOrientations.Reserve(100);
+
+	AlObjectType objectType = DagNode.type();
+	switch (objectType)
+	{
+
+	case kShellNodeType:
+	{
+		if (AlShellNode* ShellNode = DagNode.asShellNodePtr())
+		{
+			TUniquePtr<AlShell> AliasShell(ShellNode->shell());
+			if(AlIsValid(AliasShell.Get()))
+			{
+				for (TUniquePtr<AlTrimRegion> TrimRegion(AliasShell->firstTrimRegion()); TrimRegion.IsValid(); TrimRegion = TUniquePtr<AlTrimRegion>(TrimRegion->nextRegion()))
+				{
+					A3DTopoFace* TSFace = AddTrimRegion(*TrimRegion, InObjectReference, AlMatrix);
+					if (TSFace != nullptr)
+					{
+						TSFaces.Add(TSFace);
+						FaceOrientations.Add(bOrientation);
+					}
+				}
+			}
+		}
+		break;
+	}
+
+	case kSurfaceNodeType:
+	{
+		if (AlSurfaceNode* SurfaceNode = DagNode.asSurfaceNodePtr())
+		{
+			TUniquePtr<AlSurface> AliasSurface(SurfaceNode->surface());
+			if (AlIsValid(AliasSurface.Get()))
+			{
+				TUniquePtr<AlTrimRegion> TrimRegion(AliasSurface->firstTrimRegion());
+				if (TrimRegion.IsValid())
+				{
+					for (; TrimRegion.IsValid(); TrimRegion = TUniquePtr<AlTrimRegion>(TrimRegion->nextRegion()))
+					{
+						A3DTopoFace* TSFace = AddTrimRegion(*TrimRegion, InObjectReference, AlMatrix);
+						if (TSFace != nullptr)
+						{
+							TSFaces.Add(TSFace);
+							FaceOrientations.Add(bOrientation);
+						}
+					}
+					break;
+				}
+
+				A3DSurfBase* TSSurface = AliasToTechSoftUtils::AddNURBSSurface(*AliasSurface, InObjectReference, AlMatrix);
+				if (TSSurface != nullptr)
+				{
+					A3DTopoFace* TSFace = CADLibrary::TechSoftUtils::CreateTopoFaceWithNaturalLoop(TSSurface);
+					if (TSFace != nullptr)
+					{
+						TSFaces.Add(TSFace);
+						FaceOrientations.Add(bOrientation);
+					}
+				}
+			}
+		}
+		break;
+	}
+	default:
+		break;
+	}
+
+	if (TSFaces.IsEmpty())
+	{
+		return false;
+	}
+
+	A3DTopoShell* TopoShellPtr = nullptr;
+	{
+		CADLibrary::TUniqueTSObj<A3DTopoShellData> TopoShellData;
+		TopoShellData->m_bClosed = false;
+		TopoShellData->m_ppFaces = TSFaces.GetData();
+		TopoShellData->m_uiFaceSize = TSFaces.Num();
+		TopoShellData->m_pucOrientationWithShell = FaceOrientations.GetData();
+
+		TopoShellPtr = CADLibrary::TechSoftInterface::CreateTopoShell(*TopoShellData);
+
+		if (TopoShellPtr == nullptr)
+		{
+			return false;
+		}
+	}
+
+	A3DRiRepresentationItem* RiRepresentationItem = CADLibrary::TechSoftUtils::CreateRIBRep(TopoShellPtr);
+
+	if (RiRepresentationItem != nullptr)
+	{
+		RiRepresentationItems.Add(RiRepresentationItem);
+		return true;
+	}
+#endif
+	return false;
+}
+
+}
+#endif
diff --git a/Engine/Plugins/Enterprise/DatasmithCADImporter/Source/DatasmithWireTranslator/Private/AliasModelToTechSoftConverter.h b/Engine/Plugins/Enterprise/DatasmithCADImporter/Source/DatasmithWireTranslator/Private/AliasModelToTechSoftConverter.h
new file mode 100644
--- /dev/null
+++ b/Engine/Plugins/Enterprise/DatasmithCADImporter/Source/DatasmithWireTranslator/Private/AliasModelToTechSoftConverter.h
@@ -0,0 +1,56 @@
+// Copyright Epic Games, Inc. All Rights Reserved.
+
+#pragma once
+
+#include "CoreMinimal.h"
+
+#include "CADModelToTechSoftConverterBase.h"
+
+#include "AliasBRepConverter.h"
+
+#include "CADKernel/Core/Session.h"
+#include "CADKernel/Core/Types.h"
+
+class AlDagNode;
+class AlShell;
+class AlSurface;
+class AlTrimBoundary;
+class AlTrimCurve;
+class AlTrimRegion;
+
+typedef void A3DTopoCoEdge;
+typedef void A3DTopoFace;
+typedef void A3DTopoLoop;
+typedef void A3DTopoCoEdge;
+typedef void A3DCrvBase;
+
+typedef double AlMatrix4x4[4][4];
+
+namespace UE_DATASMITHWIRETRANSLATOR_NAMESPACE
+{
+
+class FAliasModelToTechSoftConverter : public FCADModelToTechSoftConverterBase, public IAliasBRepConverter
+{
+
+public:
+	FAliasModelToTechSoftConverter(CADLibrary::FImportParameters InImportParameters)
+		: FCADModelToTechSoftConverterBase(InImportParameters)
+	{
+	}
+
+	virtual bool AddBRep(AlDagNode& DagNode, EAliasObjectReference ObjectReference) override;
+
+protected:
+
+	TMap<void*, A3DTopoCoEdge*> AlEdgeToTSCoEdge;
+
+	A3DTopoFace* AddTrimRegion(const AlTrimRegion& InTrimRegion, EAliasObjectReference InObjectReference, const AlMatrix4x4& InAlMatrix);
+	A3DTopoLoop* CreateTopoLoop(const AlTrimBoundary& TrimBoundary);
+	A3DTopoCoEdge* CreateEdge(const AlTrimCurve& TrimCurve);
+	A3DCrvBase* CreateCurve(const AlTrimCurve& TrimCurve);
+
+	void LinkEdgesLoop(const AlTrimBoundary& TrimBoundary);
+};
+
+}
+
diff --git a/Engine/Plugins/Enterprise/DatasmithCADImporter/Source/DatasmithWireTranslator/Private/DatasmithWireTranslator.cpp b/Engine/Plugins/Enterprise/DatasmithCADImporter/Source/DatasmithWireTranslator/Private/DatasmithWireTranslator.cpp
--- a/Engine/Plugins/Enterprise/DatasmithCADImporter/Source/DatasmithWireTranslator/Private/DatasmithWireTranslator.cpp
+++ b/Engine/Plugins/Enterprise/DatasmithCADImporter/Source/DatasmithWireTranslator/Private/DatasmithWireTranslator.cpp
@@ -26,13 +26,18 @@
 #include "Modules/ModuleManager.h"
 #endif
 
-#include "AliasModelToCoretechConverter.h" // requires CoreTech as public dependency
+#include "AliasBrepConverter.h"
 #include "AliasModelToCADKernelConverter.h"
+#include "AliasModelToCoretechConverter.h" // requires CoreTech as public dependency
+#include "AliasModelToTechSoftConverter.h" // requires Techsoft as public dependency
 #include "CADInterfacesModule.h"
-#include "AliasBrepConverter.h"
 #include "CoreTechSurfaceExtension.h"
 #include "CoreTechSurfaceHelper.h"
 
+#if PLATFORM_WINDOWS
+#include "Windows/AllowWindowsPlatformTypes.h"
+#endif
+
 #ifdef USE_OPENMODEL
 #include <AlChannel.h>
 #include <AlDagNode.h>
@@ -225,9 +230,18 @@ public:
 		CADLibrary::FImportParameters ImportParameters(0.01, 1);
 		if(CADLibrary::FImportParameters::bGDisableCADKernelTessellation)
 		{
-			TSharedRef<FAliasModelToCoretechConverter> AliasToCoretechConverter = MakeShared<FAliasModelToCoretechConverter>(TEXT("Al2CTSharedSession"), ImportParameters);
-			CADModelConverter = AliasToCoretechConverter;
-			AliasBRepConverter = AliasToCoretechConverter;
+			if (CADLibrary::FImportParameters::GCADLibrary == TEXT("KernelIO"))
+			{
+				TSharedRef<FAliasModelToCoretechConverter> AliasToCoretechConverter = MakeShared<FAliasModelToCoretechConverter>(TEXT("Al2CTSharedSession"), ImportParameters);
+				CADModelConverter = AliasToCoretechConverter;
+				AliasBRepConverter = AliasToCoretechConverter;
+			}
+			else if (CADLibrary::FImportParameters::GCADLibrary == TEXT("TechSoft"))
+			{
+				TSharedRef<FAliasModelToTechSoftConverter> AliasToTechSoftConverter = MakeShared<FAliasModelToTechSoftConverter>(ImportParameters);
+				CADModelConverter = AliasToTechSoftConverter;
+				AliasBRepConverter = AliasToTechSoftConverter;
+			}
 		}
 		else
 		{
@@ -1855,7 +1869,7 @@ TOptional<FMeshDescription> FWireTranslatorImpl::MeshDagNodeWithExternalMesher(A
 
 	CADModelConverter->RepairTopology();
 
-	CADModelConverter->SaveBRep(*OutputPath, MeshElement);
+	CADModelConverter->SaveModel(*OutputPath, MeshElement);
 
 	FMeshDescription MeshDescription;
 	DatasmithMeshHelper::PrepareAttributeForStaticMesh(MeshDescription);
@@ -1889,7 +1903,7 @@ TOptional<FMeshDescription> FWireTranslatorImpl::MeshDagNodeWithExternalMesher(T
 
 	CADModelConverter->RepairTopology();
 
-	CADModelConverter->SaveBRep(*OutputPath, MeshElement);
+	CADModelConverter->SaveModel(*OutputPath, MeshElement);
 
 	FMeshDescription MeshDescription;
 	DatasmithMeshHelper::PrepareAttributeForStaticMesh(MeshDescription);
@@ -2221,3 +2235,6 @@ void FDatasmithWireTranslator::SetSceneImportOptions(TArray<TStrongObjectPtr<UDa
 
 #undef LOCTEXT_NAMESPACE // "DatasmithWireTranslator"
 
+#if PLATFORM_WINDOWS
+#include "Windows/HideWindowsPlatformTypes.h"
+#endif
\ No newline at end of file
diff --git a/Engine/Plugins/Enterprise/DatasmithCADImporter/Source/DatasmithWireTranslator/Private/OpenModelUtils.cpp b/Engine/Plugins/Enterprise/DatasmithCADImporter/Source/DatasmithWireTranslator/Private/OpenModelUtils.cpp
--- a/Engine/Plugins/Enterprise/DatasmithCADImporter/Source/DatasmithWireTranslator/Private/OpenModelUtils.cpp
+++ b/Engine/Plugins/Enterprise/DatasmithCADImporter/Source/DatasmithWireTranslator/Private/OpenModelUtils.cpp
@@ -8,15 +8,20 @@
 #include "DatasmithTranslator.h"
 #include "IDatasmithSceneElements.h"
 
-#include "AlDagNode.h"
-#include "AlMesh.h"
-#include "AlLayer.h"
-#include "AlTesselate.h"
 #include "MeshAttributes.h"
 #include "MeshDescription.h"
 #include "StaticMeshAttributes.h"
 #include "StaticMeshOperations.h"
 
+#if PLATFORM_WINDOWS
+#include "Windows/AllowWindowsPlatformTypes.h"
+#endif
+
+#include "AlDagNode.h"
+#include "AlMesh.h"
+#include "AlLayer.h"
+#include "AlTesselate.h"
+
 namespace UE_DATASMITHWIRETRANSLATOR_NAMESPACE
 {
 
@@ -1835,6 +1840,10 @@ TSharedPtr<AlDagNode> OpenModelUtils::TesselateDagLeaf(const AlDagNode& DagLeaf,
 
 }
 
+#if PLATFORM_WINDOWS
+#include "Windows/HideWindowsPlatformTypes.h"
+#endif
+
 #endif
 
 
diff --git a/Engine/Plugins/Enterprise/DatasmithCADImporter/Source/DatasmithWireTranslator/Private/OpenModelUtils.h b/Engine/Plugins/Enterprise/DatasmithCADImporter/Source/DatasmithWireTranslator/Private/OpenModelUtils.h
--- a/Engine/Plugins/Enterprise/DatasmithCADImporter/Source/DatasmithWireTranslator/Private/OpenModelUtils.h
+++ b/Engine/Plugins/Enterprise/DatasmithCADImporter/Source/DatasmithWireTranslator/Private/OpenModelUtils.h
@@ -8,16 +8,21 @@
 
 #ifdef USE_OPENMODEL
 
-#include "AlAccessTypes.h"
+#if PLATFORM_WINDOWS
+#include "Windows/AllowWindowsPlatformTypes.h"
+#endif
+
 #include "AlShadingFields.h"
 #include "AlDagNode.h"
 #include "AlPersistentID.h"
 
+#if PLATFORM_WINDOWS
+#include "Windows/HideWindowsPlatformTypes.h"
+#endif
 
 class IDatasmithActorElement;
-class AlDagNode;
 class AlMesh;
-class AlPersistentID;
+
 struct FMeshDescription;
 
 typedef double AlMatrix4x4[4][4];
diff --git a/Engine/Plugins/Enterprise/DatasmithCADImporter/Source/ParametricSurface/ParametricSurface.Build.cs b/Engine/Plugins/Enterprise/DatasmithCADImporter/Source/ParametricSurface/ParametricSurface.Build.cs
--- a/Engine/Plugins/Enterprise/DatasmithCADImporter/Source/ParametricSurface/ParametricSurface.Build.cs
+++ b/Engine/Plugins/Enterprise/DatasmithCADImporter/Source/ParametricSurface/ParametricSurface.Build.cs
@@ -26,6 +26,7 @@ namespace UnrealBuildTool.Rules
 					"DatasmithContent",
 					"DatasmithTranslator",
 					"Engine",
+					"Json",
 					"MeshDescription",
 					"StaticMeshDescription",
 				}
diff --git a/Engine/Plugins/Enterprise/DatasmithCADImporter/Source/ParametricSurface/Private/ParametricSurfaceTranslator.cpp b/Engine/Plugins/Enterprise/DatasmithCADImporter/Source/ParametricSurface/Private/ParametricSurfaceTranslator.cpp
--- a/Engine/Plugins/Enterprise/DatasmithCADImporter/Source/ParametricSurface/Private/ParametricSurfaceTranslator.cpp
+++ b/Engine/Plugins/Enterprise/DatasmithCADImporter/Source/ParametricSurface/Private/ParametricSurfaceTranslator.cpp
@@ -36,7 +36,7 @@ void FParametricSurfaceTranslator::SetSceneImportOptions(TArray<TStrongObjectPtr
 	}
 }
 
-bool FParametricSurfaceTranslator::AddSurfaceData(const TCHAR* MeshFilePath, const CADLibrary::FImportParameters& ImportParameters, const CADLibrary::FMeshParameters& MeshParameters, FDatasmithMeshElementPayload& OutMeshPayload)
+bool ParametricSurfaceUtils::AddSurfaceData(const TCHAR* MeshFilePath, const CADLibrary::FImportParameters& ImportParameters, const CADLibrary::FMeshParameters& InMeshParameters, const FDatasmithTessellationOptions& InCommonTessellationOptions, FDatasmithMeshElementPayload& OutMeshPayload)
 {
 	if (MeshFilePath && IFileManager::Get().FileExists(MeshFilePath))
 	{
@@ -48,8 +48,8 @@ bool FParametricSurfaceTranslator::AddSurfaceData(const TCHAR* MeshFilePath, con
 		}
 
 		ParametricSurfaceData->SetImportParameters(ImportParameters);
-		ParametricSurfaceData->SetMeshParameters(MeshParameters);
-		ParametricSurfaceData->SetLastTessellationOptions(CommonTessellationOptions);
+		ParametricSurfaceData->SetMeshParameters(InMeshParameters);
+		ParametricSurfaceData->SetLastTessellationOptions(InCommonTessellationOptions);
 
 		OutMeshPayload.AdditionalData.Add(ParametricSurfaceData);
 
diff --git a/Engine/Plugins/Enterprise/DatasmithCADImporter/Source/ParametricSurface/Private/TechSoft/CADModelToTechSoftConverterBase.cpp b/Engine/Plugins/Enterprise/DatasmithCADImporter/Source/ParametricSurface/Private/TechSoft/CADModelToTechSoftConverterBase.cpp
new file mode 100644
--- /dev/null
+++ b/Engine/Plugins/Enterprise/DatasmithCADImporter/Source/ParametricSurface/Private/TechSoft/CADModelToTechSoftConverterBase.cpp
@@ -0,0 +1,105 @@
+// Copyright Epic Games, Inc. All Rights Reserved.
+
+#include "CADModelToTechSoftConverterBase.h"
+
+
+#include "CADData.h"
+#include "MeshDescriptionHelper.h"
+#include "ParametricSurfaceTranslator.h"
+#include "Serialization/JsonSerializer.h"
+#include "Serialization/JsonWriter.h"
+#include "TechSoftInterface.h"
+#include "TechSoftUtils.h"
+#include "TUniqueTechSoftObj.h"
+
+bool FCADModelToTechSoftConverterBase::RepairTopology()
+{
+#ifdef USE_TECHSOFT_SDK
+	// Apply stitching if applicable
+	if (ImportParameters.GetStitchingTechnique() != CADLibrary::StitchingNone)
+	{
+		CADLibrary::TUniqueTSObj<A3DSewOptionsData> SewOptionsData;
+		SewOptionsData->m_bComputePreferredOpenShellOrientation = false;
+
+		A3DRiBrepModel** OutNewBReps;
+		uint32 OutNewBRepCount;
+		CADLibrary::TechSoftInterface::SewBReps((A3DRiBrepModel**) RiRepresentationItems.GetData(), RiRepresentationItems.Num(), 0.01, SewOptionsData.GetPtr(), &OutNewBReps, OutNewBRepCount);
+
+		RiRepresentationItems.Empty(OutNewBRepCount);
+		for (uint32 Index = 0; Index < OutNewBRepCount; ++Index)
+		{
+			RiRepresentationItems.Add(OutNewBReps[Index]);
+		}
+
+		if (OutNewBRepCount != 0)
+		{
+			delete OutNewBReps;
+		}
+	}
+#endif
+	return true;
+}
+
+void FCADModelToTechSoftConverterBase::InitializeProcess(double InMetricUnit)
+{
+	RiRepresentationItems.Empty();
+	ModelFile.Reset();
+}
+
+
+bool FCADModelToTechSoftConverterBase::SaveModel(const TCHAR* InFolderPath, TSharedRef<IDatasmithMeshElement>& MeshElement)
+{
+#ifdef USE_TECHSOFT_SDK
+	FString FilePath = FPaths::Combine(InFolderPath, MeshElement->GetName()) + TEXT(".prc");
+
+	FString JsonString;
+	{
+		// Save file unit and default color and material attributes in a json string
+		// This will be used when the file is reloaded
+		TSharedPtr<FJsonObject> JsonObject = MakeShared<FJsonObject>();
+		JsonObject->SetNumberField(JSON_ENTRY_FILE_UNIT, 1);
+		JsonObject->SetNumberField(JSON_ENTRY_COLOR_NAME, 0);
+		JsonObject->SetNumberField(JSON_ENTRY_MATERIAL_NAME, 0);
+
+		TSharedRef< TJsonWriter< TCHAR, TPrettyJsonPrintPolicy<TCHAR> > > JsonWriter = TJsonWriterFactory< TCHAR, TPrettyJsonPrintPolicy<TCHAR> >::Create(&JsonString);
+
+		FJsonSerializer::Serialize(JsonObject.ToSharedRef(), JsonWriter);
+	}
+
+	ModelFile = CADLibrary::TechSoftUtils::SaveBodiesToPcrFile(RiRepresentationItems.GetData(), RiRepresentationItems.Num(), FilePath, JsonString);
+
+	MeshElement->SetFile(*FilePath);
+#endif
+	return true;
+}
+
+bool FCADModelToTechSoftConverterBase::Tessellate(const CADLibrary::FMeshParameters& InMeshParameters, FMeshDescription& OutMeshDescription)
+{
+	CADLibrary::FBodyMesh BodyMesh;
+	BodyMesh.BodyID = 1;
+
+#ifdef USE_TECHSOFT_SDK
+	for (A3DRiRepresentationItem* Representation : RiRepresentationItems)
+	{
+		CADLibrary::TechSoftUtils::FillBodyMesh(Representation, ImportParameters, 1.0, BodyMesh);
+	}
+#endif
+
+	if (BodyMesh.Faces.Num() == 0)
+	{
+		return false;
+	}
+
+	if (!CADLibrary::ConvertBodyMeshToMeshDescription(ImportParameters, InMeshParameters, BodyMesh, OutMeshDescription))
+	{
+		ensureMsgf(false, TEXT("Error during mesh conversion"));
+		return false;
+	}
+	return true;
+}
+
+void FCADModelToTechSoftConverterBase::AddSurfaceDataForMesh(const TCHAR* InFilePath, const CADLibrary::FMeshParameters& InMeshParameters, const FDatasmithTessellationOptions& InTessellationOptions, FDatasmithMeshElementPayload& OutMeshPayload) const
+{
+	ParametricSurfaceUtils::AddSurfaceData(InFilePath, ImportParameters, InMeshParameters, InTessellationOptions, OutMeshPayload);
+}
+
diff --git a/Engine/Plugins/Enterprise/DatasmithCADImporter/Source/ParametricSurface/Private/TechSoft/TechSoftParametricSurface.cpp b/Engine/Plugins/Enterprise/DatasmithCADImporter/Source/ParametricSurface/Private/TechSoft/TechSoftParametricSurface.cpp
--- a/Engine/Plugins/Enterprise/DatasmithCADImporter/Source/ParametricSurface/Private/TechSoft/TechSoftParametricSurface.cpp
+++ b/Engine/Plugins/Enterprise/DatasmithCADImporter/Source/ParametricSurface/Private/TechSoft/TechSoftParametricSurface.cpp
@@ -2,11 +2,11 @@
 
 #include "TechSoftParametricSurface.h"
 
-#include "TechSoftInterface.h"
-
 #include "CADInterfacesModule.h"
 #include "DatasmithAdditionalData.h"
 #include "DatasmithPayload.h"
+#include "TechSoftInterface.h"
+#include "TechSoftUtils.h"
 
 #include "Engine/StaticMesh.h"
 #include "HAL/PlatformFileManager.h"
@@ -24,8 +24,11 @@ bool UTechSoftParametricSurfaceData::Tessellate(UStaticMesh& StaticMesh, const F
 
 #if WITH_EDITOR
 	// make a temporary file as TechSoft can only deal with files.
+	FString CachePath = FPaths::ConvertRelativePathToFull(FPaths::Combine(FPaths::ProjectIntermediateDir(), TEXT("Retessellate")));
+	IFileManager::Get().MakeDirectory(*CachePath, true);
+
 	int32 Hash = GetTypeHash(StaticMesh.GetPathName());
-	FString ResourceFile = FPaths::ConvertRelativePathToFull(FPaths::ProjectIntermediateDir() / FString::Printf(TEXT("0x%08x.hsf"), Hash));
+	FString ResourceFile = FPaths::ConvertRelativePathToFull(FPaths::Combine(CachePath, FString::Printf(TEXT("0x%08x.prc"), Hash)));
 
 	FFileHelper::SaveArrayToFile(RawData, *ResourceFile);
 
@@ -44,14 +47,14 @@ bool UTechSoftParametricSurfaceData::Tessellate(UStaticMesh& StaticMesh, const F
 		CADLibrary::FImportParameters ImportParameters(SceneParameters.MetricUnit, SceneParameters.ScaleFactor, (FDatasmithUtils::EModelCoordSystem)SceneParameters.ModelCoordSys);
 		ImportParameters.SetTesselationParameters(RetessellateOptions.ChordTolerance, RetessellateOptions.MaxEdgeLength, RetessellateOptions.NormalTolerance, (CADLibrary::EStitchingTechnique)RetessellateOptions.StitchingTechnique);
 
-		FTechSoftInterface& TechSoftInterface = TechSoftUtils::GetTechSoftInterface();
+		CADLibrary::FTechSoftInterface& TechSoftInterface = CADLibrary::FTechSoftInterface::Get();
 
 		bSuccessfulTessellation = TechSoftInterface.InitializeKernel(*FPaths::EnginePluginsDir());
 		if (bSuccessfulTessellation)
 		{
 			CADLibrary::FBodyMesh BodyMesh;
 
-			bSuccessfulTessellation = TechSoftInterface.GetBodyFromHsfFile(ResourceFile, ImportParameters, BodyMesh);
+			bSuccessfulTessellation = CADLibrary::TechSoftUtils::GetBodyFromPcrFile(ResourceFile, ImportParameters, BodyMesh);
 
 			if (bSuccessfulTessellation)
 			{
diff --git a/Engine/Plugins/Enterprise/DatasmithCADImporter/Source/ParametricSurface/Public/CADModelConverter.h b/Engine/Plugins/Enterprise/DatasmithCADImporter/Source/ParametricSurface/Public/CADModelConverter.h
--- a/Engine/Plugins/Enterprise/DatasmithCADImporter/Source/ParametricSurface/Public/CADModelConverter.h
+++ b/Engine/Plugins/Enterprise/DatasmithCADImporter/Source/ParametricSurface/Public/CADModelConverter.h
@@ -44,7 +44,7 @@ public:
 	virtual void InitializeProcess(double InMetricUnit) = 0;
 
 	virtual bool RepairTopology() = 0;
-	virtual bool SaveBRep(const TCHAR* OutputPath, TSharedRef<IDatasmithMeshElement>& MeshElement) = 0;
+	virtual bool SaveModel(const TCHAR* OutputPath, TSharedRef<IDatasmithMeshElement>& MeshElement) = 0;
 	virtual bool Tessellate(const CADLibrary::FMeshParameters& InMeshParameters, FMeshDescription& OutMeshDescription) = 0;
 
 	/**
diff --git a/Engine/Plugins/Enterprise/DatasmithCADImporter/Source/ParametricSurface/Public/CADModelToTechSoftConverterBase.h b/Engine/Plugins/Enterprise/DatasmithCADImporter/Source/ParametricSurface/Public/CADModelToTechSoftConverterBase.h
new file mode 100644
--- /dev/null
+++ b/Engine/Plugins/Enterprise/DatasmithCADImporter/Source/ParametricSurface/Public/CADModelToTechSoftConverterBase.h
@@ -0,0 +1,71 @@
+// Copyright Epic Games, Inc. All Rights Reserved.
+
+#pragma once
+
+#include "CoreMinimal.h"
+
+#include "CADModelConverter.h"
+#include "CADOptions.h"
+#include "IDatasmithSceneElements.h"
+
+#include "TechSoftInterface.h"
+
+struct FDatasmithMeshElementPayload;
+struct FDatasmithTessellationOptions;
+
+typedef void A3DAsmPartDefinition;
+typedef void A3DAsmProductOccurrence;
+typedef void A3DEntity;
+typedef void A3DRiRepresentationItem;
+typedef void A3DTopoBrepData;
+
+
+class PARAMETRICSURFACE_API FCADModelToTechSoftConverterBase : public CADLibrary::ICADModelConverter
+{
+public:
+
+	FCADModelToTechSoftConverterBase(CADLibrary::FImportParameters InImportParameters)
+		: ImportParameters(InImportParameters)
+	{
+	}
+
+	virtual void InitializeProcess(double InMetricUnit) override;
+
+	virtual bool RepairTopology() override;
+
+	virtual bool SaveModel(const TCHAR* InFolderPath, TSharedRef<IDatasmithMeshElement>& MeshElement) override;
+
+	virtual bool Tessellate(const CADLibrary::FMeshParameters& InMeshParameters, FMeshDescription& OutMeshDescription) override;
+
+	virtual void SetImportParameters(double ChordTolerance, double MaxEdgeLength, double NormalTolerance, CADLibrary::EStitchingTechnique StitchingTechnique) override
+	{
+		ImportParameters.SetTesselationParameters(ChordTolerance, MaxEdgeLength, NormalTolerance, StitchingTechnique);
+	}
+
+	virtual void SetMetricUnit(double NewMetricUnit) override
+	{
+		ImportParameters.SetMetricUnit(NewMetricUnit);
+	}
+
+	virtual double GetScaleFactor() const override
+	{
+		return ImportParameters.GetScaleFactor();
+	}
+
+	virtual double GetMetricUnit() const override
+	{
+		return ImportParameters.GetMetricUnit();
+	}
+
+	virtual bool IsSessionValid() override
+	{
+		return true;
+	}
+
+	virtual void AddSurfaceDataForMesh(const TCHAR* InFilePath, const CADLibrary::FMeshParameters& InMeshParameters, const FDatasmithTessellationOptions& InTessellationOptions, FDatasmithMeshElementPayload& OutMeshPayload) const override;
+
+protected:
+	CADLibrary::FImportParameters ImportParameters;
+	TArray<A3DRiRepresentationItem*> RiRepresentationItems;
+	CADLibrary::FUniqueTechSoftModelFile ModelFile;
+};
diff --git a/Engine/Plugins/Enterprise/DatasmithCADImporter/Source/ParametricSurface/Public/ParametricSurfaceTranslator.h b/Engine/Plugins/Enterprise/DatasmithCADImporter/Source/ParametricSurface/Public/ParametricSurfaceTranslator.h
--- a/Engine/Plugins/Enterprise/DatasmithCADImporter/Source/ParametricSurface/Public/ParametricSurfaceTranslator.h
+++ b/Engine/Plugins/Enterprise/DatasmithCADImporter/Source/ParametricSurface/Public/ParametricSurfaceTranslator.h
@@ -33,8 +33,6 @@ public:
 	virtual void SetSceneImportOptions(TArray<TStrongObjectPtr<UDatasmithOptionsBase>>& Options) override;
 	// End IDatasmithTranslator overrides
 
-	bool AddSurfaceData(const TCHAR* MeshFilePath, const CADLibrary::FImportParameters& InSceneParameters, const CADLibrary::FMeshParameters& MeshParameters, FDatasmithMeshElementPayload& OutMeshPayload);
-
 protected:
 	const FDatasmithTessellationOptions& GetCommonTessellationOptions()
 	{
@@ -52,3 +50,7 @@ private:
 	FDatasmithTessellationOptions CommonTessellationOptions;
 };
 
+namespace ParametricSurfaceUtils
+{
+	PARAMETRICSURFACE_API bool AddSurfaceData(const TCHAR* MeshFilePath, const CADLibrary::FImportParameters& InSceneParameters, const CADLibrary::FMeshParameters& InMeshParameters, const FDatasmithTessellationOptions& InCommonTessellationOptions, FDatasmithMeshElementPayload& OutMeshPayload);
+}
\ No newline at end of file
