Fix Use TechSoft to tessellate Alias BReps (linkage pbl)

diff --git a/Engine/Plugins/Enterprise/DatasmithCADImporter/Source/DatasmithWireTranslator/DatasmithWireTranslator2020.Build.cs b/Engine/Plugins/Enterprise/DatasmithCADImporter/Source/DatasmithWireTranslator/DatasmithWireTranslator2020.Build.cs
--- a/Engine/Plugins/Enterprise/DatasmithCADImporter/Source/DatasmithWireTranslator/DatasmithWireTranslator2020.Build.cs
+++ b/Engine/Plugins/Enterprise/DatasmithCADImporter/Source/DatasmithWireTranslator/DatasmithWireTranslator2020.Build.cs
@@ -7,6 +7,9 @@ public abstract class DatasmithWireTranslatorBase : ModuleRules
 {
 	public DatasmithWireTranslatorBase(ReadOnlyTargetRules Target) : base(Target)
 	{
+		// TODO: investigate to remove that (Jira UETOOL-4975)
+		bUseUnity = false;
+
 		PublicDependencyModuleNames.AddRange(
 			new string[]
 			{
