Fix import transform matrix

diff --git a/Engine/Plugins/Enterprise/DatasmithCADImporter/Source/CADInterfaces/Private/CADSceneGraph.cpp b/Engine/Plugins/Enterprise/DatasmithCADImporter/Source/CADInterfaces/Private/CADSceneGraph.cpp
--- a/Engine/Plugins/Enterprise/DatasmithCADImporter/Source/CADInterfaces/Private/CADSceneGraph.cpp
+++ b/Engine/Plugins/Enterprise/DatasmithCADImporter/Source/CADInterfaces/Private/CADSceneGraph.cpp
@@ -9,11 +9,17 @@
 namespace CADLibrary
 {
 
+FArchive& operator<<(FArchive& Ar, FCADArchiveObject& Object)
+{
+	Ar << Object.ObjectId;
+	Ar << Object.MetaData;
+	Ar << Object.TransformMatrix;
+	return Ar;
+}
+
 FArchive& operator<<(FArchive& Ar, FArchiveInstance& Instance) 
 {
-	Ar << Instance.ObjectId;
-	Ar << Instance.MetaData;
-	Ar << Instance.TransformMatrix;
+	Ar << (FCADArchiveObject&) Instance;
 	Ar << Instance.ReferenceNodeId;
 	Ar << Instance.bIsExternalReference;
 	Ar << Instance.ExternalReference;
@@ -22,23 +28,20 @@ FArchive& operator<<(FArchive& Ar, FArchiveInstance& Instance)
 
 FArchive& operator<<(FArchive& Ar, FArchiveComponent& Component)
 {
-	Ar << Component.ObjectId;
-	Ar << Component.MetaData;
+	Ar << (FCADArchiveObject&) Component;
 	Ar << Component.Children;
 	return Ar;
 }
 
 FArchive& operator<<(FArchive& Ar, FArchiveUnloadedComponent& Unloaded) 
 {
-	Ar << Unloaded.ObjectId;
-	Ar << Unloaded.MetaData;
+	Ar << (FArchiveComponent&) Unloaded;
 	return Ar;
 }
 
 FArchive& operator<<(FArchive& Ar, FArchiveBody& Body) 
 {
-	Ar << Body.ObjectId;
-	Ar << Body.MetaData;
+	Ar << (FCADArchiveObject&) Body;
 	Ar << Body.MaterialFaceSet;
 	Ar << Body.ColorFaceSet;
 	Ar << Body.MeshActorName;
diff --git a/Engine/Plugins/Enterprise/DatasmithCADImporter/Source/CADInterfaces/Private/CoreTechFileParser.h b/Engine/Plugins/Enterprise/DatasmithCADImporter/Source/CADInterfaces/Private/CoreTechFileParser.h
--- a/Engine/Plugins/Enterprise/DatasmithCADImporter/Source/CADInterfaces/Private/CoreTechFileParser.h
+++ b/Engine/Plugins/Enterprise/DatasmithCADImporter/Source/CADInterfaces/Private/CoreTechFileParser.h
@@ -82,7 +82,7 @@ namespace CADLibrary
 		uint32 GetMaterialNum();
 		void ReadMaterials();
 
-		uint32 GetObjectMaterial(ICADArchiveObject& Object);
+		uint32 GetObjectMaterial(FCADArchiveObject& Object);
 		FArchiveColor& FindOrAddColor(uint32 ColorHId);
 		FArchiveMaterial& FindOrAddMaterial(CT_MATERIAL_ID MaterialId);
 		void SetFaceMainMaterial(FObjectDisplayDataId& InFaceMaterial, FObjectDisplayDataId& InBodyMaterial, FBodyMesh& BodyMesh, int32 FaceIndex);
diff --git a/Engine/Plugins/Enterprise/DatasmithCADImporter/Source/CADInterfaces/Private/CoreTechFileParserImpl.cpp b/Engine/Plugins/Enterprise/DatasmithCADImporter/Source/CADInterfaces/Private/CoreTechFileParserImpl.cpp
--- a/Engine/Plugins/Enterprise/DatasmithCADImporter/Source/CADInterfaces/Private/CoreTechFileParserImpl.cpp
+++ b/Engine/Plugins/Enterprise/DatasmithCADImporter/Source/CADInterfaces/Private/CoreTechFileParserImpl.cpp
@@ -80,7 +80,7 @@ namespace CADLibrary
 		return NewColor;
 	}
 
-	uint32 FCoreTechFileParser::GetObjectMaterial(ICADArchiveObject& Object)
+	uint32 FCoreTechFileParser::GetObjectMaterial(FCADArchiveObject& Object)
 	{
 		if (FString* Material = Object.MetaData.Find(TEXT("MaterialName")))
 		{
diff --git a/Engine/Plugins/Enterprise/DatasmithCADImporter/Source/CADInterfaces/Private/TechSoftFileParser.cpp b/Engine/Plugins/Enterprise/DatasmithCADImporter/Source/CADInterfaces/Private/TechSoftFileParser.cpp
--- a/Engine/Plugins/Enterprise/DatasmithCADImporter/Source/CADInterfaces/Private/TechSoftFileParser.cpp
+++ b/Engine/Plugins/Enterprise/DatasmithCADImporter/Source/CADInterfaces/Private/TechSoftFileParser.cpp
@@ -308,6 +308,17 @@ ECADParsingResult FTechSoftFileParser::Process()
 		return ECADParsingResult::ProcessFailed;
 	}
 
+	{
+		TUniqueTSObj<A3DAsmModelFileData> ModelFileData(ModelFile.Get());
+		if (!ModelFileData.IsValid())
+		{
+			return ECADParsingResult::ProcessFailed;
+		}
+
+		ModellerType = (EModellerType)ModelFileData->m_eModellerType;
+		FileUnit = ModelFileData->m_dUnit;
+	}
+
 	// save the file for the next load
 	if (CADFileData.IsCacheDefined())
 	{
@@ -498,9 +509,6 @@ ECADParsingResult FTechSoftFileParser::TraverseModel()
 		return ECADParsingResult::ProcessFailed;
 	}
 
-	ModellerType = (EModellerType)ModelFileData->m_eModellerType;
-	FileUnit = ModelFileData->m_dUnit;
-
 	FEntityMetaData MetaData;
 	ExtractMetaData(ModelFile.Get(), MetaData);
 	ExtractSpecificMetaData(ModelFile.Get(), MetaData);
@@ -513,7 +521,7 @@ ECADParsingResult FTechSoftFileParser::TraverseModel()
 		}
 		else
 		{
-			TraverseReference(ModelFileData->m_ppPOccurrences[Index]);
+			TraverseReference(ModelFileData->m_ppPOccurrences[Index], FMatrix::Identity);
 		}
 	}
 
@@ -534,6 +542,13 @@ void FTechSoftFileParser::TraverseConfigurationSet(const A3DAsmProductOccurrence
 
 	const FString& ConfigurationToLoad = CADFileData.GetCADFileDescription().GetConfiguration();
 
+	FMatrix TransformMatrix = FMatrix::Identity;;
+	A3DMiscTransformation* Location = ConfigurationSetData->m_pLocation;
+	if (Location)
+	{
+		TransformMatrix = TraverseTransformation(Location);
+	}
+
 	TUniqueTSObj<A3DAsmProductOccurrenceData> ConfigurationData;
 	for (unsigned int Index = 0; Index < ConfigurationSetData->m_uiPOccurrencesSize; ++Index)
 	{
@@ -563,7 +578,7 @@ void FTechSoftFileParser::TraverseConfigurationSet(const A3DAsmProductOccurrence
 
 			if (bIsConfigurationToLoad)
 			{
-				TraverseReference(ConfigurationSetData->m_ppPOccurrences[Index]);
+				TraverseReference(ConfigurationSetData->m_ppPOccurrences[Index], TransformMatrix);
 				return;
 			}
 		}
@@ -582,7 +597,7 @@ void FTechSoftFileParser::TraverseConfigurationSet(const A3DAsmProductOccurrence
 
 			if (ConfigurationData->m_uiProductFlags & A3D_PRODUCT_FLAG_CONFIG)
 			{
-				TraverseReference(ConfigurationSetData->m_ppPOccurrences[Index]);
+				TraverseReference(ConfigurationSetData->m_ppPOccurrences[Index], TransformMatrix);
 			}
 		}
 	}
@@ -619,7 +634,7 @@ void FTechSoftFileParser::CountUnderConfigurationSet(const A3DAsmProductOccurren
 	}
 }
 
-void FTechSoftFileParser::TraverseReference(const A3DAsmProductOccurrence* ReferencePtr)
+void FTechSoftFileParser::TraverseReference(const A3DAsmProductOccurrence* ReferencePtr, const FMatrix& ParentMatrix)
 {
 	FEntityMetaData MetaData;
 	ExtractMetaData(ReferencePtr, MetaData);
@@ -645,6 +660,15 @@ void FTechSoftFileParser::TraverseReference(const A3DAsmProductOccurrence* Refer
 		return;
 	}
 
+	FMatrix ReferenceMatrix = FMatrix::Identity;;
+	A3DMiscTransformation* Location = ReferenceData->m_pLocation;
+	if (Location)
+	{
+		ReferenceMatrix = TraverseTransformation(Location);
+	}
+	
+	Component.TransformMatrix = ParentMatrix * ReferenceMatrix;
+
 	for (uint32 OccurenceIndex = 0; OccurenceIndex < ReferenceData->m_uiPOccurrencesSize; ++OccurenceIndex)
 	{
 		int32 ChildrenId = TraverseOccurrence(ReferenceData->m_ppPOccurrences[OccurenceIndex]);
@@ -720,11 +744,12 @@ FArchiveComponent& FTechSoftFileParser::AddOccurence(FEntityMetaData& InstanceMe
 	return Prototype;
 }
 
-int32 FTechSoftFileParser::AddBody(FEntityMetaData& BodyMetaData)
+int32 FTechSoftFileParser::AddBody(FEntityMetaData& BodyMetaData, const FMatrix& Matrix)
 {
 	FCadId BodyId = LastEntityId++;
 	int32 BodyIndex = CADFileData.AddBody(BodyId);
 	FArchiveBody& Body = CADFileData.GetBodyAt(BodyIndex);
+	Body.TransformMatrix = Matrix;
 	Body.MetaData = MoveTemp(BodyMetaData.MetaData);
 	if(BodyMetaData.ColorName != 0)
 	{
@@ -1131,7 +1156,14 @@ FCadId FTechSoftFileParser::TraverseBRepModel(A3DRiBrepModel* BRepModelPtr, FEnt
 	ExtractSpecificMetaData(BRepModelPtr, BRepMetaData);
 	ExtractMaterialProperties(BRepModelPtr);
 
-	int32 BodyIndex = AddBody(BRepMetaData);
+	FMatrix Matrix = FMatrix::Identity;
+	TUniqueTSObj<A3DRiRepresentationItemData> RepresentationData(BRepModelPtr);
+	if (RepresentationData->m_pCoordinateSystem)
+	{
+		Matrix = FTechSoftFileParser::TraverseCoordinateSystem(RepresentationData->m_pCoordinateSystem);
+	}
+
+	int32 BodyIndex = AddBody(BRepMetaData, Matrix);
 	FArchiveBody& Body = CADFileData.GetBodyAt(BodyIndex);
 
 	RepresentationItemsCache.Add(BRepModelPtr, BodyIndex);
@@ -1161,7 +1193,14 @@ FCadId FTechSoftFileParser::TraversePolyBRepModel(A3DRiPolyBrepModel* PolygonalP
 	ExtractSpecificMetaData(PolygonalPtr, BRepMetaData);
 	ExtractMaterialProperties(PolygonalPtr);
 
-	int32 BodyIndex = AddBody(BRepMetaData);
+	FMatrix Matrix = FMatrix::Identity;
+	TUniqueTSObj<A3DRiRepresentationItemData> RepresentationData(PolygonalPtr);
+	if (RepresentationData->m_pCoordinateSystem)
+	{
+		Matrix = FTechSoftFileParser::TraverseCoordinateSystem(RepresentationData->m_pCoordinateSystem);
+	}
+
+	int32 BodyIndex = AddBody(BRepMetaData, Matrix);
 	FArchiveBody& Body = CADFileData.GetBodyAt(BodyIndex);
 
 	RepresentationItemsCache.Add(PolygonalPtr, BodyIndex);
@@ -1658,7 +1697,7 @@ FMatrix FTechSoftFileParser::TraverseGeneralTransformation(const A3DMiscTransfor
 	TUniqueTSObj<A3DMiscGeneralTransformationData> GeneralTransformationData(GeneralTransformation);
 	if (GeneralTransformationData.IsValid())
 	{
-		FMatrix Matrix;
+		FMatrix Matrix = FMatrix::Identity;;
 		int32 Index = 0;
 		for (int32 Andex = 0; Andex < 4; ++Andex)
 		{
diff --git a/Engine/Plugins/Enterprise/DatasmithCADImporter/Source/CADInterfaces/Private/TechSoftFileParser.h b/Engine/Plugins/Enterprise/DatasmithCADImporter/Source/CADInterfaces/Private/TechSoftFileParser.h
--- a/Engine/Plugins/Enterprise/DatasmithCADImporter/Source/CADInterfaces/Private/TechSoftFileParser.h
+++ b/Engine/Plugins/Enterprise/DatasmithCADImporter/Source/CADInterfaces/Private/TechSoftFileParser.h
@@ -155,7 +155,7 @@ private:
 
 	// Traverse ASM tree by starting from the model
 	ECADParsingResult TraverseModel();
-	void TraverseReference(const A3DAsmProductOccurrence* Reference);
+	void TraverseReference(const A3DAsmProductOccurrence* Reference, const FMatrix& ParentMatrix);
 	bool IsConfigurationSet(const A3DAsmProductOccurrence* Occurrence);
 	void TraverseConfigurationSet(const A3DAsmProductOccurrence* ConfigurationSet);
 	FCadId TraverseOccurrence(const A3DAsmProductOccurrence* Occurrence);
@@ -212,7 +212,7 @@ private:
 	FArchiveUnloadedComponent& AddUnloadedComponent(FEntityMetaData& ComponentMetaData, FArchiveInstance& Instance);
 	FArchiveComponent& AddOccurence(FEntityMetaData& InstanceMetaData, FEntityMetaData& ReferenceMetaData, FCadId& OutComponentId);
 	FArchiveComponent& AddOccurence(FEntityMetaData& InstanceMetaData, FCadId& OutComponentId);
-	int32 AddBody(FEntityMetaData& BodyMetaData);
+	int32 AddBody(FEntityMetaData& BodyMetaData, const FMatrix& Matrix);
 #endif
 
 private:
diff --git a/Engine/Plugins/Enterprise/DatasmithCADImporter/Source/CADInterfaces/Public/CADSceneGraph.h b/Engine/Plugins/Enterprise/DatasmithCADImporter/Source/CADInterfaces/Public/CADSceneGraph.h
--- a/Engine/Plugins/Enterprise/DatasmithCADImporter/Source/CADInterfaces/Public/CADSceneGraph.h
+++ b/Engine/Plugins/Enterprise/DatasmithCADImporter/Source/CADInterfaces/Public/CADSceneGraph.h
@@ -10,43 +10,45 @@ class FArchive;
 namespace CADLibrary
 {
 
-class CADINTERFACES_API ICADArchiveObject
+class CADINTERFACES_API FCADArchiveObject
 {
 public:
-	ICADArchiveObject(FCadId Id = 0)
+	FCADArchiveObject(FCadId Id = 0)
 		: ObjectId(Id)
 	{
 	}
 
-	virtual ~ICADArchiveObject() = default;
+	virtual ~FCADArchiveObject() = default;
+
+	friend FArchive& operator<<(FArchive& Ar, FCADArchiveObject& C);
 
 public:
 	uint32 ObjectId;
 	TMap<FString, FString> MetaData;
+	FMatrix TransformMatrix = FMatrix::Identity;
 };
 
-class CADINTERFACES_API FArchiveInstance : public ICADArchiveObject
+class CADINTERFACES_API FArchiveInstance : public FCADArchiveObject
 {
 public:
 	FArchiveInstance(FCadId Id = 0)
-		: ICADArchiveObject(Id)
+		: FCADArchiveObject(Id)
 	{
 	}
 
 	friend FArchive& operator<<(FArchive& Ar, FArchiveInstance& C);
 
 public:
-	FMatrix TransformMatrix = FMatrix::Identity;
 	FCadId ReferenceNodeId = 0;
 	bool bIsExternalReference = false;
 	FFileDescriptor ExternalReference;
 };
 
-class CADINTERFACES_API FArchiveComponent : public ICADArchiveObject
+class CADINTERFACES_API FArchiveComponent : public FCADArchiveObject
 {
 public:
 	FArchiveComponent(FCadId Id = 0)
-		: ICADArchiveObject(Id)
+		: FCADArchiveObject(Id)
 	{
 	}
 
@@ -67,11 +69,11 @@ public:
 	friend FArchive& operator<<(FArchive& Ar, FArchiveUnloadedComponent& C);
 };
 
-class CADINTERFACES_API FArchiveBody : public ICADArchiveObject
+class CADINTERFACES_API FArchiveBody : public FCADArchiveObject
 {
 public:
 	FArchiveBody(FCadId Id = 0)
-		: ICADArchiveObject(Id)
+		: FCADArchiveObject(Id)
 	{
 	}
 
diff --git a/Engine/Plugins/Enterprise/DatasmithCADImporter/Source/DatasmithCADTranslator/Private/DatasmithSceneGraphBuilder.cpp b/Engine/Plugins/Enterprise/DatasmithCADImporter/Source/DatasmithCADTranslator/Private/DatasmithSceneGraphBuilder.cpp
--- a/Engine/Plugins/Enterprise/DatasmithCADImporter/Source/DatasmithCADTranslator/Private/DatasmithSceneGraphBuilder.cpp
+++ b/Engine/Plugins/Enterprise/DatasmithCADImporter/Source/DatasmithCADTranslator/Private/DatasmithSceneGraphBuilder.cpp
@@ -31,14 +31,14 @@ namespace DatasmithSceneGraphBuilderImpl
 		}
 	}
 
-	void AddTransformToActor(const CADLibrary::FArchiveInstance& Instance, TSharedPtr< IDatasmithActorElement > Actor, const CADLibrary::FImportParameters& ImportParameters)
+	void AddTransformToActor(const CADLibrary::FCADArchiveObject& Object, TSharedPtr< IDatasmithActorElement > Actor, const CADLibrary::FImportParameters& ImportParameters)
 	{
 		if (!Actor.IsValid())
 		{
 			return;
 		}
 
-		FTransform LocalTransform(Instance.TransformMatrix);
+		FTransform LocalTransform(Object.TransformMatrix);
 		FTransform LocalUETransform = FDatasmithUtils::ConvertTransform(ImportParameters.GetModelCoordSys(), LocalTransform);
 
 		Actor->SetTranslation(LocalUETransform.GetTranslation() * ImportParameters.GetScaleFactor());
@@ -413,7 +413,6 @@ void FDatasmithSceneBaseGraphBuilder::GetNodeUUIDAndName(
 	{
 		OutName = "NoName";
 	}
-	DatasmithSceneGraphBuilderImpl::CleanName(OutName);
 
 	FCADUUID UEUUID = HashCombine(GetTypeHash(InParentUEUUID), GetTypeHash(InComponentIndex));
 
@@ -468,6 +467,8 @@ TSharedPtr<IDatasmithActorElement> FDatasmithSceneBaseGraphBuilder::BuildCompone
 
 	AddChildren(Actor, Component, ComponentData);
 
+	DatasmithSceneGraphBuilderImpl::AddTransformToActor(Component, Actor, ImportParameters);
+
 	return Actor;
 }
 
@@ -506,6 +507,8 @@ TSharedPtr<IDatasmithActorElement> FDatasmithSceneBaseGraphBuilder::BuildBody(in
 	ActorElement->SetLabel(*BodyLabel);
 	ActorElement->SetStaticMeshPathName(MeshElement->GetName());
 
+	DatasmithSceneGraphBuilderImpl::AddTransformToActor(Body, ActorElement, ImportParameters);
+
 	if (MaterialUuid && ImportParameters.GetPropagation() != CADLibrary::EDisplayDataPropagationMode::BodyOnly)
 	{
 		TSharedPtr< IDatasmithMaterialIDElement > PartMaterialIDElement = FindOrAddMaterial(MaterialUuid);
