Improve MeshOperator::OrientMesh

In case of non-manifold body, the orientation can be badly propagate due to a local moebius strip.
The main idea of the new algorithm is to propagate a front of faces oriented in the same direction and to define the boundary of faces oriented differently.
When the first front is ended, starting from the defined boundary, the other faces are swapped.

diff --git a/Engine/Plugins/Enterprise/DatasmithCADImporter/Source/CADLibrary/Private/MeshOperator.cpp b/Engine/Plugins/Enterprise/DatasmithCADImporter/Source/CADLibrary/Private/MeshOperator.cpp
--- a/Engine/Plugins/Enterprise/DatasmithCADImporter/Source/CADLibrary/Private/MeshOperator.cpp
+++ b/Engine/Plugins/Enterprise/DatasmithCADImporter/Source/CADLibrary/Private/MeshOperator.cpp
@@ -17,12 +17,13 @@ bool MeshOperator::OrientMesh(FMeshDescription& MeshDescription)
 {
 	FMeshEditingWrapper MeshWrapper(MeshDescription);
 
-	TQueue<FTriangleID> Front; 
+	TQueue<FTriangleID> Front;
+	TQueue<FTriangleID> BadOrientationFront;
+
 	FTriangleID AdjacentTriangle;
 	FEdgeID Edge;
 
 	TArrayView < const FVertexInstanceID> Vertices;
-	int32 NbSwappedTriangles = 0;
 
 	FVector MaxCorner;
 	FVector MinCorner;
@@ -65,42 +66,92 @@ bool MeshOperator::OrientMesh(FMeshDescription& MeshDescription)
 		int32 NbConnectedFaces = 1;
 		int32 NbBorderEdges = 0;
 		int32 NbSurfaceEdges = 0;
+		int32 NbSwappedTriangles = 0;
 		while (!Front.IsEmpty())
 		{
-			Front.Dequeue(Triangle);
-
-			TArrayView<const FEdgeID> EdgeSet = MeshDescription.GetTriangleEdges(Triangle);
-
-			for (int32 IEdge = 0; IEdge < 3; IEdge++)
+			while (!Front.IsEmpty())
 			{
-				Edge = EdgeSet[IEdge];
+				Front.Dequeue(Triangle);
 
-				if (!MeshWrapper.IsEdgeOfCategory(Edge, EElementCategory::ElementCategorySurface))
-				{
-					NbBorderEdges++;
-					continue;
-				}
+				TArrayView<const FEdgeID> EdgeSet = MeshDescription.GetTriangleEdges(Triangle);
 
-				AdjacentTriangle = MeshWrapper.GetOtherTriangleAtEdge(Edge, Triangle);
-				if (MeshWrapper.IsTriangleMarked(AdjacentTriangle))
+				for (int32 IEdge = 0; IEdge < 3; IEdge++)
 				{
-					continue;
+					Edge = EdgeSet[IEdge];
+
+					if (!MeshWrapper.IsEdgeOfCategory(Edge, EElementCategory::ElementCategorySurface))
+					{
+						NbBorderEdges++;
+						continue;
+					}
+
+					AdjacentTriangle = MeshWrapper.GetOtherTriangleAtEdge(Edge, Triangle);
+					if (MeshWrapper.IsTriangleMarked(AdjacentTriangle))
+					{
+						continue;
+					}
+
+					NbSurfaceEdges++;
+					NbConnectedFaces++;
+
+					ConnectedTriangles.Add(AdjacentTriangle);
+
+					MeshWrapper.SetTriangleMarked(AdjacentTriangle);
+					MeshWrapper.GetTriangleBoundingBox(AdjacentTriangle, MinCorner, MaxCorner, HighestVertex, LowestVertex);
+
+					if (MeshWrapper.GetEdgeDirectionInTriangle(Edge, 0) == MeshWrapper.GetEdgeDirectionInTriangle(Edge, 1))
+					{
+						BadOrientationFront.Enqueue(AdjacentTriangle);
+					}
+					else
+					{
+						Front.Enqueue(AdjacentTriangle);
+					}
 				}
+			}
 
-				NbSurfaceEdges++;
-				NbConnectedFaces++;
+			while (!BadOrientationFront.IsEmpty())
+			{
+				BadOrientationFront.Dequeue(Triangle);
+				MeshWrapper.SwapTriangleOrientation(Triangle);
+				NbSwappedTriangles++;
 
-				ConnectedTriangles.Add(AdjacentTriangle);
+				TArrayView<const FEdgeID> EdgeSet = MeshDescription.GetTriangleEdges(Triangle);
 
-				if(MeshWrapper.GetEdgeDirectionInTriangle(Edge, 0) == MeshWrapper.GetEdgeDirectionInTriangle(Edge, 1))
+				for (int32 IEdge = 0; IEdge < 3; IEdge++)
 				{
-					MeshWrapper.SwapTriangleOrientation(AdjacentTriangle);
-					NbSwappedTriangles++;
-				}
-				MeshWrapper.SetTriangleMarked(AdjacentTriangle);
-				MeshWrapper.GetTriangleBoundingBox(AdjacentTriangle, MinCorner, MaxCorner, HighestVertex, LowestVertex);
+					Edge = EdgeSet[IEdge];
+
+					if (!MeshWrapper.IsEdgeOfCategory(Edge, EElementCategory::ElementCategorySurface))
+					{
+						NbBorderEdges++;
+						continue;
+					}
+
+					AdjacentTriangle = MeshWrapper.GetOtherTriangleAtEdge(Edge, Triangle);
+					if (MeshWrapper.IsTriangleMarked(AdjacentTriangle))
+					{
+						continue;
+					}
+
+					NbSurfaceEdges++;
+					NbConnectedFaces++;
+
+					ConnectedTriangles.Add(AdjacentTriangle);
 
-				Front.Enqueue(AdjacentTriangle);
+					MeshWrapper.SetTriangleMarked(AdjacentTriangle);
+					MeshWrapper.GetTriangleBoundingBox(AdjacentTriangle, MinCorner, MaxCorner, HighestVertex, LowestVertex);
+					if (MeshWrapper.GetEdgeDirectionInTriangle(Edge, 0) == MeshWrapper.GetEdgeDirectionInTriangle(Edge, 1))
+					{
+						BadOrientationFront.Enqueue(AdjacentTriangle);
+						MeshWrapper.SwapTriangleOrientation(AdjacentTriangle);
+						NbSwappedTriangles++;
+					}
+					else
+					{
+						Front.Enqueue(AdjacentTriangle);
+					}
+				}
 			}
 		}
 
@@ -117,7 +168,7 @@ bool MeshOperator::OrientMesh(FMeshDescription& MeshDescription)
 
 			if (HighestVertex[0] != INDEX_NONE)
 			{
-				FStaticMeshConstAttributes StaticMeshAttributes( MeshDescription );
+				FStaticMeshConstAttributes StaticMeshAttributes(MeshDescription);
 				TVertexInstanceAttributesConstRef<FVector3f> Normals = StaticMeshAttributes.GetVertexInstanceNormals();
 
 				for (int32 VertexIndex = 0; VertexIndex < 3; VertexIndex++)
@@ -129,7 +180,7 @@ bool MeshOperator::OrientMesh(FMeshDescription& MeshDescription)
 						float MaxComponent = 0;
 						for (const FVertexInstanceID VertexInstanceID : CoincidentVertexInstanceIdSet)
 						{
-							FVector Normal = (FVector)Normals[ VertexInstanceID ];
+							FVector Normal = (FVector)Normals[VertexInstanceID];
 							if (FMath::Abs(MaxComponent) < FMath::Abs(Normal[VertexIndex]))
 							{
 								MaxComponent = Normal[VertexIndex];
@@ -153,7 +204,7 @@ bool MeshOperator::OrientMesh(FMeshDescription& MeshDescription)
 						float MaxComponent = 0;
 						for (const FVertexInstanceID VertexInstanceID : CoincidentVertexInstanceIdSet)
 						{
-							FVector Normal = (FVector)Normals[ VertexInstanceID ];
+							FVector Normal = (FVector)Normals[VertexInstanceID];
 							if (FMath::Abs(MaxComponent) < FMath::Abs(Normal[VertexIndex]))
 							{
 								MaxComponent = Normal[VertexIndex];
@@ -172,14 +223,14 @@ bool MeshOperator::OrientMesh(FMeshDescription& MeshDescription)
 				}
 			}
 		}
-		else if (NbSwappedTriangles*2 > NbConnectedFaces)
+		else if (NbSwappedTriangles * 2 > NbConnectedFaces)
 		{
 			// case of surface mesh
 			// this means that more triangles of surface shape have been swapped than no swapped, the good orientation has been reversed.
 			// The mesh need to be re swapped
 			NbInverted++;
 		}
-		
+
 		// if needed swap all the mesh
 		if (NbInverted > NbNotInverted)
 		{
