IGES: force sew at the import

diff --git a/Engine/Plugins/Enterprise/DatasmithCADImporter/Source/CADInterfaces/Private/TechSoftFileParser.cpp b/Engine/Plugins/Enterprise/DatasmithCADImporter/Source/CADInterfaces/Private/TechSoftFileParser.cpp
--- a/Engine/Plugins/Enterprise/DatasmithCADImporter/Source/CADInterfaces/Private/TechSoftFileParser.cpp
+++ b/Engine/Plugins/Enterprise/DatasmithCADImporter/Source/CADInterfaces/Private/TechSoftFileParser.cpp
@@ -233,12 +233,13 @@ void SetIOOption(A3DImport& Importer)
 	Importer.m_sLoadData.m_sIncremental.m_bLoadStructureOnly = false;
 }
 
-void UpdateIOOptionAccordingToFormat(const CADLibrary::ECADFormat Format, A3DImport& Importer)
+void UpdateIOOptionAccordingToFormat(const CADLibrary::ECADFormat Format, A3DImport& Importer, bool& OutForceSew)
 {
 	switch (Format)
 	{
 	case CADLibrary::ECADFormat::IGES:
 	{
+		OutForceSew = true;
 		break;
 	}
 
@@ -299,12 +300,46 @@ ECADParsingResult FTechSoftFileParser::Process()
 
 	// Add specific options according to format
 	Format = File.GetFileFormat();
-	TechSoftFileParserImpl::UpdateIOOptionAccordingToFormat(Format, Import);
+	bool bForceSew = false;
+	TechSoftFileParserImpl::UpdateIOOptionAccordingToFormat(Format, Import, bForceSew);
 
-	ModelFile = TechSoftInterface::LoadModelFileFromFile(Import);
+	A3DStatus LoadStatus = A3DStatus::A3D_SUCCESS;
+	ModelFile = TechSoftInterface::LoadModelFileFromFile(Import, LoadStatus);
 
 	if (!ModelFile.IsValid())
 	{
+		switch (LoadStatus)
+		{
+		case A3DStatus::A3D_LOAD_FILE_TOO_OLD:
+		{
+			CADFileData.AddWarningMessages(FString::Printf(TEXT("File %s hasn't been loaded because the version is less than the oldest supported version."), *File.GetFileName()));
+			break;
+		}
+
+		case A3DStatus::A3D_LOAD_FILE_TOO_RECENT:
+		{
+			CADFileData.AddWarningMessages(FString::Printf(TEXT("File %s hasn't been loaded because the version is more recent than supported version."), *File.GetFileName()));
+			break;
+		}
+
+		case A3DStatus::A3D_LOAD_CANNOT_ACCESS_CADFILE: 
+		{
+			CADFileData.AddWarningMessages(FString::Printf(TEXT("File %s hasn't been loaded because the input path cannot be opened by the running process for reading."), *File.GetFileName()));
+			break;
+		}
+
+		case A3DStatus::A3D_LOAD_INVALID_FILE_FORMAT:
+		{
+			CADFileData.AddWarningMessages(FString::Printf(TEXT("File %s hasn't been loaded because the format is not supported."), *File.GetFileName()));
+			break;
+		}
+
+		default:
+		{
+			CADFileData.AddWarningMessages(FString::Printf(TEXT("File %s hasn't been loaded because an error occured while reading the file."), *File.GetFileName()));
+			break;
+		}
+		};
 		return ECADParsingResult::ProcessFailed;
 	}
 
@@ -325,17 +360,17 @@ ECADParsingResult FTechSoftFileParser::Process()
 		FString CacheFilePath = CADFileData.GetCADCachePath();
 		if (CacheFilePath != File.GetPathOfFileToLoad())
 		{
-			TechSoftUtils::SaveModelFileToPcrFile(ModelFile.Get(), CacheFilePath);
+			TechSoftUtils::SaveModelFileToPrcFile(ModelFile.Get(), CacheFilePath);
 		}
 	}
 
-	if (CADFileData.GetImportParameters().GetStitchingTechnique() == StitchingSew && FImportParameters::bGDisableCADKernelTessellation)
+	if (bForceSew || CADFileData.GetImportParameters().GetStitchingTechnique() == StitchingSew && FImportParameters::bGDisableCADKernelTessellation)
 	{
 		CADLibrary::TUniqueTSObj<A3DSewOptionsData> SewData;
 		SewData->m_bComputePreferredOpenShellOrientation = false;
 		double ToleranceMM = 0.01 / FileUnit;
 
-		A3DStatus Status = TechSoftInterface::SewModel(ModelFile.GetPtr(), ToleranceMM, SewData.GetPtr());
+		A3DStatus Status = TechSoftInterface::SewModel(ModelFile.Get(), ToleranceMM, SewData.GetPtr());
 		if(Status != A3DStatus::A3D_SUCCESS)
 		{
 			// To do but what ?
@@ -453,7 +488,7 @@ void FTechSoftFileParser::GenerateBodyMeshes()
 
 					FJsonSerializer::Serialize(JsonObject.ToSharedRef(), JsonWriter);
 
-					TechSoftUtils::SaveBodiesToPcrFile(&RepresentationItemPtr, 1, FilePath, JsonString);
+					TechSoftUtils::SaveBodiesToPrcFile(&RepresentationItemPtr, 1, FilePath, JsonString);
 				}
 			}
 		}
diff --git a/Engine/Plugins/Enterprise/DatasmithCADImporter/Source/CADInterfaces/Private/TechSoftInterface.cpp b/Engine/Plugins/Enterprise/DatasmithCADImporter/Source/CADInterfaces/Private/TechSoftInterface.cpp
--- a/Engine/Plugins/Enterprise/DatasmithCADImporter/Source/CADInterfaces/Private/TechSoftInterface.cpp
+++ b/Engine/Plugins/Enterprise/DatasmithCADImporter/Source/CADInterfaces/Private/TechSoftInterface.cpp
@@ -85,14 +85,22 @@ FString GetTechSoftVersion()
 
 #ifdef USE_TECHSOFT_SDK
 
-FUniqueTechSoftModelFile LoadModelFileFromFile(const A3DImport& Import)
+FUniqueTechSoftModelFile LoadModelFileFromFile(const A3DImport& Import, A3DStatus& Status)
 {
 	A3DAsmModelFile* ModelFile = nullptr;
-	if (A3DAsmModelFileLoadFromFile(Import.GetFilePath(), &Import.m_sLoadData, &ModelFile) != A3DStatus::A3D_SUCCESS)
+	Status = A3DAsmModelFileLoadFromFile(Import.GetFilePath(), &Import.m_sLoadData, &ModelFile);
+
+	switch (Status)
 	{
-		return FUniqueTechSoftModelFile();
+	case A3D_LOAD_MULTI_MODELS_CADFILE: //if the file contains multiple entries (see A3DRWParamsMultiEntriesData).
+	case A3D_LOAD_MISSING_COMPONENTS: //_[I don't know about this one]_
+	case A3D_SUCCESS:
+		return FUniqueTechSoftModelFile(ModelFile);
+
+	default:
+		break;
 	}
-	return FUniqueTechSoftModelFile(ModelFile);
+	return FUniqueTechSoftModelFile();
 }
 
 FUniqueTechSoftModelFile LoadModelFileFromPrcFile(const A3DUTF8Char* CADFileName, A3DRWParamsPrcReadHelper** ReadHelper)
diff --git a/Engine/Plugins/Enterprise/DatasmithCADImporter/Source/CADInterfaces/Private/TechSoftUtils.cpp b/Engine/Plugins/Enterprise/DatasmithCADImporter/Source/CADInterfaces/Private/TechSoftUtils.cpp
--- a/Engine/Plugins/Enterprise/DatasmithCADImporter/Source/CADInterfaces/Private/TechSoftUtils.cpp
+++ b/Engine/Plugins/Enterprise/DatasmithCADImporter/Source/CADInterfaces/Private/TechSoftUtils.cpp
@@ -26,7 +26,7 @@ namespace TechSoftUtils
 #ifdef USE_TECHSOFT_SDK
 TSharedPtr<FJsonObject> GetJsonObject(A3DAsmProductOccurrence* ProductOcccurence);
 void RestoreMaterials(const TSharedPtr<FJsonObject>& DefaultValues, CADLibrary::FBodyMesh& BodyMesh);
-void SaveModelFileToPcrFile(void* ModelFile, const FString& Filename);
+void SaveModelFileToPrcFile(void* ModelFile, const FString& Filename);
 A3DUns32 CreateRGBColor(FColor& Color);
 void SetRootOccurenceAttributes(A3DEntity* Entity);
 #endif
@@ -81,7 +81,7 @@ bool GetBodyFromPcrFile(const FString& Filename, const FImportParameters& Import
 	return bExtractionSuccessful;
 }
 
-FUniqueTechSoftModelFile SaveBodiesToPcrFile(void** Bodies, uint32 BodyCount, const FString& Filename, const FString& JsonString)
+FUniqueTechSoftModelFile SaveBodiesToPrcFile(void** Bodies, uint32 BodyCount, const FString& Filename, const FString& JsonString)
 {
 #if defined USE_TECHSOFT_SDK && !defined CADKERNEL_DEV
 	if (!Bodies)
@@ -129,7 +129,7 @@ FUniqueTechSoftModelFile SaveBodiesToPcrFile(void** Bodies, uint32 BodyCount, co
 	FUniqueTechSoftModelFile ModelFile = TechSoftInterface::CreateAsmModelFile(*ModelFileData);
 
 	// Save ModelFile to Pcr file
-	SaveModelFileToPcrFile(ModelFile.Get(), Filename);
+	SaveModelFileToPrcFile(ModelFile.Get(), Filename);
 
 	// #ueent_techsoft: Deleting the model seems to delete the entire content. To be double-checked
 	TechSoftInterface::DeleteEntity(Attributes);
@@ -326,7 +326,7 @@ FColor GetColorAt(uint32 ColorIndex)
 	return FColor(200, 200, 200);
 }
 
-void SaveModelFileToPcrFile(void* ModelFile, const FString& Filename)
+void SaveModelFileToPrcFile(void* ModelFile, const FString& Filename)
 {
 #if defined USE_TECHSOFT_SDK && !defined CADKERNEL_DEV
 	TUniqueTSObj<A3DRWParamsExportPrcData> ParamsExportData;
diff --git a/Engine/Plugins/Enterprise/DatasmithCADImporter/Source/CADInterfaces/Public/TechSoftInterface.h b/Engine/Plugins/Enterprise/DatasmithCADImporter/Source/CADInterfaces/Public/TechSoftInterface.h
--- a/Engine/Plugins/Enterprise/DatasmithCADImporter/Source/CADInterfaces/Public/TechSoftInterface.h
+++ b/Engine/Plugins/Enterprise/DatasmithCADImporter/Source/CADInterfaces/Public/TechSoftInterface.h
@@ -83,7 +83,7 @@ CADINTERFACES_API FString GetTechSoftVersion();
 
 #ifdef USE_TECHSOFT_SDK
 
-CADINTERFACES_API FUniqueTechSoftModelFile LoadModelFileFromFile(const A3DImport& Import);
+CADINTERFACES_API FUniqueTechSoftModelFile LoadModelFileFromFile(const A3DImport& Import, A3DStatus& Status);
 CADINTERFACES_API FUniqueTechSoftModelFile LoadModelFileFromPrcFile(const A3DUTF8Char* CADFileName, A3DRWParamsPrcReadHelper** ReadHelper);
 
 CADINTERFACES_API A3DStatus DeleteModelFile(A3DAsmModelFile* ModelFile);
@@ -198,11 +198,6 @@ public:
 		return ModelFile;
 	}
 
-	A3DAsmModelFile** GetPtr()
-	{
-		return &ModelFile;
-	}
-
 	void Reset()
 	{
 		if (ModelFile != nullptr)
diff --git a/Engine/Plugins/Enterprise/DatasmithCADImporter/Source/CADInterfaces/Public/TechSoftUtils.h b/Engine/Plugins/Enterprise/DatasmithCADImporter/Source/CADInterfaces/Public/TechSoftUtils.h
--- a/Engine/Plugins/Enterprise/DatasmithCADImporter/Source/CADInterfaces/Public/TechSoftUtils.h
+++ b/Engine/Plugins/Enterprise/DatasmithCADImporter/Source/CADInterfaces/Public/TechSoftUtils.h
@@ -38,8 +38,8 @@ namespace TechSoftUtils
 CADINTERFACES_API bool GetBodyFromPcrFile(const FString& Filename, const FImportParameters& ImportParameters, FBodyMesh& BodyMesh);
 CADINTERFACES_API bool FillBodyMesh(void* BodyPtr, const FImportParameters& ImportParameters, double FileUnit, FBodyMesh& BodyMesh);
 
-CADINTERFACES_API FUniqueTechSoftModelFile SaveBodiesToPcrFile(void** Bodies, uint32 BodyCount, const FString& Filename, const FString& JsonString);
-CADINTERFACES_API void SaveModelFileToPcrFile(void* ModelFile, const FString& Filename);
+CADINTERFACES_API FUniqueTechSoftModelFile SaveBodiesToPrcFile(void** Bodies, uint32 BodyCount, const FString& Filename, const FString& JsonString);
+CADINTERFACES_API void SaveModelFileToPrcFile(void* ModelFile, const FString& Filename);
 
 CADINTERFACES_API int32 SetEntityGraphicsColor(A3DEntity* InEntity, FColor Color);
 
diff --git a/Engine/Plugins/Enterprise/DatasmithCADImporter/Source/CADTools/Public/CADOptions.h b/Engine/Plugins/Enterprise/DatasmithCADImporter/Source/CADTools/Public/CADOptions.h
--- a/Engine/Plugins/Enterprise/DatasmithCADImporter/Source/CADTools/Public/CADOptions.h
+++ b/Engine/Plugins/Enterprise/DatasmithCADImporter/Source/CADTools/Public/CADOptions.h
@@ -214,7 +214,7 @@ namespace CADLibrary
 
 		if (FImportParameters::bGDisableCADKernelTessellation)
 		{
-			OutFileName += FImportParameters::GCADLibrary.Equals("TechSoft") ? TEXT(".hsf") : TEXT(".ct");
+			OutFileName += FImportParameters::GCADLibrary.Equals("TechSoft") ? TEXT(".prc") : TEXT(".ct");
 		}
 		else
 		{
diff --git a/Engine/Plugins/Enterprise/DatasmithCADImporter/Source/ParametricSurface/Private/TechSoft/CADModelToTechSoftConverterBase.cpp b/Engine/Plugins/Enterprise/DatasmithCADImporter/Source/ParametricSurface/Private/TechSoft/CADModelToTechSoftConverterBase.cpp
--- a/Engine/Plugins/Enterprise/DatasmithCADImporter/Source/ParametricSurface/Private/TechSoft/CADModelToTechSoftConverterBase.cpp
+++ b/Engine/Plugins/Enterprise/DatasmithCADImporter/Source/ParametricSurface/Private/TechSoft/CADModelToTechSoftConverterBase.cpp
@@ -66,7 +66,7 @@ bool FCADModelToTechSoftConverterBase::SaveModel(const TCHAR* InFolderPath, TSha
 		FJsonSerializer::Serialize(JsonObject.ToSharedRef(), JsonWriter);
 	}
 
-	ModelFile = CADLibrary::TechSoftUtils::SaveBodiesToPcrFile(RiRepresentationItems.GetData(), RiRepresentationItems.Num(), FilePath, JsonString);
+	ModelFile = CADLibrary::TechSoftUtils::SaveBodiesToPrcFile(RiRepresentationItems.GetData(), RiRepresentationItems.Num(), FilePath, JsonString);
 
 	MeshElement->SetFile(*FilePath);
 #endif
