Fix MeshOperator::OrientMesh

diff --git a/Engine/Plugins/Enterprise/DatasmithCADImporter/Source/CADLibrary/Private/MeshOperator.cpp b/Engine/Plugins/Enterprise/DatasmithCADImporter/Source/CADLibrary/Private/MeshOperator.cpp
--- a/Engine/Plugins/Enterprise/DatasmithCADImporter/Source/CADLibrary/Private/MeshOperator.cpp
+++ b/Engine/Plugins/Enterprise/DatasmithCADImporter/Source/CADLibrary/Private/MeshOperator.cpp
@@ -101,6 +101,8 @@ bool MeshOperator::OrientMesh(FMeshDescription& MeshDescription)
 
 					if (MeshWrapper.GetEdgeDirectionInTriangle(Edge, 0) == MeshWrapper.GetEdgeDirectionInTriangle(Edge, 1))
 					{
+						MeshWrapper.SwapTriangleOrientation(AdjacentTriangle);
+						NbSwappedTriangles++;
 						BadOrientationFront.Enqueue(AdjacentTriangle);
 					}
 					else
@@ -113,8 +115,6 @@ bool MeshOperator::OrientMesh(FMeshDescription& MeshDescription)
 			while (!BadOrientationFront.IsEmpty())
 			{
 				BadOrientationFront.Dequeue(Triangle);
-				MeshWrapper.SwapTriangleOrientation(Triangle);
-				NbSwappedTriangles++;
 
 				TArrayView<const FEdgeID> EdgeSet = MeshDescription.GetTriangleEdges(Triangle);
 
