[TechSoft ] Improve DWG import

diff --git a/Engine/Plugins/Enterprise/DatasmithCADImporter/Source/CADInterfaces/Private/TechSoftFileParser.cpp b/Engine/Plugins/Enterprise/DatasmithCADImporter/Source/CADInterfaces/Private/TechSoftFileParser.cpp
--- a/Engine/Plugins/Enterprise/DatasmithCADImporter/Source/CADInterfaces/Private/TechSoftFileParser.cpp
+++ b/Engine/Plugins/Enterprise/DatasmithCADImporter/Source/CADInterfaces/Private/TechSoftFileParser.cpp
@@ -850,6 +850,27 @@ FCadId FTechSoftFileParser::TraverseOccurrence(const A3DAsmProductOccurrence* Oc
 		return Instance.ObjectId;
 	}
 	
+	FArchiveComponent& Component = AddComponent(InstanceMetaData, Instance);
+
+	// Add External data
+	if (OccurrenceData->m_pExternalData != nullptr)
+	{
+		TUniqueTSObj<A3DAsmProductOccurrenceData> ExternalData(OccurrenceData->m_pExternalData);
+		if (ExternalData->m_pPart != nullptr)
+		{
+			A3DAsmPartDefinition* PartDefinition = ExternalData->m_pPart;
+			TraversePartDefinition(PartDefinition, Component);
+		}
+
+		uint32 ExternalChildrenCount = ExternalData->m_uiPOccurrencesSize;
+		A3DAsmProductOccurrence** ExternalChildren = ExternalData->m_ppPOccurrences;
+		for (uint32 Index = 0; Index < ExternalChildrenCount; ++Index)
+		{
+			int32 ChildrenId = TraverseOccurrence(ExternalChildren[Index]);
+			Component.Children.Add(ChildrenId);
+		}
+	}
+
 	while (OccurrenceData->m_pPrototype != nullptr && OccurrenceData->m_pPart == nullptr && OccurrenceData->m_uiPOccurrencesSize == 0)
 	{
 		CachedOccurrencePtr = OccurrenceData->m_pPrototype;
@@ -861,8 +882,6 @@ FCadId FTechSoftFileParser::TraverseOccurrence(const A3DAsmProductOccurrence* Oc
 		return Instance.ObjectId;
 	}
 
-	FArchiveComponent& Component = AddComponent(InstanceMetaData, Instance);
-
 	// Add part
 	while (OccurrenceData->m_pPrototype != nullptr && OccurrenceData->m_pPart == nullptr)
 	{
@@ -900,6 +919,11 @@ void FTechSoftFileParser::CountUnderOccurrence(const A3DAsmProductOccurrence* Oc
 		ComponentCount[EComponentType::Occurrence]++;
 		ComponentCount[EComponentType::Reference]++;
 
+		if (OccurrenceData->m_pExternalData != nullptr)
+		{
+			CountUnderOccurrence(OccurrenceData->m_pExternalData);
+		}
+
 		A3DAsmProductOccurrence* PrototypePtr = OccurrenceData->m_pPrototype;
 		A3DAsmPartDefinition* PartDefinition = OccurrenceData->m_pPart;
 
