Techsoft interface implementation cleaning

diff --git a/Engine/Plugins/Enterprise/DatasmithCADImporter/Source/CADInterfaces/Private/CADInterfacesModule.cpp b/Engine/Plugins/Enterprise/DatasmithCADImporter/Source/CADInterfaces/Private/CADInterfacesModule.cpp
--- a/Engine/Plugins/Enterprise/DatasmithCADImporter/Source/CADInterfaces/Private/CADInterfacesModule.cpp
+++ b/Engine/Plugins/Enterprise/DatasmithCADImporter/Source/CADInterfaces/Private/CADInterfacesModule.cpp
@@ -91,10 +91,16 @@ void FCADInterfacesModule::StartupModule()
 #if WITH_EDITOR & defined(USE_TECHSOFT_SDK)
 	check(TechSoftLibHandle == nullptr);
 
-	FString TechSoftDllPath = FPaths::Combine(CADImporterDllPath, "TechSoft");
+	FString TechSoftDllPath = FPaths::ConvertRelativePathToFull(FPaths::Combine(CADImporterDllPath, "TechSoft"));
 	FPlatformProcess::PushDllDirectory(*TechSoftDllPath);
 
+#if PLATFORM_WINDOWS
 	FString TechSoftDll = TEXT("A3DLIBS.dll");
+#elif PLATFORM_LINUX
+	FString TechSoftDll = TEXT("libA3DLIBS.so");
+#else
+#error Platform not supported
+#endif
 	TechSoftDll = FPaths::Combine(TechSoftDllPath, TechSoftDll);
 
 	if (!FPaths::FileExists(TechSoftDll))
diff --git a/Engine/Plugins/Enterprise/DatasmithCADImporter/Source/CADInterfaces/Private/CTSession.cpp b/Engine/Plugins/Enterprise/DatasmithCADImporter/Source/CADInterfaces/Private/CTSession.cpp
--- a/Engine/Plugins/Enterprise/DatasmithCADImporter/Source/CADInterfaces/Private/CTSession.cpp
+++ b/Engine/Plugins/Enterprise/DatasmithCADImporter/Source/CADInterfaces/Private/CTSession.cpp
@@ -17,7 +17,7 @@ namespace CADLibrary
 
 	bool FCTSession::SaveBrep(const FString& FilePath)
 	{
-		return CTKIO_SaveFile({ MainObjectId }, *FilePath, L"Ct");
+		return CTKIO_SaveFile({ MainObjectId }, *FilePath, TEXT("Ct"));
 	}
 
 	bool FCTSession::TopoFixes(double SewingToleranceFactor)
diff --git a/Engine/Plugins/Enterprise/DatasmithCADImporter/Source/CADInterfaces/Private/TechSoftInterface.cpp b/Engine/Plugins/Enterprise/DatasmithCADImporter/Source/CADInterfaces/Private/TechSoftInterface.cpp
--- a/Engine/Plugins/Enterprise/DatasmithCADImporter/Source/CADInterfaces/Private/TechSoftInterface.cpp
+++ b/Engine/Plugins/Enterprise/DatasmithCADImporter/Source/CADInterfaces/Private/TechSoftInterface.cpp
@@ -45,7 +45,7 @@ bool FTechSoftInterface::InitializeKernel(const TCHAR* InEnginePluginsPath)
 
 	TechSoftDllPath = FPaths::ConvertRelativePathToFull(TechSoftDllPath);
 
-	if (A3DSDKLoadLibrary(*TechSoftDllPath))
+	if (A3DSDKLoadLibraryA(TCHAR_TO_UTF8(*TechSoftDllPath)))
 	{
 		A3DLicPutUnifiedLicense(HOOPS_LICENSE);
 
@@ -434,347 +434,352 @@ A3DStatus SewBReps(A3DRiBrepModel** BRepsToSew, uint32 const BRepCount, double T
 #ifdef USE_TECHSOFT_SDK
 
 // TUniqueTSObj InitializeData -----------------------------------
-
+template<>
 void TUniqueTSObj<A3DAsmModelFileData>::InitializeData()
 {
 	A3D_INITIALIZE_DATA(A3DAsmModelFileData, Data);
 }
-
+template<>
 void TUniqueTSObj<A3DAsmPartDefinitionData>::InitializeData()
 {
 	A3D_INITIALIZE_DATA(A3DAsmPartDefinitionData, Data);
 }
-
+template<>
 void TUniqueTSObj<A3DAsmProductOccurrenceData>::InitializeData()
 {
 	A3D_INITIALIZE_DATA(A3DAsmProductOccurrenceData, Data);
 }
-
+template<>
 void TUniqueTSObj<A3DAsmProductOccurrenceDataCV5>::InitializeData()
 {
 	A3D_INITIALIZE_DATA(A3DAsmProductOccurrenceDataCV5, Data);
 }
-
+template<>
 void TUniqueTSObj<A3DAsmProductOccurrenceDataSLW>::InitializeData()
 {
 	A3D_INITIALIZE_DATA(A3DAsmProductOccurrenceDataSLW, Data);
 }
-
+template<>
 void TUniqueTSObj<A3DAsmProductOccurrenceDataUg>::InitializeData()
 {
 	A3D_INITIALIZE_DATA(A3DAsmProductOccurrenceDataUg, Data);
 }
-
+template<>
 void TUniqueTSObj<A3DBoundingBoxData>::InitializeData()
 {
 	A3D_INITIALIZE_DATA(A3DBoundingBoxData, Data);
 }
-
+template<>
+void TUniqueTSObj<A3DCopyAndAdaptBrepModelData>::InitializeData()
+{
+	A3D_INITIALIZE_DATA(A3DCopyAndAdaptBrepModelData, Data);
+}
+template<>
 void TUniqueTSObj<A3DCrvCircleData>::InitializeData()
 {
 	A3D_INITIALIZE_DATA(A3DCrvCircleData, Data);
 }
-
+template<>
 void TUniqueTSObj<A3DCrvCompositeData>::InitializeData()
 {
 	A3D_INITIALIZE_DATA(A3DCrvCompositeData, Data);
 }
-
+template<>
 void TUniqueTSObj<A3DCrvEllipseData>::InitializeData()
 {
 	A3D_INITIALIZE_DATA(A3DCrvEllipseData, Data);
 }
-
+template<>
 void TUniqueTSObj<A3DCrvHelixData>::InitializeData()
 {
 	A3D_INITIALIZE_DATA(A3DCrvHelixData, Data);
 }
-
+template<>
 void TUniqueTSObj<A3DCrvHyperbolaData>::InitializeData()
 {
 	A3D_INITIALIZE_DATA(A3DCrvHyperbolaData, Data);
 }
-
+template<>
 void TUniqueTSObj<A3DCrvLineData>::InitializeData()
 {
 	A3D_INITIALIZE_DATA(A3DCrvLineData, Data);
 }
-
+template<>
 void TUniqueTSObj<A3DCrvNurbsData>::InitializeData()
 {
 	A3D_INITIALIZE_DATA(A3DCrvNurbsData, Data);
 }
-
+template<>
 void TUniqueTSObj<A3DCrvParabolaData>::InitializeData()
 {
 	A3D_INITIALIZE_DATA(A3DCrvParabolaData, Data);
 }
-
+template<>
 void TUniqueTSObj<A3DCrvPolyLineData>::InitializeData()
 {
 	A3D_INITIALIZE_DATA(A3DCrvPolyLineData, Data);
 }
-
+template<>
 void TUniqueTSObj<A3DCrvTransformData>::InitializeData()
 {
 	A3D_INITIALIZE_DATA(A3DCrvTransformData, Data);
 }
-
+template<>
 void TUniqueTSObj<A3DGlobalData>::InitializeData()
 {
 	A3D_INITIALIZE_DATA(A3DGlobalData, Data);
 }
-
+template<>
 void TUniqueTSObj<A3DGraphicsData>::InitializeData()
 {
 	A3D_INITIALIZE_DATA(A3DGraphicsData, Data);
 }
-
+template<>
 void TUniqueTSObj<A3DIntervalData>::InitializeData()
 {
 	A3D_INITIALIZE_DATA(A3DIntervalData, Data);
 }
-
+template<>
 void TUniqueTSObj<A3DMiscAttributeData>::InitializeData()
 {
 	A3D_INITIALIZE_DATA(A3DMiscAttributeData, Data);
 }
-
+template<>
 void TUniqueTSObj<A3DMiscCartesianTransformationData>::InitializeData()
 {
 	A3D_INITIALIZE_DATA(A3DMiscCartesianTransformationData, Data);
 }
-
+template<>
 void TUniqueTSObj<A3DMiscEntityReferenceData>::InitializeData()
 {
 	A3D_INITIALIZE_DATA(A3DMiscEntityReferenceData, Data);
 }
-
+template<>
 void TUniqueTSObj<A3DMiscGeneralTransformationData>::InitializeData()
 {
 	A3D_INITIALIZE_DATA(A3DMiscGeneralTransformationData, Data);
 }
-
+template<>
 void TUniqueTSObj<A3DMiscMaterialPropertiesData>::InitializeData()
 {
 	A3D_INITIALIZE_DATA(A3DMiscMaterialPropertiesData, Data);
 }
-
+template<>
 void TUniqueTSObj<A3DMiscReferenceOnCsysItemData>::InitializeData()
 {
 	A3D_INITIALIZE_DATA(A3DMiscReferenceOnCsysItemData, Data);
 }
-
+template<>
 void TUniqueTSObj<A3DMiscReferenceOnTessData>::InitializeData()
 {
 	A3D_INITIALIZE_DATA(A3DMiscReferenceOnTessData, Data);
 }
-
+template<>
 void TUniqueTSObj<A3DMiscReferenceOnTopologyData>::InitializeData()
 {
 	A3D_INITIALIZE_DATA(A3DMiscReferenceOnTopologyData, Data);
 }
-
+template<>
 void TUniqueTSObj<A3DMiscSingleAttributeData>::InitializeData()
 {
 	A3D_INITIALIZE_DATA(A3DMiscSingleAttributeData, Data);
 }
-
+template<>
 void TUniqueTSObj<A3DRWParamsExportPrcData>::InitializeData()
 {
 	A3D_INITIALIZE_DATA(A3DRWParamsExportPrcData, Data);
 }
-
+template<>
 void TUniqueTSObj<A3DRiBrepModelData>::InitializeData()
 {
 	A3D_INITIALIZE_DATA(A3DRiBrepModelData, Data);
 }
-
+template<>
 void TUniqueTSObj<A3DRiCoordinateSystemData>::InitializeData()
 {
 	A3D_INITIALIZE_DATA(A3DRiCoordinateSystemData, Data);
 }
-
+template<>
 void TUniqueTSObj<A3DRiDirectionData>::InitializeData()
 {
 	A3D_INITIALIZE_DATA(A3DRiDirectionData, Data);
 }
-
+template<>
 void TUniqueTSObj<A3DRiPolyBrepModelData>::InitializeData()
 {
 	A3D_INITIALIZE_DATA(A3DRiPolyBrepModelData, Data);
 }
-
+template<>
 void TUniqueTSObj<A3DRiRepresentationItemData>::InitializeData()
 {
 	A3D_INITIALIZE_DATA(A3DRiRepresentationItemData, Data);
 }
-
+template<>
 void TUniqueTSObj<A3DRiSetData>::InitializeData()
 {
 	A3D_INITIALIZE_DATA(A3DRiSetData, Data);
 }
-
+template<>
 void TUniqueTSObj<A3DRootBaseData>::InitializeData()
 {
 	A3D_INITIALIZE_DATA(A3DRootBaseData, Data);
 }
-
+template<>
 void TUniqueTSObj<A3DRootBaseWithGraphicsData>::InitializeData()
 {
 	A3D_INITIALIZE_DATA(A3DRootBaseWithGraphicsData, Data);
 }
-
+template<>
 void TUniqueTSObj<A3DRWParamsTessellationData>::InitializeData()
 {
 	A3D_INITIALIZE_DATA(A3DRWParamsTessellationData, Data);
 }
-
+template<>
 void TUniqueTSObj<A3DSewOptionsData>::InitializeData()
 {
 	A3D_INITIALIZE_DATA(A3DSewOptionsData, Data);
 }
-
+template<>
 void TUniqueTSObj<A3DSurfBlend01Data>::InitializeData()
 {
 	A3D_INITIALIZE_DATA(A3DSurfBlend01Data, Data);
 }
-
+template<>
 void TUniqueTSObj<A3DSurfBlend02Data>::InitializeData()
 {
 	A3D_INITIALIZE_DATA(A3DSurfBlend02Data, Data);
 }
-
+template<>
 void TUniqueTSObj<A3DSurfBlend03Data>::InitializeData()
 {
 	A3D_INITIALIZE_DATA(A3DSurfBlend03Data, Data);
 }
-
+template<>
 void TUniqueTSObj<A3DSurfConeData>::InitializeData()
 {
 	A3D_INITIALIZE_DATA(A3DSurfConeData, Data);
 }
-
+template<>
 void TUniqueTSObj<A3DSurfCylinderData>::InitializeData()
 {
 	A3D_INITIALIZE_DATA(A3DSurfCylinderData, Data);
 }
-
+template<>
 void TUniqueTSObj<A3DSurfCylindricalData>::InitializeData()
 {
 	A3D_INITIALIZE_DATA(A3DSurfCylindricalData, Data);
 }
-
+template<>
 void TUniqueTSObj<A3DSurfExtrusionData>::InitializeData()
 {
 	A3D_INITIALIZE_DATA(A3DSurfExtrusionData, Data);
 }
-
+template<>
 void TUniqueTSObj<A3DSurfFromCurvesData>::InitializeData()
 {
 	A3D_INITIALIZE_DATA(A3DSurfFromCurvesData, Data);
 }
-
+template<>
 void TUniqueTSObj<A3DSurfNurbsData>::InitializeData()
 {
 	A3D_INITIALIZE_DATA(A3DSurfNurbsData, Data);
 }
-
+template<>
 void TUniqueTSObj<A3DSurfPipeData>::InitializeData()
 {
 	A3D_INITIALIZE_DATA(A3DSurfPipeData, Data);
 }
-
+template<>
 void TUniqueTSObj<A3DSurfPlaneData>::InitializeData()
 {
 	A3D_INITIALIZE_DATA(A3DSurfPlaneData, Data);
 }
-
+template<>
 void TUniqueTSObj<A3DSurfRevolutionData>::InitializeData()
 {
 	A3D_INITIALIZE_DATA(A3DSurfRevolutionData, Data);
 }
-
+template<>
 void TUniqueTSObj<A3DSurfRuledData>::InitializeData()
 {
 	A3D_INITIALIZE_DATA(A3DSurfRuledData, Data);
 }
-
+template<>
 void TUniqueTSObj<A3DSurfSphereData>::InitializeData()
 {
 	A3D_INITIALIZE_DATA(A3DSurfSphereData, Data);
 }
-
+template<>
 void TUniqueTSObj<A3DSurfTorusData>::InitializeData()
 {
 	A3D_INITIALIZE_DATA(A3DSurfTorusData, Data);
 }
-
+template<>
 void TUniqueTSObj<A3DTess3DData>::InitializeData()
 {
 	A3D_INITIALIZE_DATA(A3DTess3DData, Data);
 }
-
+template<>
 void TUniqueTSObj<A3DTessBaseData>::InitializeData()
 {
 	A3D_INITIALIZE_DATA(A3DTessBaseData, Data);
 }
-
+template<>
 void TUniqueTSObj<A3DTopoBodyData>::InitializeData()
 {
 	A3D_INITIALIZE_DATA(A3DTopoBodyData, Data);
 }
-
+template<>
 void TUniqueTSObj<A3DTopoBrepDataData>::InitializeData()
 {
 	A3D_INITIALIZE_DATA(A3DTopoBrepDataData, Data);
 }
-
+template<>
 void TUniqueTSObj<A3DTopoCoEdgeData>::InitializeData()
 {
 	A3D_INITIALIZE_DATA(A3DTopoCoEdgeData, Data);
 }
-
+template<>
 void TUniqueTSObj<A3DTopoConnexData>::InitializeData()
 {
 	A3D_INITIALIZE_DATA(A3DTopoConnexData, Data);
 }
-
+template<>
 void TUniqueTSObj<A3DTopoContextData>::InitializeData()
 {
 	A3D_INITIALIZE_DATA(A3DTopoContextData, Data);
 }
-
+template<>
 void TUniqueTSObj<A3DTopoEdgeData>::InitializeData()
 {
 	A3D_INITIALIZE_DATA(A3DTopoEdgeData, Data);
 }
-
+template<>
 void TUniqueTSObj<A3DTopoFaceData>::InitializeData()
 {
 	A3D_INITIALIZE_DATA(A3DTopoFaceData, Data);
 }
-
+template<>
 void TUniqueTSObj<A3DTopoLoopData>::InitializeData()
 {
 	A3D_INITIALIZE_DATA(A3DTopoLoopData, Data);
 }
-
+template<>
 void TUniqueTSObj<A3DTopoShellData>::InitializeData()
 {
 	A3D_INITIALIZE_DATA(A3DTopoShellData, Data);
 }
-
+template<>
 void TUniqueTSObj<A3DTopoUniqueVertexData>::InitializeData()
 {
 	A3D_INITIALIZE_DATA(A3DTopoUniqueVertexData, Data);
 }
-
+template<>
 void TUniqueTSObj<A3DTopoWireEdgeData>::InitializeData()
 {
 	A3D_INITIALIZE_DATA(A3DTopoWireEdgeData, Data);
 }
-
+template<>
 void TUniqueTSObj<A3DUTF8Char*>::InitializeData()
 {
 	Data = nullptr;
@@ -783,33 +788,32 @@ void TUniqueTSObj<A3DUTF8Char*>::InitializeData()
 
 
 // TUniqueTSObjFromIndex InitializeData -----------------------------------
-
-
+template<>
 void TUniqueTSObjFromIndex<A3DGraphMaterialData>::InitializeData()
 {
 	A3D_INITIALIZE_DATA(A3DGraphMaterialData, Data);
 }
-
+template<>
 void TUniqueTSObjFromIndex<A3DGraphPictureData>::InitializeData()
 {
 	A3D_INITIALIZE_DATA(A3DGraphPictureData, Data);
 }
-
+template<>
 void TUniqueTSObjFromIndex<A3DGraphRgbColorData>::InitializeData()
 {
 	A3D_INITIALIZE_DATA(A3DGraphRgbColorData, Data);
 }
-
+template<>
 void TUniqueTSObjFromIndex<A3DGraphStyleData>::InitializeData()
 {
 	A3D_INITIALIZE_DATA(A3DGraphStyleData, Data);
 }
-
+template<>
 void TUniqueTSObjFromIndex<A3DGraphTextureApplicationData>::InitializeData()
 {
 	A3D_INITIALIZE_DATA(A3DGraphTextureApplicationData, Data);
 }
-
+template<>
 void TUniqueTSObjFromIndex<A3DGraphTextureDefinitionData>::InitializeData()
 {
 	A3D_INITIALIZE_DATA(A3DGraphTextureDefinitionData, Data);
@@ -819,465 +823,612 @@ void TUniqueTSObjFromIndex<A3DGraphTextureDefinitionData>::InitializeData()
 
 // TUniqueTSObj GetData -----------------------------------
 
-
+template<>
 A3DStatus TUniqueTSObj<A3DAsmModelFileData>::GetData(const A3DEntity* InEntityPtr)
 {
 	return A3DAsmModelFileGet(InEntityPtr, &Data);
 }
 
+template<>
 A3DStatus TUniqueTSObj<A3DAsmPartDefinitionData>::GetData(const A3DEntity* InEntityPtr)
 {
 	return A3DAsmPartDefinitionGet(InEntityPtr, &Data);
 }
 
+template<>
 A3DStatus TUniqueTSObj<A3DAsmProductOccurrenceData>::GetData(const A3DEntity* InEntityPtr)
 {
 	return A3DAsmProductOccurrenceGet(InEntityPtr, &Data);
 }
 
+template<>
 A3DStatus TUniqueTSObj<A3DAsmProductOccurrenceDataCV5>::GetData(const A3DEntity* InEntityPtr)
 {
 	return A3DAsmProductOccurrenceGetCV5(InEntityPtr, &Data);
 }
 
+template<>
 A3DStatus TUniqueTSObj<A3DAsmProductOccurrenceDataSLW>::GetData(const A3DEntity* InEntityPtr)
 {
 	return A3DAsmProductOccurrenceGetSLW(InEntityPtr, &Data);
 }
 
+template<>
 A3DStatus TUniqueTSObj<A3DAsmProductOccurrenceDataUg>::GetData(const A3DEntity* InEntityPtr)
 {
 	return A3DAsmProductOccurrenceGetUg(InEntityPtr, &Data);
 }
 
+template<>
 A3DStatus TUniqueTSObj<A3DBoundingBoxData>::GetData(const A3DEntity* InEntityPtr)
 {
 	return A3DMiscGetBoundingBox(InEntityPtr, &Data);
 }
 
+template<>
+A3DStatus TUniqueTSObj<A3DCopyAndAdaptBrepModelData>::GetData(const A3DEntity* InEntityPtr)
+{
+	return A3DStatus::A3D_ERROR;
+}
+
+template<>
 A3DStatus TUniqueTSObj<A3DCrvCircleData>::GetData(const A3DEntity* InEntityPtr)
 {
 	return A3DCrvCircleGet(InEntityPtr, &Data);
 }
 
+template<>
 A3DStatus TUniqueTSObj<A3DCrvCompositeData>::GetData(const A3DEntity* InEntityPtr)
 {
 	return A3DCrvCompositeGet(InEntityPtr, &Data);
 }
 
+template<>
 A3DStatus TUniqueTSObj<A3DCrvEllipseData>::GetData(const A3DEntity* InEntityPtr)
 {
 	return A3DCrvEllipseGet(InEntityPtr, &Data);
 }
 
+template<>
 A3DStatus TUniqueTSObj<A3DCrvHelixData>::GetData(const A3DEntity* InEntityPtr)
 {
 	return A3DCrvHelixGet(InEntityPtr, &Data);
 }
 
+template<>
 A3DStatus TUniqueTSObj<A3DCrvHyperbolaData>::GetData(const A3DEntity* InEntityPtr)
 {
 	return A3DCrvHyperbolaGet(InEntityPtr, &Data);
 }
 
+template<>
 A3DStatus TUniqueTSObj<A3DCrvLineData>::GetData(const A3DEntity* InEntityPtr)
 {
 	return A3DCrvLineGet(InEntityPtr, &Data);
 }
 
+template<>
 A3DStatus TUniqueTSObj<A3DCrvNurbsData>::GetData(const A3DEntity* InEntityPtr)
 {
 	return A3DCrvNurbsGet(InEntityPtr, &Data);
 }
 
+template<>
 A3DStatus TUniqueTSObj<A3DCrvParabolaData>::GetData(const A3DEntity* InEntityPtr)
 {
 	return A3DCrvParabolaGet(InEntityPtr, &Data);
 }
 
+template<>
 A3DStatus TUniqueTSObj<A3DCrvPolyLineData>::GetData(const A3DEntity* InEntityPtr)
 {
 	return A3DCrvPolyLineGet(InEntityPtr, &Data);
 }
 
+template<>
 A3DStatus TUniqueTSObj<A3DCrvTransformData>::GetData(const A3DEntity* InEntityPtr)
 {
 	return A3DCrvTransformGet(InEntityPtr, &Data);
 }
 
+template<>
 A3DStatus TUniqueTSObj<A3DGlobalData>::GetData(const A3DEntity* InEntityPtr)
 {
 	return A3DGlobalGet(InEntityPtr, &Data);
 }
 
+template<>
 A3DStatus TUniqueTSObj<A3DGraphicsData>::GetData(const A3DEntity* InEntityPtr)
 {
 	return A3DGraphicsGet(InEntityPtr, &Data);
 }
 
+template<>
 A3DStatus TUniqueTSObj<A3DIntervalData>::GetData(const A3DEntity* InEntityPtr)
 {
 	return A3DCrvGetInterval(InEntityPtr, &Data);
 }
 
+template<>
 A3DStatus TUniqueTSObj<A3DMiscAttributeData>::GetData(const A3DEntity* InEntityPtr)
 {
 	return A3DMiscAttributeGet(InEntityPtr, &Data);
 }
 
+template<>
 A3DStatus TUniqueTSObj<A3DMiscCartesianTransformationData>::GetData(const A3DEntity* InEntityPtr)
 {
 	return A3DMiscCartesianTransformationGet(InEntityPtr, &Data);
 }
 
+template<>
 A3DStatus TUniqueTSObj<A3DMiscEntityReferenceData>::GetData(const A3DEntity* InEntityPtr)
 {
 	return A3DMiscEntityReferenceGet(InEntityPtr, &Data);
 }
 
+template<>
 A3DStatus TUniqueTSObj<A3DMiscGeneralTransformationData>::GetData(const A3DEntity* InEntityPtr)
 {
 	return A3DMiscGeneralTransformationGet(InEntityPtr, &Data);
 }
 
+template<>
 A3DStatus TUniqueTSObj<A3DMiscMaterialPropertiesData>::GetData(const A3DEntity* InEntityPtr)
 {
 	return A3DMiscGetMaterialProperties(InEntityPtr, &Data);
 }
 
+template<>
 A3DStatus TUniqueTSObj<A3DMiscReferenceOnCsysItemData>::GetData(const A3DEntity* InEntityPtr)
 {
 	return A3DMiscReferenceOnCsysItemGet(InEntityPtr, &Data);
 }
 
+template<>
 A3DStatus TUniqueTSObj<A3DMiscReferenceOnTessData>::GetData(const A3DEntity* InEntityPtr)
 {
 	return A3DMiscReferenceOnTessGet(InEntityPtr, &Data);
 }
 
+template<>
 A3DStatus TUniqueTSObj<A3DMiscReferenceOnTopologyData>::GetData(const A3DEntity* InEntityPtr)
 {
 	return A3DMiscReferenceOnTopologyGet(InEntityPtr, &Data);
 }
 
+template<>
 A3DStatus TUniqueTSObj<A3DMiscSingleAttributeData>::GetData(const A3DEntity* InEntityPtr)
 {
 	return A3DStatus::A3D_ERROR;
 }
 
+template<>
 A3DStatus TUniqueTSObj<A3DRWParamsExportPrcData>::GetData(const A3DEntity* InEntityPtr)
 {
 	return A3DStatus::A3D_ERROR;
 }
 
+template<>
 A3DStatus TUniqueTSObj<A3DRiBrepModelData>::GetData(const A3DEntity* InEntityPtr)
 {
 	return A3DRiBrepModelGet(InEntityPtr, &Data);
 }
 
+template<>
 A3DStatus TUniqueTSObj<A3DRiCoordinateSystemData>::GetData(const A3DEntity* InEntityPtr)
 {
 	return A3DRiCoordinateSystemGet(InEntityPtr, &Data);
 }
 
+template<>
 A3DStatus TUniqueTSObj<A3DRiDirectionData>::GetData(const A3DEntity* InEntityPtr)
 {
 	return A3DRiDirectionGet(InEntityPtr, &Data);
 }
 
+template<>
 A3DStatus TUniqueTSObj<A3DRiPolyBrepModelData>::GetData(const A3DEntity* InEntityPtr)
 {
 	return A3DRiPolyBrepModelGet(InEntityPtr, &Data);
 }
 
+template<>
 A3DStatus TUniqueTSObj<A3DRiRepresentationItemData>::GetData(const A3DEntity* InEntityPtr)
 {
 	return A3DRiRepresentationItemGet(InEntityPtr, &Data);
 }
 
+template<>
 A3DStatus TUniqueTSObj<A3DRiSetData>::GetData(const A3DEntity* InEntityPtr)
 {
 	return A3DRiSetGet(InEntityPtr, &Data);
 }
 
+template<>
 A3DStatus TUniqueTSObj<A3DRootBaseData>::GetData(const A3DEntity* InEntityPtr)
 {
 	return A3DRootBaseGet(InEntityPtr, &Data);
 }
 
+template<>
 A3DStatus TUniqueTSObj<A3DRootBaseWithGraphicsData>::GetData(const A3DEntity* InEntityPtr)
 {
 	return A3DRootBaseWithGraphicsGet(InEntityPtr, &Data);
 }
 
+template<>
 A3DStatus TUniqueTSObj<A3DRWParamsTessellationData>::GetData(const A3DEntity* InEntityPtr)
 {
 	return A3DStatus::A3D_ERROR;
 }
 
+template<>
 A3DStatus TUniqueTSObj<A3DSewOptionsData>::GetData(const A3DEntity* InEntityPtr)
 {
 	return A3DStatus::A3D_ERROR;
 }
 
+template<>
 A3DStatus TUniqueTSObj<A3DSurfBlend01Data>::GetData(const A3DEntity* InEntityPtr)
 {
 	return A3DSurfBlend01Get(InEntityPtr, &Data);
 }
 
+template<>
 A3DStatus TUniqueTSObj<A3DSurfBlend02Data>::GetData(const A3DEntity* InEntityPtr)
 {
 	return A3DSurfBlend02Get(InEntityPtr, &Data);
 }
 
+template<>
 A3DStatus TUniqueTSObj<A3DSurfBlend03Data>::GetData(const A3DEntity* InEntityPtr)
 {
 	return A3DSurfBlend03Get(InEntityPtr, &Data);
 }
 
+template<>
 A3DStatus TUniqueTSObj<A3DSurfConeData>::GetData(const A3DEntity* InEntityPtr)
 {
 	return A3DSurfConeGet(InEntityPtr, &Data);
 }
 
+template<>
 A3DStatus TUniqueTSObj<A3DSurfCylinderData>::GetData(const A3DEntity* InEntityPtr)
 {
 	return A3DSurfCylinderGet(InEntityPtr, &Data);
 }
 
+template<>
 A3DStatus TUniqueTSObj<A3DSurfCylindricalData>::GetData(const A3DEntity* InEntityPtr)
 {
 	return A3DSurfCylindricalGet(InEntityPtr, &Data);
 }
 
+template<>
 A3DStatus TUniqueTSObj<A3DSurfExtrusionData>::GetData(const A3DEntity* InEntityPtr)
 {
 	return A3DSurfExtrusionGet(InEntityPtr, &Data);
 }
 
+template<>
 A3DStatus TUniqueTSObj<A3DSurfFromCurvesData>::GetData(const A3DEntity* InEntityPtr)
 {
 	return A3DSurfFromCurvesGet(InEntityPtr, &Data);
 }
 
+template<>
 A3DStatus TUniqueTSObj<A3DSurfNurbsData>::GetData(const A3DEntity* InEntityPtr)
 {
 	return A3DSurfNurbsGet(InEntityPtr, &Data);
 }
 
+template<>
 A3DStatus TUniqueTSObj<A3DSurfPipeData>::GetData(const A3DEntity* InEntityPtr)
 {
 	return A3DSurfPipeGet(InEntityPtr, &Data);
 }
 
+template<>
 A3DStatus TUniqueTSObj<A3DSurfPlaneData>::GetData(const A3DEntity* InEntityPtr)
 {
 	return A3DSurfPlaneGet(InEntityPtr, &Data);
 }
 
+template<>
 A3DStatus TUniqueTSObj<A3DSurfRevolutionData>::GetData(const A3DEntity* InEntityPtr)
 {
 	return A3DSurfRevolutionGet(InEntityPtr, &Data);
 }
 
+template<>
 A3DStatus TUniqueTSObj<A3DSurfRuledData>::GetData(const A3DEntity* InEntityPtr)
 {
 	return A3DSurfRuledGet(InEntityPtr, &Data);
 }
 
+template<>
 A3DStatus TUniqueTSObj<A3DSurfSphereData>::GetData(const A3DEntity* InEntityPtr)
 {
 	return A3DSurfSphereGet(InEntityPtr, &Data);
 }
 
+template<>
 A3DStatus TUniqueTSObj<A3DSurfTorusData>::GetData(const A3DEntity* InEntityPtr)
 {
 	return A3DSurfTorusGet(InEntityPtr, &Data);
 }
 
+template<>
 A3DStatus TUniqueTSObj<A3DTess3DData>::GetData(const A3DEntity* InEntityPtr)
 {
 	return A3DTess3DGet(InEntityPtr, &Data);
 }
 
+template<>
 A3DStatus TUniqueTSObj<A3DTessBaseData>::GetData(const A3DEntity* InEntityPtr)
 {
 	return A3DTessBaseGet(InEntityPtr, &Data);
 }
 
+template<>
 A3DStatus TUniqueTSObj<A3DTopoBodyData>::GetData(const A3DEntity* InEntityPtr)
 {
 	return A3DTopoBodyGet(InEntityPtr, &Data);
 }
 
+template<>
 A3DStatus TUniqueTSObj<A3DTopoBrepDataData>::GetData(const A3DEntity* InEntityPtr)
 {
 	return A3DTopoBrepDataGet(InEntityPtr, &Data);
 }
 
+template<>
 A3DStatus TUniqueTSObj<A3DTopoCoEdgeData>::GetData(const A3DEntity* InEntityPtr)
 {
 	return A3DTopoCoEdgeGet(InEntityPtr, &Data);
 }
 
+template<>
 A3DStatus TUniqueTSObj<A3DTopoConnexData>::GetData(const A3DEntity* InEntityPtr)
 {
 	return A3DTopoConnexGet(InEntityPtr, &Data);
 }
 
+template<>
 A3DStatus TUniqueTSObj<A3DTopoContextData>::GetData(const A3DEntity* InEntityPtr)
 {
 	return A3DTopoContextGet(InEntityPtr, &Data);
 }
 
+template<>
 A3DStatus TUniqueTSObj<A3DTopoEdgeData>::GetData(const A3DEntity* InEntityPtr)
 {
 	return A3DTopoEdgeGet(InEntityPtr, &Data);
 }
 
+template<>
 A3DStatus TUniqueTSObj<A3DTopoFaceData>::GetData(const A3DEntity* InEntityPtr)
 {
 	return A3DTopoFaceGet(InEntityPtr, &Data);
 }
 
+template<>
 A3DStatus TUniqueTSObj<A3DTopoLoopData>::GetData(const A3DEntity* InEntityPtr)
 {
 	return A3DTopoLoopGet(InEntityPtr, &Data);
 }
 
+template<>
 A3DStatus TUniqueTSObj<A3DTopoShellData>::GetData(const A3DEntity* InEntityPtr)
 {
 	return A3DTopoShellGet(InEntityPtr, &Data);
 }
 
+template<>
 A3DStatus TUniqueTSObj<A3DTopoUniqueVertexData>::GetData(const A3DEntity* InEntityPtr)
 {
 	return A3DTopoUniqueVertexGet(InEntityPtr, &Data);
 }
 
+template<>
 A3DStatus TUniqueTSObj<A3DTopoWireEdgeData>::GetData(const A3DEntity* InEntityPtr)
 {
 	return A3DTopoWireEdgeGet(InEntityPtr, &Data);
 }
 
+template<>
 A3DStatus TUniqueTSObj<A3DUTF8Char*>::GetData(const A3DEntity* InEntityPtr)
 {
 	return A3DStatus::A3D_ERROR;
 }
 
 
-
 // TUniqueTSObjFromIndex GetData -----------------------------------
-
-
+template<>
 A3DStatus TUniqueTSObjFromIndex<A3DGraphMaterialData>::GetData(const uint32 InEntityIndex)
 {
 	return A3DGlobalGetGraphMaterialData(InEntityIndex, &Data);
 }
-
+template<>
 A3DStatus TUniqueTSObjFromIndex<A3DGraphPictureData>::GetData(const uint32 InEntityIndex)
 {
 	return A3DGlobalGetGraphPictureData(InEntityIndex, &Data);
 }
-
+template<>
 A3DStatus TUniqueTSObjFromIndex<A3DGraphRgbColorData>::GetData(const uint32 InEntityIndex)
 {
 	return A3DGlobalGetGraphRgbColorData(InEntityIndex, &Data);
 }
-
+template<>
 A3DStatus TUniqueTSObjFromIndex<A3DGraphStyleData>::GetData(const uint32 InEntityIndex)
 {
 	return A3DGlobalGetGraphStyleData(InEntityIndex, &Data);
 }
-
+template<>
 A3DStatus TUniqueTSObjFromIndex<A3DGraphTextureApplicationData>::GetData(const uint32 InEntityIndex)
 {
 	return A3DGlobalGetGraphTextureApplicationData(InEntityIndex, &Data);
 }
-
+template<>
 A3DStatus TUniqueTSObjFromIndex<A3DGraphTextureDefinitionData>::GetData(const uint32 InEntityIndex)
 {
 	return A3DGlobalGetGraphTextureDefinitionData(InEntityIndex, &Data);
 }
 
 // DefaultValue -----------------------------------
-
+template<>
 const A3DEntity* TUniqueTSObj<A3DAsmModelFileData>::GetDefaultIndexerValue() const { return nullptr; }
+template<>
 const A3DEntity* TUniqueTSObj<A3DAsmPartDefinitionData>::GetDefaultIndexerValue() const { return nullptr; }
+template<>
 const A3DEntity* TUniqueTSObj<A3DAsmProductOccurrenceData>::GetDefaultIndexerValue() const { return nullptr; }
+template<>
 const A3DEntity* TUniqueTSObj<A3DAsmProductOccurrenceDataCV5>::GetDefaultIndexerValue() const { return nullptr; }
+template<>
 const A3DEntity* TUniqueTSObj<A3DAsmProductOccurrenceDataSLW>::GetDefaultIndexerValue() const { return nullptr; }
+template<>
 const A3DEntity* TUniqueTSObj<A3DAsmProductOccurrenceDataUg>::GetDefaultIndexerValue() const { return nullptr; }
+template<>
 const A3DEntity* TUniqueTSObj<A3DBoundingBoxData>::GetDefaultIndexerValue() const { return nullptr; }
+template<>
+const A3DEntity* TUniqueTSObj<A3DCopyAndAdaptBrepModelData>::GetDefaultIndexerValue() const { return nullptr; }
+template<>
 const A3DEntity* TUniqueTSObj<A3DCrvCircleData>::GetDefaultIndexerValue() const { return nullptr; }
+template<>
 const A3DEntity* TUniqueTSObj<A3DCrvCompositeData>::GetDefaultIndexerValue() const { return nullptr; }
+template<>
 const A3DEntity* TUniqueTSObj<A3DCrvEllipseData>::GetDefaultIndexerValue() const { return nullptr; }
+template<>
 const A3DEntity* TUniqueTSObj<A3DCrvHelixData>::GetDefaultIndexerValue() const { return nullptr; }
+template<>
 const A3DEntity* TUniqueTSObj<A3DCrvHyperbolaData>::GetDefaultIndexerValue() const { return nullptr; }
+template<>
 const A3DEntity* TUniqueTSObj<A3DCrvLineData>::GetDefaultIndexerValue() const { return nullptr; }
+template<>
 const A3DEntity* TUniqueTSObj<A3DCrvNurbsData>::GetDefaultIndexerValue() const { return nullptr; }
+template<>
 const A3DEntity* TUniqueTSObj<A3DCrvParabolaData>::GetDefaultIndexerValue() const { return nullptr; }
+template<>
 const A3DEntity* TUniqueTSObj<A3DCrvPolyLineData>::GetDefaultIndexerValue() const { return nullptr; }
+template<>
 const A3DEntity* TUniqueTSObj<A3DCrvTransformData>::GetDefaultIndexerValue() const { return nullptr; }
+template<>
 const A3DEntity* TUniqueTSObj<A3DGlobalData>::GetDefaultIndexerValue() const { return nullptr; }
+template<>
 const A3DEntity* TUniqueTSObj<A3DGraphicsData>::GetDefaultIndexerValue() const { return nullptr; }
+template<>
 const A3DEntity* TUniqueTSObj<A3DIntervalData>::GetDefaultIndexerValue() const { return nullptr; }
+template<>
 const A3DEntity* TUniqueTSObj<A3DMiscAttributeData>::GetDefaultIndexerValue() const { return nullptr; }
+template<>
 const A3DEntity* TUniqueTSObj<A3DMiscCartesianTransformationData>::GetDefaultIndexerValue() const { return nullptr; }
+template<>
 const A3DEntity* TUniqueTSObj<A3DMiscEntityReferenceData>::GetDefaultIndexerValue() const { return nullptr; }
+template<>
 const A3DEntity* TUniqueTSObj<A3DMiscGeneralTransformationData>::GetDefaultIndexerValue() const { return nullptr; }
+template<>
 const A3DEntity* TUniqueTSObj<A3DMiscMaterialPropertiesData>::GetDefaultIndexerValue() const { return nullptr; }
+template<>
 const A3DEntity* TUniqueTSObj<A3DMiscReferenceOnCsysItemData>::GetDefaultIndexerValue() const { return nullptr; }
+template<>
 const A3DEntity* TUniqueTSObj<A3DMiscReferenceOnTessData>::GetDefaultIndexerValue() const { return nullptr; }
+template<>
 const A3DEntity* TUniqueTSObj<A3DMiscReferenceOnTopologyData>::GetDefaultIndexerValue() const { return nullptr; }
+template<>
 const A3DEntity* TUniqueTSObj<A3DMiscSingleAttributeData>::GetDefaultIndexerValue() const { return nullptr; }
+template<>
 const A3DEntity* TUniqueTSObj<A3DRWParamsExportPrcData>::GetDefaultIndexerValue() const { return nullptr; }
+template<>
 const A3DEntity* TUniqueTSObj<A3DRiBrepModelData>::GetDefaultIndexerValue() const { return nullptr; }
+template<>
 const A3DEntity* TUniqueTSObj<A3DRiCoordinateSystemData>::GetDefaultIndexerValue() const { return nullptr; }
+template<>
 const A3DEntity* TUniqueTSObj<A3DRiDirectionData>::GetDefaultIndexerValue() const { return nullptr; }
+template<>
 const A3DEntity* TUniqueTSObj<A3DRiPolyBrepModelData>::GetDefaultIndexerValue() const { return nullptr; }
+template<>
 const A3DEntity* TUniqueTSObj<A3DRiRepresentationItemData>::GetDefaultIndexerValue() const { return nullptr; }
+template<>
 const A3DEntity* TUniqueTSObj<A3DRiSetData>::GetDefaultIndexerValue() const { return nullptr; }
+template<>
 const A3DEntity* TUniqueTSObj<A3DRootBaseData>::GetDefaultIndexerValue() const { return nullptr; }
+template<>
 const A3DEntity* TUniqueTSObj<A3DRootBaseWithGraphicsData>::GetDefaultIndexerValue() const { return nullptr; }
+template<>
 const A3DEntity* TUniqueTSObj<A3DRWParamsTessellationData>::GetDefaultIndexerValue() const { return nullptr; }
+template<>
 const A3DEntity* TUniqueTSObj<A3DSewOptionsData>::GetDefaultIndexerValue() const { return nullptr; }
+template<>
 const A3DEntity* TUniqueTSObj<A3DSurfBlend01Data>::GetDefaultIndexerValue() const { return nullptr; }
+template<>
 const A3DEntity* TUniqueTSObj<A3DSurfBlend02Data>::GetDefaultIndexerValue() const { return nullptr; }
+template<>
 const A3DEntity* TUniqueTSObj<A3DSurfBlend03Data>::GetDefaultIndexerValue() const { return nullptr; }
+template<>
 const A3DEntity* TUniqueTSObj<A3DSurfConeData>::GetDefaultIndexerValue() const { return nullptr; }
+template<>
 const A3DEntity* TUniqueTSObj<A3DSurfCylinderData>::GetDefaultIndexerValue() const { return nullptr; }
+template<>
 const A3DEntity* TUniqueTSObj<A3DSurfCylindricalData>::GetDefaultIndexerValue() const { return nullptr; }
+template<>
 const A3DEntity* TUniqueTSObj<A3DSurfExtrusionData>::GetDefaultIndexerValue() const { return nullptr; }
+template<>
 const A3DEntity* TUniqueTSObj<A3DSurfFromCurvesData>::GetDefaultIndexerValue() const { return nullptr; }
+template<>
 const A3DEntity* TUniqueTSObj<A3DSurfNurbsData>::GetDefaultIndexerValue() const { return nullptr; }
+template<>
 const A3DEntity* TUniqueTSObj<A3DSurfPipeData>::GetDefaultIndexerValue() const { return nullptr; }
+template<>
 const A3DEntity* TUniqueTSObj<A3DSurfPlaneData>::GetDefaultIndexerValue() const { return nullptr; }
+template<>
 const A3DEntity* TUniqueTSObj<A3DSurfRevolutionData>::GetDefaultIndexerValue() const { return nullptr; }
+template<>
 const A3DEntity* TUniqueTSObj<A3DSurfRuledData>::GetDefaultIndexerValue() const { return nullptr; }
+template<>
 const A3DEntity* TUniqueTSObj<A3DSurfSphereData>::GetDefaultIndexerValue() const { return nullptr; }
+template<>
 const A3DEntity* TUniqueTSObj<A3DSurfTorusData>::GetDefaultIndexerValue() const { return nullptr; }
+template<>
 const A3DEntity* TUniqueTSObj<A3DTess3DData>::GetDefaultIndexerValue() const { return nullptr; }
+template<>
 const A3DEntity* TUniqueTSObj<A3DTessBaseData>::GetDefaultIndexerValue() const { return nullptr; }
+template<>
 const A3DEntity* TUniqueTSObj<A3DTopoBodyData>::GetDefaultIndexerValue() const { return nullptr; }
+template<>
 const A3DEntity* TUniqueTSObj<A3DTopoBrepDataData>::GetDefaultIndexerValue() const { return nullptr; }
+template<>
 const A3DEntity* TUniqueTSObj<A3DTopoCoEdgeData>::GetDefaultIndexerValue() const { return nullptr; }
+template<>
 const A3DEntity* TUniqueTSObj<A3DTopoConnexData>::GetDefaultIndexerValue() const { return nullptr; }
+template<>
 const A3DEntity* TUniqueTSObj<A3DTopoContextData>::GetDefaultIndexerValue() const { return nullptr; }
+template<>
 const A3DEntity* TUniqueTSObj<A3DTopoEdgeData>::GetDefaultIndexerValue() const { return nullptr; }
+template<>
 const A3DEntity* TUniqueTSObj<A3DTopoFaceData>::GetDefaultIndexerValue() const { return nullptr; }
+template<>
 const A3DEntity* TUniqueTSObj<A3DTopoLoopData>::GetDefaultIndexerValue() const { return nullptr; }
+template<>
 const A3DEntity* TUniqueTSObj<A3DTopoShellData>::GetDefaultIndexerValue() const { return nullptr; }
+template<>
 const A3DEntity* TUniqueTSObj<A3DTopoUniqueVertexData>::GetDefaultIndexerValue() const { return nullptr; }
+template<>
 const A3DEntity* TUniqueTSObj<A3DTopoWireEdgeData>::GetDefaultIndexerValue() const { return nullptr; }
+template<>
 
 uint32 TUniqueTSObjFromIndex<A3DGraphMaterialData>::GetDefaultIndexerValue() const { return A3D_DEFAULT_MATERIAL_INDEX; }
+template<>
 uint32 TUniqueTSObjFromIndex<A3DGraphPictureData>::GetDefaultIndexerValue() const { return A3D_DEFAULT_PICTURE_INDEX; }
+template<>
 uint32 TUniqueTSObjFromIndex<A3DGraphRgbColorData>::GetDefaultIndexerValue() const { return A3D_DEFAULT_COLOR_INDEX; }
+template<>
 uint32 TUniqueTSObjFromIndex<A3DGraphStyleData>::GetDefaultIndexerValue() const { return A3D_DEFAULT_STYLE_INDEX; }
+template<>
 uint32 TUniqueTSObjFromIndex<A3DGraphTextureApplicationData>::GetDefaultIndexerValue() const { return A3D_DEFAULT_TEXTURE_APPLICATION_INDEX; }
+template<>
 uint32 TUniqueTSObjFromIndex<A3DGraphTextureDefinitionData>::GetDefaultIndexerValue() const { return A3D_DEFAULT_TEXTURE_DEFINITION_INDEX; }
-
+template<>
 const A3DEntity* TUniqueTSObj<A3DUTF8Char*>::GetDefaultIndexerValue() const { return nullptr; }
 
 #endif
diff --git a/Engine/Plugins/Enterprise/DatasmithCADImporter/Source/CADInterfaces/Private/TechSoftParserUtils.cpp b/Engine/Plugins/Enterprise/DatasmithCADImporter/Source/CADInterfaces/Private/TechSoftParserUtils.cpp
--- a/Engine/Plugins/Enterprise/DatasmithCADImporter/Source/CADInterfaces/Private/TechSoftParserUtils.cpp
+++ b/Engine/Plugins/Enterprise/DatasmithCADImporter/Source/CADInterfaces/Private/TechSoftParserUtils.cpp
@@ -619,7 +619,8 @@ namespace TechSoftParserUtils
 		if (UsedEntitiesFlags & kA3DTessFaceDataTriangleFan)
 		{
 			ensure(FaceTessData.m_uiSizesTriangulatedSize > FaceSetIndex);
-			uint32 LastFanIndex = 1 + FaceSetIndex + SizesTriangulated[FaceSetIndex++];
+			uint32 LastFanIndex = 1 + FaceSetIndex;
+			LastFanIndex += SizesTriangulated[FaceSetIndex++];
 			for (; FaceSetIndex < LastFanIndex; FaceSetIndex++)
 			{
 				uint32 FanSize = (FaceTessData.m_puiSizesTriangulated[FaceSetIndex] & kA3DTessFaceDataNormalMask);
diff --git a/Engine/Plugins/Enterprise/DatasmithCADImporter/Source/CADInterfaces/Private/TechSoftUtils.cpp b/Engine/Plugins/Enterprise/DatasmithCADImporter/Source/CADInterfaces/Private/TechSoftUtils.cpp
--- a/Engine/Plugins/Enterprise/DatasmithCADImporter/Source/CADInterfaces/Private/TechSoftUtils.cpp
+++ b/Engine/Plugins/Enterprise/DatasmithCADImporter/Source/CADInterfaces/Private/TechSoftUtils.cpp
@@ -7,6 +7,7 @@
 
 #include "CADOptions.h"
 
+#include "GenericPlatform/GenericPlatformMisc.h"
 #include "Math/Color.h"
 #include "Serialization/JsonSerializer.h"
 #include "Serialization/JsonWriter.h"
@@ -106,7 +107,7 @@ FUniqueTechSoftModelFile SaveBodiesToPrcFile(void** Bodies, uint32 BodyCount, co
 
 	TUniqueTSObj<A3DMiscSingleAttributeData> SingleAttributeData;
 	SingleAttributeData->m_eType = kA3DModellerAttributeTypeString;
-	SingleAttributeData->m_pcTitle = "MaterialTable";
+	SingleAttributeData->m_pcTitle = (char*)"MaterialTable";
 	SingleAttributeData->m_pcData = (char*)StringAnsi.c_str();
 
 	TUniqueTSObj<A3DMiscAttributeData> AttributesData;
@@ -154,7 +155,6 @@ bool FillBodyMesh(void* BodyPtr, const FImportParameters& ImportParameters, doub
 		return Extractor.FillBodyMesh(BodyMesh, FileUnit);
 	}
 
-
 	// TUniqueTechSoftObj does not work in this case
 	TUniqueTSObj<A3DRWParamsTessellationData> TessellationParameters;
 
@@ -333,8 +333,15 @@ void SaveModelFileToPrcFile(void* ModelFile, const FString& Filename)
 	ParamsExportData->m_bCompressBrep = false;
 	ParamsExportData->m_bCompressTessellation = false;
 
-	A3DUTF8Char HsfFileName[_MAX_PATH];
-	FCStringAnsi::Strncpy(HsfFileName, TCHAR_TO_UTF8(*Filename), _MAX_PATH);
+#if PLATFORM_WINDOWS
+	A3DUTF8Char HsfFileName[MAX_PATH];
+	FCStringAnsi::Strncpy(HsfFileName, TCHAR_TO_UTF8(*Filename), MAX_PATH);
+#elif PLATFORM_LINUX
+	A3DUTF8Char HsfFileName[PATH_MAX];
+	FCStringAnsi::Strncpy(HsfFileName, TCHAR_TO_UTF8(*Filename), PATH_MAX);
+#else
+#error Platform not supported
+#endif // PLATFORM_WINDOWS
 
 	TechSoftInterface::ExportModelFileToPrcFile(ModelFile, ParamsExportData.GetPtr(), HsfFileName, nullptr);
 #endif
@@ -575,6 +582,180 @@ A3DCrvNurbs* CreateTrimNurbsCurve(A3DCrvNurbs* CurveNurbsPtr, double UMin, doubl
 
 	return CADLibrary::TechSoftInterface::CreateCurveNurbs(*TransformedNurbsData);
 }
+
+void ExtractAttribute(const A3DMiscAttributeData& AttributeData, TMap<FString, FString>& OutMetaData)
+{
+	FString AttributeName;
+	if (AttributeData.m_bTitleIsInt)
+	{
+		A3DUns32 UnsignedValue = 0;
+		memcpy(&UnsignedValue, AttributeData.m_pcTitle, sizeof(A3DUns32));
+		AttributeName = FString::Printf(TEXT("%u"), UnsignedValue);
+	}
+	else if (AttributeData.m_pcTitle && AttributeData.m_pcTitle[0] != '\0')
+	{
+		AttributeName = UTF8_TO_TCHAR(AttributeData.m_pcTitle);
+	}
+
+	for (A3DUns32 Index = 0; Index < AttributeData.m_uiSize; ++Index)
+	{
+		FString AttributeValue;
+		switch (AttributeData.m_asSingleAttributesData[Index].m_eType)
+		{
+		case kA3DModellerAttributeTypeTime:
+		case kA3DModellerAttributeTypeInt:
+		{
+			A3DInt32 Value;
+			memcpy(&Value, AttributeData.m_asSingleAttributesData[Index].m_pcData, sizeof(A3DInt32));
+			AttributeValue = FString::Printf(TEXT("%d"), Value);
+			break;
+		}
+
+		case kA3DModellerAttributeTypeReal:
+		{
+			A3DDouble Value;
+			memcpy(&Value, AttributeData.m_asSingleAttributesData[Index].m_pcData, sizeof(A3DDouble));
+			AttributeValue = FString::Printf(TEXT("%f"), Value);
+			break;
+		}
+
+		case kA3DModellerAttributeTypeString:
+		{
+			if (AttributeData.m_asSingleAttributesData[Index].m_pcData && AttributeData.m_asSingleAttributesData[Index].m_pcData[0] != '\0')
+			{
+				AttributeValue = UTF8_TO_TCHAR(AttributeData.m_asSingleAttributesData[Index].m_pcData);
+			}
+			break;
+		}
+
+		default:
+			break;
+		}
+
+		if (AttributeName.Len())
+		{
+			if (Index)
+			{
+				OutMetaData.Emplace(FString::Printf(TEXT("%s_%u"), *AttributeName, (int32)Index), AttributeValue);
+			}
+			else
+			{
+				OutMetaData.Emplace(AttributeName, AttributeValue);
+			}
+		}
+	}
+}
+
+FString CleanSdkName(const FString& Name)
+{
+	int32 Index;
+	if (Name.FindLastChar(TEXT('['), Index))
+	{
+		return Name.Left(Index);
+	}
+	return Name;
+}
+
+FString CleanCatiaInstanceSdkName(const FString& Name)
+{
+	int32 Index;
+	if (Name.FindChar(TEXT('('), Index))
+	{
+		FString NewName = Name.RightChop(Index + 1);
+		if (NewName.FindLastChar(TEXT(')'), Index))
+		{
+			NewName = NewName.Left(Index);
+		}
+		return NewName;
+	}
+	return Name;
+}
+
+FString Clean3dxmlReferenceSdkName(const FString& Name)
+{
+	int32 Index;
+	if (Name.FindChar(TEXT('('), Index))
+	{
+		FString NewName = Name.Left(Index);
+		return NewName;
+	}
+	return Name;
+}
+
+FString CleanSwInstanceSdkName(const FString& Name)
+{
+	int32 Position;
+	if (Name.FindLastChar(TEXT('-'), Position))
+	{
+		FString NewName = Name.Left(Position) + TEXT("<") + Name.RightChop(Position + 1) + TEXT(">");
+		return NewName;
+	}
+	return Name;
+}
+
+FString CleanSwReferenceSdkName(const FString& Name)
+{
+	int32 Position;
+	if (Name.FindLastChar(TEXT('-'), Position))
+	{
+		FString NewName = Name.Left(Position);
+		return NewName;
+	}
+	return Name;
+}
+
+FString CleanCatiaReferenceName(const FString& Name)
+{
+	int32 Position;
+	if (Name.FindLastChar(TEXT('.'), Position))
+	{
+		FString Indice = Name.RightChop(Position + 1);
+		if (Indice.IsNumeric())
+		{
+			FString NewName = Name.Left(Position);
+			return NewName;
+		}
+	}
+	return Name;
+}
+
+FString CleanNameByRemoving_prt(const FString& Name)
+{
+	int32 Position;
+	if (Name.FindLastChar(TEXT('.'), Position))
+	{
+		FString Extension = Name.RightChop(Position);
+		if (Extension.Equals(TEXT("prt"), ESearchCase::IgnoreCase))
+		{
+			FString NewName = Name.Left(Position);
+			return NewName;
+		}
+	}
+	return Name;
+}
+
+bool CheckIfNameExists(TMap<FString, FString>& MetaData)
+{
+	FString* NamePtr = MetaData.Find(TEXT("Name"));
+	if (NamePtr != nullptr)
+	{
+		return true;
+	}
+	return false;
+}
+
+bool ReplaceOrAddNameValue(TMap<FString, FString>& MetaData, const TCHAR* Key)
+{
+	FString* NamePtr = MetaData.Find(Key);
+	if (NamePtr != nullptr)
+	{
+		FString& Name = MetaData.FindOrAdd(TEXT("Name"));
+		Name = *NamePtr;
+		return true;
+	}
+	return false;
+}
+
 #endif
 
 } // NS TechSoftUtils
diff --git a/Engine/Plugins/Enterprise/DatasmithCADImporter/Source/CADInterfaces/Public/CADFileParser.h b/Engine/Plugins/Enterprise/DatasmithCADImporter/Source/CADInterfaces/Public/CADFileParser.h
--- a/Engine/Plugins/Enterprise/DatasmithCADImporter/Source/CADInterfaces/Public/CADFileParser.h
+++ b/Engine/Plugins/Enterprise/DatasmithCADImporter/Source/CADInterfaces/Public/CADFileParser.h
@@ -5,6 +5,7 @@
 
 namespace CADLibrary
 {
+enum class ECADParsingResult : uint8;
 
 class ICADFileParser
 {
diff --git a/Engine/Plugins/Enterprise/DatasmithCADImporter/Source/CADInterfaces/Public/TUniqueTechSoftObj.h b/Engine/Plugins/Enterprise/DatasmithCADImporter/Source/CADInterfaces/Public/TUniqueTechSoftObj.h
--- a/Engine/Plugins/Enterprise/DatasmithCADImporter/Source/CADInterfaces/Public/TUniqueTechSoftObj.h
+++ b/Engine/Plugins/Enterprise/DatasmithCADImporter/Source/CADInterfaces/Public/TUniqueTechSoftObj.h
@@ -225,240 +225,471 @@ using TUniqueTSObjFromIndex = TUniqueTSObjBase<ObjectType, uint32>;
 
 // TUniqueTSObj -----------------------------------
 
+template<>
 CADINTERFACES_API void TUniqueTSObj<A3DAsmModelFileData>::InitializeData();
+template<>
 CADINTERFACES_API void TUniqueTSObj<A3DAsmPartDefinitionData>::InitializeData();
+template<>
 CADINTERFACES_API void TUniqueTSObj<A3DAsmProductOccurrenceData>::InitializeData();
+template<>
 CADINTERFACES_API void TUniqueTSObj<A3DAsmProductOccurrenceDataCV5>::InitializeData();
+template<>
 CADINTERFACES_API void TUniqueTSObj<A3DAsmProductOccurrenceDataSLW>::InitializeData();
+template<>
 CADINTERFACES_API void TUniqueTSObj<A3DAsmProductOccurrenceDataUg>::InitializeData();
+template<>
 CADINTERFACES_API void TUniqueTSObj<A3DBoundingBoxData>::InitializeData();
+template<>
+CADINTERFACES_API void TUniqueTSObj<A3DCopyAndAdaptBrepModelData>::InitializeData();
+template<>
 CADINTERFACES_API void TUniqueTSObj<A3DCrvCircleData>::InitializeData();
+template<>
 CADINTERFACES_API void TUniqueTSObj<A3DCrvCompositeData>::InitializeData();
+template<>
 CADINTERFACES_API void TUniqueTSObj<A3DCrvEllipseData>::InitializeData();
+template<>
 CADINTERFACES_API void TUniqueTSObj<A3DCrvHelixData>::InitializeData();
+template<>
 CADINTERFACES_API void TUniqueTSObj<A3DCrvHyperbolaData>::InitializeData();
+template<>
 CADINTERFACES_API void TUniqueTSObj<A3DCrvLineData>::InitializeData();
+template<>
 CADINTERFACES_API void TUniqueTSObj<A3DCrvNurbsData>::InitializeData();
+template<>
 CADINTERFACES_API void TUniqueTSObj<A3DCrvParabolaData>::InitializeData();
+template<>
 CADINTERFACES_API void TUniqueTSObj<A3DCrvPolyLineData>::InitializeData();
+template<>
 CADINTERFACES_API void TUniqueTSObj<A3DCrvTransformData>::InitializeData();
+template<>
 CADINTERFACES_API void TUniqueTSObj<A3DGlobalData>::InitializeData();
+template<>
 CADINTERFACES_API void TUniqueTSObj<A3DGraphicsData>::InitializeData();
+template<>
 CADINTERFACES_API void TUniqueTSObj<A3DIntervalData>::InitializeData();
+template<>
 CADINTERFACES_API void TUniqueTSObj<A3DMiscAttributeData>::InitializeData();
+template<>
 CADINTERFACES_API void TUniqueTSObj<A3DMiscCartesianTransformationData>::InitializeData();
+template<>
 CADINTERFACES_API void TUniqueTSObj<A3DMiscEntityReferenceData>::InitializeData();
+template<>
 CADINTERFACES_API void TUniqueTSObj<A3DMiscGeneralTransformationData>::InitializeData();
+template<>
 CADINTERFACES_API void TUniqueTSObj<A3DMiscMaterialPropertiesData>::InitializeData();
+template<>
 CADINTERFACES_API void TUniqueTSObj<A3DMiscReferenceOnCsysItemData>::InitializeData();
+template<>
 CADINTERFACES_API void TUniqueTSObj<A3DMiscReferenceOnTessData>::InitializeData();
+template<>
 CADINTERFACES_API void TUniqueTSObj<A3DMiscReferenceOnTopologyData>::InitializeData();
+template<>
 CADINTERFACES_API void TUniqueTSObj<A3DMiscSingleAttributeData>::InitializeData();
+template<>
 CADINTERFACES_API void TUniqueTSObj<A3DRWParamsExportPrcData>::InitializeData();
+template<>
 CADINTERFACES_API void TUniqueTSObj<A3DRiBrepModelData>::InitializeData();
+template<>
 CADINTERFACES_API void TUniqueTSObj<A3DRiCoordinateSystemData>::InitializeData();
+template<>
 CADINTERFACES_API void TUniqueTSObj<A3DRiDirectionData>::InitializeData();
+template<>
 CADINTERFACES_API void TUniqueTSObj<A3DRiPolyBrepModelData>::InitializeData();
+template<>
 CADINTERFACES_API void TUniqueTSObj<A3DRiRepresentationItemData>::InitializeData();
+template<>
 CADINTERFACES_API void TUniqueTSObj<A3DRiSetData>::InitializeData();
+template<>
 CADINTERFACES_API void TUniqueTSObj<A3DRootBaseData>::InitializeData();
+template<>
 CADINTERFACES_API void TUniqueTSObj<A3DRootBaseWithGraphicsData>::InitializeData();
+template<>
 CADINTERFACES_API void TUniqueTSObj<A3DRWParamsTessellationData>::InitializeData();
+template<>
 CADINTERFACES_API void TUniqueTSObj<A3DSewOptionsData>::InitializeData();
+template<>
 CADINTERFACES_API void TUniqueTSObj<A3DSurfBlend01Data>::InitializeData();
+template<>
 CADINTERFACES_API void TUniqueTSObj<A3DSurfBlend02Data>::InitializeData();
+template<>
 CADINTERFACES_API void TUniqueTSObj<A3DSurfBlend03Data>::InitializeData();
+template<>
 CADINTERFACES_API void TUniqueTSObj<A3DSurfConeData>::InitializeData();
+template<>
 CADINTERFACES_API void TUniqueTSObj<A3DSurfCylinderData>::InitializeData();
+template<>
 CADINTERFACES_API void TUniqueTSObj<A3DSurfCylindricalData>::InitializeData();
+template<>
 CADINTERFACES_API void TUniqueTSObj<A3DSurfExtrusionData>::InitializeData();
+template<>
 CADINTERFACES_API void TUniqueTSObj<A3DSurfFromCurvesData>::InitializeData();
+template<>
 CADINTERFACES_API void TUniqueTSObj<A3DSurfNurbsData>::InitializeData();
+template<>
 CADINTERFACES_API void TUniqueTSObj<A3DSurfPipeData>::InitializeData();
+template<>
 CADINTERFACES_API void TUniqueTSObj<A3DSurfPlaneData>::InitializeData();
+template<>
 CADINTERFACES_API void TUniqueTSObj<A3DSurfRevolutionData>::InitializeData();
+template<>
 CADINTERFACES_API void TUniqueTSObj<A3DSurfRuledData>::InitializeData();
+template<>
 CADINTERFACES_API void TUniqueTSObj<A3DSurfSphereData>::InitializeData();
+template<>
 CADINTERFACES_API void TUniqueTSObj<A3DSurfTorusData>::InitializeData();
+template<>
 CADINTERFACES_API void TUniqueTSObj<A3DTess3DData>::InitializeData();
+template<>
 CADINTERFACES_API void TUniqueTSObj<A3DTessBaseData>::InitializeData();
+template<>
 CADINTERFACES_API void TUniqueTSObj<A3DTopoBodyData>::InitializeData();
+template<>
 CADINTERFACES_API void TUniqueTSObj<A3DTopoBrepDataData>::InitializeData();
+template<>
 CADINTERFACES_API void TUniqueTSObj<A3DTopoCoEdgeData>::InitializeData();
+template<>
 CADINTERFACES_API void TUniqueTSObj<A3DTopoConnexData>::InitializeData();
+template<>
 CADINTERFACES_API void TUniqueTSObj<A3DTopoContextData>::InitializeData();
+template<>
 CADINTERFACES_API void TUniqueTSObj<A3DTopoEdgeData>::InitializeData();
+template<>
 CADINTERFACES_API void TUniqueTSObj<A3DTopoFaceData>::InitializeData();
+template<>
 CADINTERFACES_API void TUniqueTSObj<A3DTopoLoopData>::InitializeData();
+template<>
 CADINTERFACES_API void TUniqueTSObj<A3DTopoShellData>::InitializeData();
+template<>
 CADINTERFACES_API void TUniqueTSObj<A3DTopoUniqueVertexData>::InitializeData();
+template<>
 CADINTERFACES_API void TUniqueTSObj<A3DTopoWireEdgeData>::InitializeData();
 
+template<>
 CADINTERFACES_API A3DStatus TUniqueTSObj<A3DAsmModelFileData>::GetData(const A3DEntity* InEntityPtr);
+template<>
 CADINTERFACES_API A3DStatus TUniqueTSObj<A3DAsmPartDefinitionData>::GetData(const A3DEntity* InEntityPtr);
+template<>
 CADINTERFACES_API A3DStatus TUniqueTSObj<A3DAsmProductOccurrenceData>::GetData(const A3DEntity* InEntityPtr);
+template<>
 CADINTERFACES_API A3DStatus TUniqueTSObj<A3DAsmProductOccurrenceDataCV5>::GetData(const A3DEntity* InEntityPtr);
+template<>
 CADINTERFACES_API A3DStatus TUniqueTSObj<A3DAsmProductOccurrenceDataSLW>::GetData(const A3DEntity* InEntityPtr);
+template<>
 CADINTERFACES_API A3DStatus TUniqueTSObj<A3DAsmProductOccurrenceDataUg>::GetData(const A3DEntity* InEntityPtr);
+template<>
 CADINTERFACES_API A3DStatus TUniqueTSObj<A3DBoundingBoxData>::GetData(const A3DEntity* InEntityPtr);
+template<>
+CADINTERFACES_API A3DStatus TUniqueTSObj<A3DCopyAndAdaptBrepModelData>::GetData(const A3DEntity* InEntityPtr);
+template<>
 CADINTERFACES_API A3DStatus TUniqueTSObj<A3DCrvCircleData>::GetData(const A3DEntity* InEntityPtr);
+template<>
 CADINTERFACES_API A3DStatus TUniqueTSObj<A3DCrvCompositeData>::GetData(const A3DEntity* InEntityPtr);
+template<>
 CADINTERFACES_API A3DStatus TUniqueTSObj<A3DCrvEllipseData>::GetData(const A3DEntity* InEntityPtr);
+template<>
 CADINTERFACES_API A3DStatus TUniqueTSObj<A3DCrvHelixData>::GetData(const A3DEntity* InEntityPtr);
+template<>
 CADINTERFACES_API A3DStatus TUniqueTSObj<A3DCrvHyperbolaData>::GetData(const A3DEntity* InEntityPtr);
+template<>
 CADINTERFACES_API A3DStatus TUniqueTSObj<A3DCrvLineData>::GetData(const A3DEntity* InEntityPtr);
+template<>
 CADINTERFACES_API A3DStatus TUniqueTSObj<A3DCrvNurbsData>::GetData(const A3DEntity* InEntityPtr);
+template<>
 CADINTERFACES_API A3DStatus TUniqueTSObj<A3DCrvParabolaData>::GetData(const A3DEntity* InEntityPtr);
+template<>
 CADINTERFACES_API A3DStatus TUniqueTSObj<A3DCrvPolyLineData>::GetData(const A3DEntity* InEntityPtr);
+template<>
 CADINTERFACES_API A3DStatus TUniqueTSObj<A3DCrvTransformData>::GetData(const A3DEntity* InEntityPtr);
+template<>
 CADINTERFACES_API A3DStatus TUniqueTSObj<A3DGlobalData>::GetData(const A3DEntity* InEntityPtr);
+template<>
 CADINTERFACES_API A3DStatus TUniqueTSObj<A3DGraphicsData>::GetData(const A3DEntity* InEntityPtr);
+template<>
 CADINTERFACES_API A3DStatus TUniqueTSObj<A3DIntervalData>::GetData(const A3DEntity* InEntityPtr);
+template<>
 CADINTERFACES_API A3DStatus TUniqueTSObj<A3DMiscAttributeData>::GetData(const A3DEntity* InEntityPtr);
+template<>
 CADINTERFACES_API A3DStatus TUniqueTSObj<A3DMiscCartesianTransformationData>::GetData(const A3DEntity* InEntityPtr);
+template<>
 CADINTERFACES_API A3DStatus TUniqueTSObj<A3DMiscEntityReferenceData>::GetData(const A3DEntity* InEntityPtr);
+template<>
 CADINTERFACES_API A3DStatus TUniqueTSObj<A3DMiscGeneralTransformationData>::GetData(const A3DEntity* InEntityPtr);
+template<>
 CADINTERFACES_API A3DStatus TUniqueTSObj<A3DMiscMaterialPropertiesData>::GetData(const A3DEntity* InEntityPtr);
+template<>
 CADINTERFACES_API A3DStatus TUniqueTSObj<A3DMiscReferenceOnCsysItemData>::GetData(const A3DEntity* InEntityPtr);
+template<>
 CADINTERFACES_API A3DStatus TUniqueTSObj<A3DMiscReferenceOnTessData>::GetData(const A3DEntity* InEntityPtr);
+template<>
 CADINTERFACES_API A3DStatus TUniqueTSObj<A3DMiscReferenceOnTopologyData>::GetData(const A3DEntity* InEntityPtr);
+template<>
 CADINTERFACES_API A3DStatus TUniqueTSObj<A3DMiscSingleAttributeData>::GetData(const A3DEntity* InEntityPtr);
+template<>
 CADINTERFACES_API A3DStatus TUniqueTSObj<A3DRWParamsExportPrcData>::GetData(const A3DEntity* InEntityPtr);
+template<>
 CADINTERFACES_API A3DStatus TUniqueTSObj<A3DRiBrepModelData>::GetData(const A3DEntity* InEntityPtr);
+template<>
 CADINTERFACES_API A3DStatus TUniqueTSObj<A3DRiCoordinateSystemData>::GetData(const A3DEntity* InEntityPtr);
+template<>
 CADINTERFACES_API A3DStatus TUniqueTSObj<A3DRiDirectionData>::GetData(const A3DEntity* InEntityPtr);
+template<>
 CADINTERFACES_API A3DStatus TUniqueTSObj<A3DRiPolyBrepModelData>::GetData(const A3DEntity* InEntityPtr);
+template<>
 CADINTERFACES_API A3DStatus TUniqueTSObj<A3DRiRepresentationItemData>::GetData(const A3DEntity* InEntityPtr);
+template<>
 CADINTERFACES_API A3DStatus TUniqueTSObj<A3DRiSetData>::GetData(const A3DEntity* InEntityPtr);
+template<>
 CADINTERFACES_API A3DStatus TUniqueTSObj<A3DRootBaseData>::GetData(const A3DEntity* InEntityPtr);
+template<>
 CADINTERFACES_API A3DStatus TUniqueTSObj<A3DRootBaseWithGraphicsData>::GetData(const A3DEntity* InEntityPtr);
+template<>
 CADINTERFACES_API A3DStatus TUniqueTSObj<A3DRWParamsTessellationData>::GetData(const A3DEntity* InEntityPtr);
+template<>
 CADINTERFACES_API A3DStatus TUniqueTSObj<A3DSewOptionsData>::GetData(const A3DEntity* InEntityPtr);
+template<>
 CADINTERFACES_API A3DStatus TUniqueTSObj<A3DSurfBlend01Data>::GetData(const A3DEntity* InEntityPtr);
+template<>
 CADINTERFACES_API A3DStatus TUniqueTSObj<A3DSurfBlend02Data>::GetData(const A3DEntity* InEntityPtr);
+template<>
 CADINTERFACES_API A3DStatus TUniqueTSObj<A3DSurfBlend03Data>::GetData(const A3DEntity* InEntityPtr);
+template<>
 CADINTERFACES_API A3DStatus TUniqueTSObj<A3DSurfConeData>::GetData(const A3DEntity* InEntityPtr);
+template<>
 CADINTERFACES_API A3DStatus TUniqueTSObj<A3DSurfCylinderData>::GetData(const A3DEntity* InEntityPtr);
+template<>
 CADINTERFACES_API A3DStatus TUniqueTSObj<A3DSurfCylindricalData>::GetData(const A3DEntity* InEntityPtr);
+template<>
 CADINTERFACES_API A3DStatus TUniqueTSObj<A3DSurfExtrusionData>::GetData(const A3DEntity* InEntityPtr);
+template<>
 CADINTERFACES_API A3DStatus TUniqueTSObj<A3DSurfFromCurvesData>::GetData(const A3DEntity* InEntityPtr);
+template<>
 CADINTERFACES_API A3DStatus TUniqueTSObj<A3DSurfNurbsData>::GetData(const A3DEntity* InEntityPtr);
+template<>
 CADINTERFACES_API A3DStatus TUniqueTSObj<A3DSurfPipeData>::GetData(const A3DEntity* InEntityPtr);
+template<>
 CADINTERFACES_API A3DStatus TUniqueTSObj<A3DSurfPlaneData>::GetData(const A3DEntity* InEntityPtr);
+template<>
 CADINTERFACES_API A3DStatus TUniqueTSObj<A3DSurfRevolutionData>::GetData(const A3DEntity* InEntityPtr);
+template<>
 CADINTERFACES_API A3DStatus TUniqueTSObj<A3DSurfRuledData>::GetData(const A3DEntity* InEntityPtr);
+template<>
 CADINTERFACES_API A3DStatus TUniqueTSObj<A3DSurfSphereData>::GetData(const A3DEntity* InEntityPtr);
+template<>
 CADINTERFACES_API A3DStatus TUniqueTSObj<A3DSurfTorusData>::GetData(const A3DEntity* InEntityPtr);
+template<>
 CADINTERFACES_API A3DStatus TUniqueTSObj<A3DTess3DData>::GetData(const A3DEntity* InEntityPtr);
+template<>
 CADINTERFACES_API A3DStatus TUniqueTSObj<A3DTessBaseData>::GetData(const A3DEntity* InEntityPtr);
+template<>
 CADINTERFACES_API A3DStatus TUniqueTSObj<A3DTopoBodyData>::GetData(const A3DEntity* InEntityPtr);
+template<>
 CADINTERFACES_API A3DStatus TUniqueTSObj<A3DTopoBrepDataData>::GetData(const A3DEntity* InEntityPtr);
+template<>
 CADINTERFACES_API A3DStatus TUniqueTSObj<A3DTopoCoEdgeData>::GetData(const A3DEntity* InEntityPtr);
+template<>
 CADINTERFACES_API A3DStatus TUniqueTSObj<A3DTopoConnexData>::GetData(const A3DEntity* InEntityPtr);
+template<>
 CADINTERFACES_API A3DStatus TUniqueTSObj<A3DTopoContextData>::GetData(const A3DEntity* InEntityPtr);
+template<>
 CADINTERFACES_API A3DStatus TUniqueTSObj<A3DTopoEdgeData>::GetData(const A3DEntity* InEntityPtr);
+template<>
 CADINTERFACES_API A3DStatus TUniqueTSObj<A3DTopoFaceData>::GetData(const A3DEntity* InEntityPtr);
+template<>
 CADINTERFACES_API A3DStatus TUniqueTSObj<A3DTopoLoopData>::GetData(const A3DEntity* InEntityPtr);
+template<>
 CADINTERFACES_API A3DStatus TUniqueTSObj<A3DTopoShellData>::GetData(const A3DEntity* InEntityPtr);
+template<>
 CADINTERFACES_API A3DStatus TUniqueTSObj<A3DTopoUniqueVertexData>::GetData(const A3DEntity* InEntityPtr);
+template<>
 CADINTERFACES_API A3DStatus TUniqueTSObj<A3DTopoWireEdgeData>::GetData(const A3DEntity* InEntityPtr);
 
+template<>
 CADINTERFACES_API const A3DEntity* TUniqueTSObj<A3DAsmModelFileData>::GetDefaultIndexerValue() const;
+template<>
 CADINTERFACES_API const A3DEntity* TUniqueTSObj<A3DAsmPartDefinitionData>::GetDefaultIndexerValue() const;
+template<>
 CADINTERFACES_API const A3DEntity* TUniqueTSObj<A3DAsmProductOccurrenceData>::GetDefaultIndexerValue() const;
+template<>
 CADINTERFACES_API const A3DEntity* TUniqueTSObj<A3DAsmProductOccurrenceDataCV5>::GetDefaultIndexerValue() const;
+template<>
 CADINTERFACES_API const A3DEntity* TUniqueTSObj<A3DAsmProductOccurrenceDataSLW>::GetDefaultIndexerValue() const;
+template<>
 CADINTERFACES_API const A3DEntity* TUniqueTSObj<A3DAsmProductOccurrenceDataUg>::GetDefaultIndexerValue() const;
+template<>
 CADINTERFACES_API const A3DEntity* TUniqueTSObj<A3DBoundingBoxData>::GetDefaultIndexerValue() const;
+template<>
+CADINTERFACES_API const A3DEntity* TUniqueTSObj<A3DCopyAndAdaptBrepModelData>::GetDefaultIndexerValue() const;
+template<>
 CADINTERFACES_API const A3DEntity* TUniqueTSObj<A3DCrvCircleData>::GetDefaultIndexerValue() const;
+template<>
 CADINTERFACES_API const A3DEntity* TUniqueTSObj<A3DCrvCompositeData>::GetDefaultIndexerValue() const;
+template<>
 CADINTERFACES_API const A3DEntity* TUniqueTSObj<A3DCrvEllipseData>::GetDefaultIndexerValue() const;
+template<>
 CADINTERFACES_API const A3DEntity* TUniqueTSObj<A3DCrvHelixData>::GetDefaultIndexerValue() const;
+template<>
 CADINTERFACES_API const A3DEntity* TUniqueTSObj<A3DCrvHyperbolaData>::GetDefaultIndexerValue() const;
+template<>
 CADINTERFACES_API const A3DEntity* TUniqueTSObj<A3DCrvLineData>::GetDefaultIndexerValue() const;
+template<>
 CADINTERFACES_API const A3DEntity* TUniqueTSObj<A3DCrvNurbsData>::GetDefaultIndexerValue() const;
+template<>
 CADINTERFACES_API const A3DEntity* TUniqueTSObj<A3DCrvParabolaData>::GetDefaultIndexerValue() const;
+template<>
 CADINTERFACES_API const A3DEntity* TUniqueTSObj<A3DCrvPolyLineData>::GetDefaultIndexerValue() const;
+template<>
 CADINTERFACES_API const A3DEntity* TUniqueTSObj<A3DCrvTransformData>::GetDefaultIndexerValue() const;
+template<>
 CADINTERFACES_API const A3DEntity* TUniqueTSObj<A3DGlobalData>::GetDefaultIndexerValue() const;
+template<>
 CADINTERFACES_API const A3DEntity* TUniqueTSObj<A3DGraphicsData>::GetDefaultIndexerValue() const;
+template<>
 CADINTERFACES_API const A3DEntity* TUniqueTSObj<A3DIntervalData>::GetDefaultIndexerValue() const;
+template<>
 CADINTERFACES_API const A3DEntity* TUniqueTSObj<A3DMiscAttributeData>::GetDefaultIndexerValue() const;
+template<>
 CADINTERFACES_API const A3DEntity* TUniqueTSObj<A3DMiscCartesianTransformationData>::GetDefaultIndexerValue() const;
+template<>
 CADINTERFACES_API const A3DEntity* TUniqueTSObj<A3DMiscEntityReferenceData>::GetDefaultIndexerValue() const;
+template<>
 CADINTERFACES_API const A3DEntity* TUniqueTSObj<A3DMiscGeneralTransformationData>::GetDefaultIndexerValue() const;
+template<>
 CADINTERFACES_API const A3DEntity* TUniqueTSObj<A3DMiscMaterialPropertiesData>::GetDefaultIndexerValue() const;
+template<>
 CADINTERFACES_API const A3DEntity* TUniqueTSObj<A3DMiscReferenceOnCsysItemData>::GetDefaultIndexerValue() const;
+template<>
 CADINTERFACES_API const A3DEntity* TUniqueTSObj<A3DMiscReferenceOnTessData>::GetDefaultIndexerValue() const;
+template<>
 CADINTERFACES_API const A3DEntity* TUniqueTSObj<A3DMiscReferenceOnTopologyData>::GetDefaultIndexerValue() const;
+template<>
 CADINTERFACES_API const A3DEntity* TUniqueTSObj<A3DMiscSingleAttributeData>::GetDefaultIndexerValue() const;
+template<>
 CADINTERFACES_API const A3DEntity* TUniqueTSObj<A3DRWParamsExportPrcData>::GetDefaultIndexerValue() const;
+template<>
 CADINTERFACES_API const A3DEntity* TUniqueTSObj<A3DRiBrepModelData>::GetDefaultIndexerValue() const;
+template<>
 CADINTERFACES_API const A3DEntity* TUniqueTSObj<A3DRiCoordinateSystemData>::GetDefaultIndexerValue() const;
+template<>
 CADINTERFACES_API const A3DEntity* TUniqueTSObj<A3DRiDirectionData>::GetDefaultIndexerValue() const;
+template<>
 CADINTERFACES_API const A3DEntity* TUniqueTSObj<A3DRiPolyBrepModelData>::GetDefaultIndexerValue() const;
+template<>
 CADINTERFACES_API const A3DEntity* TUniqueTSObj<A3DRiRepresentationItemData>::GetDefaultIndexerValue() const;
+template<>
 CADINTERFACES_API const A3DEntity* TUniqueTSObj<A3DRiSetData>::GetDefaultIndexerValue() const;
+template<>
 CADINTERFACES_API const A3DEntity* TUniqueTSObj<A3DRootBaseData>::GetDefaultIndexerValue() const;
+template<>
 CADINTERFACES_API const A3DEntity* TUniqueTSObj<A3DRootBaseWithGraphicsData>::GetDefaultIndexerValue() const;
+template<>
 CADINTERFACES_API const A3DEntity* TUniqueTSObj<A3DRWParamsTessellationData>::GetDefaultIndexerValue() const;
+template<>
 CADINTERFACES_API const A3DEntity* TUniqueTSObj<A3DSewOptionsData>::GetDefaultIndexerValue() const;
+template<>
 CADINTERFACES_API const A3DEntity* TUniqueTSObj<A3DSurfBlend01Data>::GetDefaultIndexerValue() const;
+template<>
 CADINTERFACES_API const A3DEntity* TUniqueTSObj<A3DSurfBlend02Data>::GetDefaultIndexerValue() const;
+template<>
 CADINTERFACES_API const A3DEntity* TUniqueTSObj<A3DSurfBlend03Data>::GetDefaultIndexerValue() const;
+template<>
 CADINTERFACES_API const A3DEntity* TUniqueTSObj<A3DSurfConeData>::GetDefaultIndexerValue() const;
+template<>
 CADINTERFACES_API const A3DEntity* TUniqueTSObj<A3DSurfCylinderData>::GetDefaultIndexerValue() const;
+template<>
 CADINTERFACES_API const A3DEntity* TUniqueTSObj<A3DSurfCylindricalData>::GetDefaultIndexerValue() const;
+template<>
 CADINTERFACES_API const A3DEntity* TUniqueTSObj<A3DSurfExtrusionData>::GetDefaultIndexerValue() const;
+template<>
 CADINTERFACES_API const A3DEntity* TUniqueTSObj<A3DSurfFromCurvesData>::GetDefaultIndexerValue() const;
+template<>
 CADINTERFACES_API const A3DEntity* TUniqueTSObj<A3DSurfNurbsData>::GetDefaultIndexerValue() const;
+template<>
 CADINTERFACES_API const A3DEntity* TUniqueTSObj<A3DSurfPipeData>::GetDefaultIndexerValue() const;
+template<>
 CADINTERFACES_API const A3DEntity* TUniqueTSObj<A3DSurfPlaneData>::GetDefaultIndexerValue() const;
+template<>
 CADINTERFACES_API const A3DEntity* TUniqueTSObj<A3DSurfRevolutionData>::GetDefaultIndexerValue() const;
+template<>
 CADINTERFACES_API const A3DEntity* TUniqueTSObj<A3DSurfRuledData>::GetDefaultIndexerValue() const;
+template<>
 CADINTERFACES_API const A3DEntity* TUniqueTSObj<A3DSurfSphereData>::GetDefaultIndexerValue() const;
+template<>
 CADINTERFACES_API const A3DEntity* TUniqueTSObj<A3DSurfTorusData>::GetDefaultIndexerValue() const;
+template<>
 CADINTERFACES_API const A3DEntity* TUniqueTSObj<A3DTess3DData>::GetDefaultIndexerValue() const;
+template<>
 CADINTERFACES_API const A3DEntity* TUniqueTSObj<A3DTessBaseData>::GetDefaultIndexerValue() const;
+template<>
 CADINTERFACES_API const A3DEntity* TUniqueTSObj<A3DTopoBodyData>::GetDefaultIndexerValue() const;
+template<>
 CADINTERFACES_API const A3DEntity* TUniqueTSObj<A3DTopoBrepDataData>::GetDefaultIndexerValue() const;
+template<>
 CADINTERFACES_API const A3DEntity* TUniqueTSObj<A3DTopoCoEdgeData>::GetDefaultIndexerValue() const;
+template<>
 CADINTERFACES_API const A3DEntity* TUniqueTSObj<A3DTopoConnexData>::GetDefaultIndexerValue() const;
+template<>
 CADINTERFACES_API const A3DEntity* TUniqueTSObj<A3DTopoContextData>::GetDefaultIndexerValue() const;
+template<>
 CADINTERFACES_API const A3DEntity* TUniqueTSObj<A3DTopoEdgeData>::GetDefaultIndexerValue() const;
+template<>
 CADINTERFACES_API const A3DEntity* TUniqueTSObj<A3DTopoFaceData>::GetDefaultIndexerValue() const;
+template<>
 CADINTERFACES_API const A3DEntity* TUniqueTSObj<A3DTopoLoopData>::GetDefaultIndexerValue() const;
+template<>
 CADINTERFACES_API const A3DEntity* TUniqueTSObj<A3DTopoShellData>::GetDefaultIndexerValue() const;
+template<>
 CADINTERFACES_API const A3DEntity* TUniqueTSObj<A3DTopoUniqueVertexData>::GetDefaultIndexerValue() const;
+template<>
 CADINTERFACES_API const A3DEntity* TUniqueTSObj<A3DTopoWireEdgeData>::GetDefaultIndexerValue() const;
 
 // TUniqueTSObjFromIndex -----------------------------------
 
+template<>
 CADINTERFACES_API void TUniqueTSObj<A3DGraphMaterialData>::InitializeData();
+template<>
 CADINTERFACES_API void TUniqueTSObj<A3DGraphPictureData>::InitializeData();
+template<>
 CADINTERFACES_API void TUniqueTSObj<A3DGraphRgbColorData>::InitializeData();
+template<>
 CADINTERFACES_API void TUniqueTSObj<A3DGraphStyleData>::InitializeData();
+template<>
 CADINTERFACES_API void TUniqueTSObj<A3DGraphTextureApplicationData>::InitializeData();
+template<>
 CADINTERFACES_API void TUniqueTSObj<A3DGraphTextureDefinitionData>::InitializeData();
 
+template<>
 CADINTERFACES_API A3DStatus TUniqueTSObjFromIndex<A3DGraphMaterialData>::GetData(const uint32 InEntityIndex);
+template<>
 CADINTERFACES_API A3DStatus TUniqueTSObjFromIndex<A3DGraphPictureData>::GetData(const uint32 InEntityIndex);
+template<>
 CADINTERFACES_API A3DStatus TUniqueTSObjFromIndex<A3DGraphRgbColorData>::GetData(const uint32 InEntityIndex);
+template<>
 CADINTERFACES_API A3DStatus TUniqueTSObjFromIndex<A3DGraphStyleData>::GetData(const uint32 InEntityIndex);
+template<>
 CADINTERFACES_API A3DStatus TUniqueTSObjFromIndex<A3DGraphTextureApplicationData>::GetData(const uint32 InEntityIndex);
+template<>
 CADINTERFACES_API A3DStatus TUniqueTSObjFromIndex<A3DGraphTextureDefinitionData>::GetData(const uint32 InEntityIndex);
 
+template<>
 CADINTERFACES_API uint32 TUniqueTSObjFromIndex<A3DGraphMaterialData>::GetDefaultIndexerValue() const;
+template<>
 CADINTERFACES_API uint32 TUniqueTSObjFromIndex<A3DGraphPictureData>::GetDefaultIndexerValue() const;
+template<>
 CADINTERFACES_API uint32 TUniqueTSObjFromIndex<A3DGraphRgbColorData>::GetDefaultIndexerValue() const;
+template<>
 CADINTERFACES_API uint32 TUniqueTSObjFromIndex<A3DGraphStyleData>::GetDefaultIndexerValue() const;
+template<>
 CADINTERFACES_API uint32 TUniqueTSObjFromIndex<A3DGraphTextureApplicationData>::GetDefaultIndexerValue() const;
+template<>
 CADINTERFACES_API uint32 TUniqueTSObjFromIndex<A3DGraphTextureDefinitionData>::GetDefaultIndexerValue() const;
 
 // A3DUTF8Char* -----------------------------------
 
+template<>
 CADINTERFACES_API void TUniqueTSObj<A3DUTF8Char*>::InitializeData();
+template<>
 CADINTERFACES_API A3DStatus TUniqueTSObj<A3DUTF8Char*>::GetData(const A3DEntity* InEntityPtr);
+template<>
 CADINTERFACES_API const A3DEntity* TUniqueTSObj<A3DUTF8Char*>::GetDefaultIndexerValue() const;
 
 
diff --git a/Engine/Plugins/Enterprise/DatasmithCADImporter/Source/CADInterfaces/Public/TechSoftInterface.h b/Engine/Plugins/Enterprise/DatasmithCADImporter/Source/CADInterfaces/Public/TechSoftInterface.h
--- a/Engine/Plugins/Enterprise/DatasmithCADImporter/Source/CADInterfaces/Public/TechSoftInterface.h
+++ b/Engine/Plugins/Enterprise/DatasmithCADImporter/Source/CADInterfaces/Public/TechSoftInterface.h
@@ -71,12 +71,6 @@ private:
 
 namespace TechSoftInterface
 {
-
-/**
- * Single-ownership smart TeshSoft file object
- * This object ensure to delete the FileModel when it no more used
- */
-
 CADINTERFACES_API bool TECHSOFT_InitializeKernel(const TCHAR* = TEXT(""));
 
 CADINTERFACES_API FString GetTechSoftVersion();
diff --git a/Engine/Plugins/Enterprise/DatasmithCADImporter/Source/CADLibrary/Private/MeshDescriptionHelper.cpp b/Engine/Plugins/Enterprise/DatasmithCADImporter/Source/CADLibrary/Private/MeshDescriptionHelper.cpp
--- a/Engine/Plugins/Enterprise/DatasmithCADImporter/Source/CADLibrary/Private/MeshDescriptionHelper.cpp
+++ b/Engine/Plugins/Enterprise/DatasmithCADImporter/Source/CADLibrary/Private/MeshDescriptionHelper.cpp
@@ -7,7 +7,7 @@
 #include "DatasmithMaterialElements.h"
 #include "DatasmithSceneFactory.h"
 #include "DatasmithUtils.h"
-#include "Math.h"
+#include "Math/UnrealMathUtility.h"
 #include "MeshDescription.h"
 #include "MeshOperator.h"
 #include "Misc/FileHelper.h"
diff --git a/Engine/Plugins/Enterprise/DatasmithCADImporter/Source/DatasmithDispatcher/Private/DatasmithDispatcher.cpp b/Engine/Plugins/Enterprise/DatasmithCADImporter/Source/DatasmithDispatcher/Private/DatasmithDispatcher.cpp
--- a/Engine/Plugins/Enterprise/DatasmithCADImporter/Source/DatasmithDispatcher/Private/DatasmithDispatcher.cpp
+++ b/Engine/Plugins/Enterprise/DatasmithCADImporter/Source/DatasmithDispatcher/Private/DatasmithDispatcher.cpp
@@ -169,7 +169,7 @@ void FDatasmithDispatcher::Process(bool bWithProcessor)
 				break;
 			}
 
-			FWindowsPlatformProcess::Sleep(0.1);
+			FPlatformProcess::Sleep(0.1);
 		}
 
 		CloseHandlers();
diff --git a/Engine/Plugins/Enterprise/DatasmithCADImporter/Source/DatasmithDispatcher/Public/DatasmithWorkerHandler.h b/Engine/Plugins/Enterprise/DatasmithCADImporter/Source/DatasmithDispatcher/Public/DatasmithWorkerHandler.h
--- a/Engine/Plugins/Enterprise/DatasmithCADImporter/Source/DatasmithDispatcher/Public/DatasmithWorkerHandler.h
+++ b/Engine/Plugins/Enterprise/DatasmithCADImporter/Source/DatasmithDispatcher/Public/DatasmithWorkerHandler.h
@@ -19,7 +19,7 @@ namespace DatasmithDispatcher
 class FDatasmithDispatcher;
 
 //Handle a Worker by socket communication
-class FDatasmithWorkerHandler
+class DATASMITHDISPATCHER_API FDatasmithWorkerHandler
 {
 	enum class EWorkerState
 	{
