Binaries for new CAD library

diff --git a/Engine/Plugins/Enterprise/DatasmithCADImporter/Source/CADTools/Private/CADOption.cpp b/Engine/Plugins/Enterprise/DatasmithCADImporter/Source/CADTools/Private/CADOption.cpp
--- a/Engine/Plugins/Enterprise/DatasmithCADImporter/Source/CADTools/Private/CADOption.cpp
+++ b/Engine/Plugins/Enterprise/DatasmithCADImporter/Source/CADTools/Private/CADOption.cpp
@@ -4,12 +4,15 @@
 namespace CADLibrary
 {
 
-FString FImportParameters::GCADLibrary = TEXT("KernelIO");
+FString FImportParameters::GCADLibrary = TEXT("TechSoft");
+
+#ifdef USE_KERNEL_IO_SDK 
 FAutoConsoleVariableRef GCADTranslatorLibrary(
 	TEXT("ds.CADTranslator.CADLibrary"),
 	FImportParameters::GCADLibrary,
 	TEXT("Define the used library to import CAD file i.e. \"TechSoft\" or \"KernelIO\".\n"),
 	ECVF_Default);
+#endif
 
 int32 GMaxImportThreads = 0;
 FAutoConsoleVariableRef GCADTranslatorMaxImportThreads(
