Simplification of the implementation of FCADKernelTools::GetBodyTessellation.

diff --git a/Engine/Plugins/Enterprise/DatasmithCADImporter/Source/CADInterfaces/Private/CoreTechFileParserImpl.cpp b/Engine/Plugins/Enterprise/DatasmithCADImporter/Source/CADInterfaces/Private/CoreTechFileParserImpl.cpp
--- a/Engine/Plugins/Enterprise/DatasmithCADImporter/Source/CADInterfaces/Private/CoreTechFileParserImpl.cpp
+++ b/Engine/Plugins/Enterprise/DatasmithCADImporter/Source/CADInterfaces/Private/CoreTechFileParserImpl.cpp
@@ -276,7 +276,7 @@ namespace CADLibrary
 				CT_LIST_IO ObjectList;
 				ObjectList.PushBack(MainId);
 
-				CT_IO_ERROR SaveResult = CT_KERNEL_IO::SaveFile(ObjectList, *CacheFilePath, L"Ct");
+				CT_IO_ERROR SaveResult = CT_KERNEL_IO::SaveFile(ObjectList, *CacheFilePath, TEXT("Ct"));
 			}
 		}
 
@@ -660,7 +660,7 @@ namespace CADLibrary
 			CT_LIST_IO ObjectList;
 			ObjectList.PushBack(BodyId);
 			FString BodyFile = CADFileData.GetBodyCachePath(Body.MeshActorName);
-			CT_KERNEL_IO::SaveFile(ObjectList, *BodyFile, L"Ct");
+			CT_KERNEL_IO::SaveFile(ObjectList, *BodyFile, TEXT("Ct"));
 		}
 
 		FObjectDisplayDataId BodyMaterial;
@@ -759,13 +759,7 @@ namespace CADLibrary
 			CADKernel::FParametricMesher Mesher(*CADKernelModelMesh);
 			Mesher.MeshEntity(CADKernelModel);
 
-			TFunction<void(FObjectDisplayDataId, FObjectDisplayDataId, int32)> ProcessFace;
-			ProcessFace = [&](FObjectDisplayDataId FaceMaterial, FObjectDisplayDataId BodyMaterial, int32 Index)
-			{
-				SetFaceMainMaterial(FaceMaterial, BodyMaterial, BodyMesh, Index);
-			};
-
-			FCADKernelTools::GetBodyTessellation(*CADKernelModelMesh, *CADKernelBody, BodyMesh, DefaultMaterialHash, ProcessFace);
+			FCADKernelTools::GetBodyTessellation(*CADKernelModelMesh, *CADKernelBody, BodyMesh);
 
 			CADKernelSession.Clear();
 		}
@@ -808,13 +802,7 @@ namespace CADLibrary
 		CADKernel::FParametricMesher Mesher(*CADKernelModelMesh);
 		Mesher.MeshEntity(CADKernelBody);
 
-		TFunction<void(FObjectDisplayDataId, FObjectDisplayDataId, int32)> ProcessFace;
-		ProcessFace = [&](FObjectDisplayDataId FaceMaterial, FObjectDisplayDataId BodyMaterial, int32 Index)
-		{
-			SetFaceMainMaterial(FaceMaterial, BodyMaterial, BodyMesh, Index);
-		};
-
-		FCADKernelTools::GetBodyTessellation(CADKernelModelMesh.Get(), CADKernelBody, BodyMesh, DefaultMaterialHash, ProcessFace);
+		FCADKernelTools::GetBodyTessellation(CADKernelModelMesh.Get(), CADKernelBody, BodyMesh);
 
 		ArchiveBody.ColorFaceSet = BodyMesh.ColorSet;
 		ArchiveBody.MaterialFaceSet = BodyMesh.MaterialSet;
diff --git a/Engine/Plugins/Enterprise/DatasmithCADImporter/Source/CADLibrary/Private/CADKernelTools.cpp b/Engine/Plugins/Enterprise/DatasmithCADImporter/Source/CADLibrary/Private/CADKernelTools.cpp
--- a/Engine/Plugins/Enterprise/DatasmithCADImporter/Source/CADLibrary/Private/CADKernelTools.cpp
+++ b/Engine/Plugins/Enterprise/DatasmithCADImporter/Source/CADLibrary/Private/CADKernelTools.cpp
@@ -292,7 +292,7 @@ namespace CADLibrary
 	}
 
 
-	uint32 FCADKernelTools::GetFaceTessellation(CADKernel::FFaceMesh& FaceMesh, FBodyMesh& OutBodyMesh)
+	uint32 FCADKernelTools::GetFaceTessellation(CADKernel::FFaceMesh& FaceMesh, FBodyMesh& OutBodyMesh, FObjectDisplayDataId FaceMaterial)
 	{
 		// Something wrong happened, either an error or no data to collect
 		if (FaceMesh.TrianglesVerticesIndex.Num() == 0)
@@ -311,16 +311,34 @@ namespace CADLibrary
 		Tessellation.NormalArray = MoveTemp(FaceMesh.Normals);
 		Tessellation.TexCoordArray = MoveTemp(FaceMesh.UVMap);
 
+		if (FaceMaterial.Color)
+		{
+			Tessellation.ColorName = FaceMaterial.Color;
+			OutBodyMesh.ColorSet.Add(FaceMaterial.Color);
+		}
+
+		if (FaceMaterial.Material)
+		{
+			Tessellation.MaterialName = FaceMaterial.Material;
+			OutBodyMesh.MaterialSet.Add(FaceMaterial.Material);
+		}
+
 		return Tessellation.VertexIndices.Num() / 3;
 	}
 
 	void GetDisplayDataIds(const CADKernel::FTopologicalShapeEntity& ShapeEntity, FObjectDisplayDataId& DisplayDataId)
 	{
-		DisplayDataId.Color = ShapeEntity.GetColorId();
-		DisplayDataId.Material = ShapeEntity.GetMaterialId();
+		if (ShapeEntity.GetColorId() != 0)
+		{
+			DisplayDataId.Color = ShapeEntity.GetColorId();
+		}
+		if (ShapeEntity.GetMaterialId() != 0)
+		{
+			DisplayDataId.Material = ShapeEntity.GetMaterialId();
+		}
 	}
 
-	void FCADKernelTools::GetBodyTessellation(const CADKernel::FModelMesh& ModelMesh, const CADKernel::FBody& Body, FBodyMesh& OutBodyMesh, uint32 DefaultMaterialHash, TFunction<void(FObjectDisplayDataId, FObjectDisplayDataId, int32)> SetFaceMainMaterial)
+	void FCADKernelTools::GetBodyTessellation(const CADKernel::FModelMesh& ModelMesh, const CADKernel::FBody& Body, FBodyMesh& OutBodyMesh)
 	{
 		ModelMesh.GetNodeCoordinates(OutBodyMesh.VertexArray);
 
@@ -332,12 +350,11 @@ namespace CADLibrary
 		OutBodyMesh.MaterialSet.Reserve(FaceSize);
 
 		FObjectDisplayDataId BodyMaterial;
-		BodyMaterial.DefaultMaterialName = DefaultMaterialHash;
 
 		GetDisplayDataIds(Body, BodyMaterial);
+		BodyMaterial.DefaultMaterialName = BodyMaterial.Material != 0 ? BodyMaterial.Material : BodyMaterial.Color;
 
 		// Loop through the face of bodies and collect all tessellation data
-		int32 FaceIndex = 0;
 		for (const TSharedPtr<CADKernel::FShell>& Shell : Body.GetShells())
 		{
 			if (!Shell.IsValid())
@@ -360,23 +377,11 @@ namespace CADLibrary
 					continue;
 				}
 
-				FObjectDisplayDataId FaceMaterial;
+				FObjectDisplayDataId FaceMaterial = ShellMaterial;
 				GetDisplayDataIds(*Face.Entity, FaceMaterial);
 
-				uint32 TriangleNum = GetFaceTessellation(*Face.Entity->GetMesh(), OutBodyMesh);
-
-				if (TriangleNum == 0)
-				{
-					continue;
-				}
-
+				uint32 TriangleNum = GetFaceTessellation(*Face.Entity->GetMesh(), OutBodyMesh, FaceMaterial);
 				OutBodyMesh.TriangleCount += TriangleNum;
-
-				if (SetFaceMainMaterial)
-				{
-					SetFaceMainMaterial(FaceMaterial, ShellMaterial, FaceIndex);
-				}
-				FaceIndex++;
 			}
 		}
 	}
diff --git a/Engine/Plugins/Enterprise/DatasmithCADImporter/Source/CADLibrary/Public/CADKernelTools.h b/Engine/Plugins/Enterprise/DatasmithCADImporter/Source/CADLibrary/Public/CADKernelTools.h
--- a/Engine/Plugins/Enterprise/DatasmithCADImporter/Source/CADLibrary/Public/CADKernelTools.h
+++ b/Engine/Plugins/Enterprise/DatasmithCADImporter/Source/CADLibrary/Public/CADKernelTools.h
@@ -27,13 +27,13 @@ namespace CADLibrary
 	public:
 		static void DefineMeshCriteria(CADKernel::FModelMesh& MeshModel, const FImportParameters& ImportParameters, double GeometricTolerance);
 
-		static void GetBodyTessellation(const CADKernel::FModelMesh& ModelMesh, const CADKernel::FBody& Body, FBodyMesh& OutBodyMesh, uint32 DefaultMaterialHash, TFunction<void(FObjectDisplayDataId, FObjectDisplayDataId, int32)> SetFaceMainMaterial);
+		static void GetBodyTessellation(const CADKernel::FModelMesh& ModelMesh, const CADKernel::FBody& Body, FBodyMesh& OutBodyMesh);
 
 		/**
 		 * Tessellate a CADKernel entity and update the MeshDescription
 		 */
 		static bool Tessellate(CADKernel::FTopologicalShapeEntity& InCADKernelEntity, const FImportParameters& ImportParameters, const FMeshParameters& MeshParameters, FMeshDescription& OutMeshDescription);
 
-		static uint32 GetFaceTessellation(CADKernel::FFaceMesh& FaceMesh, FBodyMesh& OutBodyMesh);
+		static uint32 GetFaceTessellation(CADKernel::FFaceMesh& FaceMesh, FBodyMesh& OutBodyMesh, FObjectDisplayDataId FaceMaterial);
 	};
 }
