Fixed problem of normals on symmetric part
The problem came from a mistake in the distribution of the normal for symmetric faces.

diff --git a/Engine/Plugins/Enterprise/DatasmithCADImporter/Source/CADLibrary/Private/MeshDescriptionHelper.cpp b/Engine/Plugins/Enterprise/DatasmithCADImporter/Source/CADLibrary/Private/MeshDescriptionHelper.cpp
--- a/Engine/Plugins/Enterprise/DatasmithCADImporter/Source/CADLibrary/Private/MeshDescriptionHelper.cpp
+++ b/Engine/Plugins/Enterprise/DatasmithCADImporter/Source/CADLibrary/Private/MeshDescriptionHelper.cpp
@@ -513,7 +513,7 @@ namespace CADLibrary
 						for (int32 Index = 0; Index < 3; Index++)
 						{
 							const FVertexInstanceID VertexInstanceID = MeshVertexInstanceIDs[IndexFace + Orientation[Index]];
-							VertexInstanceNormals[VertexInstanceID] = (FVector3f)Tessellation.NormalArray[NewFaceIndex[IndexFace + Index]];
+							VertexInstanceNormals[VertexInstanceID] = (FVector3f)Tessellation.NormalArray[NewFaceIndex[IndexFace + Orientation[Index]]];
 						}
 					}
 				}
