Step 2: Preparing to plug of CADKernel into TechSoft: Add missing methodes in TechSoftInterface, TechSoftUtils and CADKernel::FSession

diff --git a/Engine/Plugins/Enterprise/DatasmithCADImporter/Source/CADInterfaces/Private/TechSoftFileParser.cpp b/Engine/Plugins/Enterprise/DatasmithCADImporter/Source/CADInterfaces/Private/TechSoftFileParser.cpp
--- a/Engine/Plugins/Enterprise/DatasmithCADImporter/Source/CADInterfaces/Private/TechSoftFileParser.cpp
+++ b/Engine/Plugins/Enterprise/DatasmithCADImporter/Source/CADInterfaces/Private/TechSoftFileParser.cpp
@@ -3,11 +3,13 @@
 #include "TechSoftFileParser.h"
 
 #include "CADOptions.h"
+#include "CADFileData.h"
 
 #ifdef USE_TECHSOFT_SDK
 
 #include "TechSoftInterface.h"
 #include "TechSoftUtils.h"
+#include "TechSoftUtilsPrivate.h"
 
 #include "Serialization/JsonSerializer.h"
 #include "Serialization/JsonWriter.h"
@@ -34,94 +36,6 @@ inline void RemoveUnwantedChar(FString& StringToClean, TCHAR UnwantedChar)
 	StringToClean = MoveTemp(NewString);
 }
 
-FString CleanSdkName(const FString& Name)
-{
-	int32 Index;
-	if (Name.FindLastChar(TEXT('['), Index))
-	{
-		return Name.Left(Index);
-	}
-	return Name;
-}
-
-FString CleanCatiaInstanceSdkName(const FString& Name)
-{
-	int32 Index;
-	if (Name.FindChar(TEXT('('), Index))
-	{
-		FString NewName = Name.RightChop(Index + 1);
-		if (NewName.FindLastChar(TEXT(')'), Index))
-		{
-			NewName = NewName.Left(Index);
-		}
-		return NewName;
-	}
-	return Name;
-}
-
-FString Clean3dxmlReferenceSdkName(const FString& Name)
-{
-	int32 Index;
-	if (Name.FindChar(TEXT('('), Index))
-	{
-		FString NewName = Name.Left(Index);
-		return NewName;
-	}
-	return Name;
-}
-
-FString CleanSwInstanceSdkName(const FString& Name)
-{
-	int32 Position;
-	if (Name.FindLastChar(TEXT('-'), Position))
-	{
-		FString NewName = Name.Left(Position) + TEXT("<") + Name.RightChop(Position + 1) + TEXT(">");
-		return NewName;
-	}
-	return Name;
-}
-
-FString CleanSwReferenceSdkName(const FString& Name)
-{
-	int32 Position;
-	if (Name.FindLastChar(TEXT('-'), Position))
-	{
-		FString NewName = Name.Left(Position);
-		return NewName;
-	}
-	return Name;
-}
-
-FString CleanCatiaReferenceName(const FString& Name)
-{
-	int32 Position;
-	if (Name.FindLastChar(TEXT('.'), Position))
-	{
-		FString Indice = Name.RightChop(Position + 1);
-		if (Indice.IsNumeric())
-		{
-			FString NewName = Name.Left(Position);
-			return NewName;
-		}
-	}
-	return Name;
-}
-
-FString CleanNameByRemoving_prt(const FString& Name)
-{
-	int32 Position;
-	if (Name.FindLastChar(TEXT('.'), Position))
-	{
-		FString Extension = Name.RightChop(Position);
-		if (Extension.Equals(TEXT("prt"), ESearchCase::IgnoreCase))
-		{
-			FString NewName = Name.Left(Position);
-			return NewName;
-		}
-	}
-	return Name;
-}
-
 bool CheckIfNameExists(TMap<FString, FString>& MetaData)
 {
 	FString* NamePtr = MetaData.Find(TEXT("Name"));
@@ -1281,7 +1195,7 @@ void FTechSoftFileParser::ExtractMetaData(const A3DEntity* Entity, FEntityMetaDa
 		if (MetaData->m_pcName && MetaData->m_pcName[0] != '\0')
 		{
 			FString SDKName = UTF8_TO_TCHAR(MetaData->m_pcName);
-			SDKName = TechSoftFileParserImpl::CleanSdkName(SDKName);
+			SDKName = TechSoftUtils::CleanSdkName(SDKName);
 			OutMetaData.MetaData.Emplace(TEXT("SDKName"), SDKName);
 		}
 
@@ -1323,7 +1237,7 @@ void FTechSoftFileParser::BuildReferenceName(TMap<FString, FString>& MetaData)
 		Name = *NamePtr;
 		if (Format == ECADFormat::CATIA)
 		{
-			Name = TechSoftFileParserImpl::CleanCatiaReferenceName(Name);
+			Name = TechSoftUtils::CleanCatiaReferenceName(Name);
 		}
 		return;
 	}
@@ -1354,11 +1268,11 @@ void FTechSoftFileParser::BuildReferenceName(TMap<FString, FString>& MetaData)
 		switch (Format)
 		{
 		case ECADFormat::CATIA_3DXML:
-			SdkName = TechSoftFileParserImpl::Clean3dxmlReferenceSdkName(SdkName);
+			SdkName = TechSoftUtils::Clean3dxmlReferenceSdkName(SdkName);
 			break;
 
 		case ECADFormat::SOLIDWORKS:
-			SdkName = TechSoftFileParserImpl::CleanSwReferenceSdkName(SdkName);
+			SdkName = TechSoftUtils::CleanSwReferenceSdkName(SdkName);
 			break;
 
 		default:
@@ -1397,11 +1311,11 @@ void FTechSoftFileParser::BuildInstanceName(TMap<FString, FString>& MetaData)
 		{
 		case ECADFormat::CATIA:
 		case ECADFormat::CATIA_3DXML:
-			SdkName = TechSoftFileParserImpl::CleanCatiaInstanceSdkName(SdkName);
+			SdkName = TechSoftUtils::CleanCatiaInstanceSdkName(SdkName);
 			break;
 
 		case ECADFormat::SOLIDWORKS:
-			SdkName = TechSoftFileParserImpl::CleanSwInstanceSdkName(SdkName);
+			SdkName = TechSoftUtils::CleanSwInstanceSdkName(SdkName);
 			break;
 
 		default:
@@ -1461,7 +1375,7 @@ void FTechSoftFileParser::BuildBodyName(TMap<FString, FString>& MetaData)
 		FString SdkName = *NamePtr;
 		if (Format == ECADFormat::CREO)
 		{
-			SdkName = TechSoftFileParserImpl::CleanNameByRemoving_prt(SdkName);
+			SdkName = TechSoftUtils::CleanNameByRemoving_prt(SdkName);
 		}
 
 		FString& Name = MetaData.FindOrAdd(TEXT("Name"));
diff --git a/Engine/Plugins/Enterprise/DatasmithCADImporter/Source/CADInterfaces/Private/TechSoftFileParser.h b/Engine/Plugins/Enterprise/DatasmithCADImporter/Source/CADInterfaces/Private/TechSoftFileParser.h
--- a/Engine/Plugins/Enterprise/DatasmithCADImporter/Source/CADInterfaces/Private/TechSoftFileParser.h
+++ b/Engine/Plugins/Enterprise/DatasmithCADImporter/Source/CADInterfaces/Private/TechSoftFileParser.h
@@ -4,11 +4,7 @@
 #include "CoreMinimal.h"
 
 #include "CADData.h"
-#include "CADFileData.h"
 #include "CADFileParser.h"
-#include "CADOptions.h"
-#include "CADSceneGraph.h"
-
 #include "TechSoftInterface.h"
 #include "TUniqueTechSoftObj.h"
 
@@ -33,8 +29,16 @@ typedef void A3DTopoBrepData;
 namespace CADLibrary
 {
 
-class FEntityData;
-class FComponentData;
+class FArchiveBody;
+class FArchiveComponent;
+class FArchiveColor;
+class FArchiveInstance;
+class FArchiveMaterial;
+class FArchiveUnloadedComponent;
+class FCADFileData;
+class FTechSoftInterface;
+
+struct FEntityMetaData;
 
 enum EComponentType : uint32
 {
@@ -46,75 +50,6 @@ enum EComponentType : uint32
 	LastType
 };
 
-enum EModellerType : uint32
-{
-	ModellerUnknown = 0,              /*!< User modeller. */
-	ModellerCatia = 2,                /*!< CATIA modeller. */
-	ModellerCatiaV5 = 3,              /*!< CATIA V5 modeller. */
-	ModellerCadds = 4,                /*!< CADDS modeller. */
-	ModellerUnigraphics = 5,          /*!< Unigraphics modeller. */
-	ModellerParasolid = 6,            /*!< Parasolid modeller. */
-	ModellerEuclid = 7,               /*!< Euclid modeller. */
-	ModellerIges = 9,                 /*!< IGES modeller. */
-	ModellerUnisurf = 10,             /*!< Unisurf modeller. */
-	ModellerVda = 11,                 /*!< VDA modeller. */
-	ModellerStl = 12,                 /*!< STL modeller. */
-	ModellerWrl = 13,                 /*!< WRL modeller. */
-	ModellerDxf = 14,                 /*!< DXF modeller. */
-	ModellerAcis = 15,                /*!< ACIS modeller. */
-	ModellerProE = 16,                /*!< Pro/E modeller. */
-	ModellerStep = 18,                /*!< STEP modeller. */
-	ModellerIdeas = 19,               /*!< I-DEAS modeller. */
-	ModellerJt = 20,                  /*!< JT modeller. */
-	ModellerSlw = 22,                 /*!< SolidWorks modeller. */
-	ModellerCgr = 23,                 /*!< CGR modeller. */
-	ModellerPrc = 24,                 /*!< PRC modeller. */
-	ModellerXvl = 25,                 /*!< XVL modeller. */
-	ModellerHpgl = 26,                /*!< HPGL modeller. */
-	ModellerTopSolid = 27,            /*!< TopSolid modeller. */
-	ModellerOneSpaceDesigner = 28,    /*!< OneSpace designer modeller. */
-	Modeller3dxml = 29,               /*!< 3DXML modeller. */
-	ModellerInventor = 30,            /*!< Inventor modeller. */
-	ModellerPostScript = 31,          /*!< Postscript modeller. */
-	ModellerPDF = 32,                 /*!< PDF modeller. */
-	ModellerU3D = 33,                 /*!< U3D modeller. */
-	ModellerIFC = 34,                 /*!< IFC modeller. */
-	ModellerDWG = 35,                 /*!< DWG modeller. */
-	ModellerDWF = 36,                 /*!< DWF modeller. */
-	ModellerSE = 37,                  /*!< SolidEdge modeller. */
-	ModellerOBJ = 38,                 /*!< OBJ modeller. */
-	ModellerKMZ = 39,                 /*!< KMZ modeller. */
-	ModellerDAE = 40,                 /*!< COLLADA modeller. */
-	Modeller3DS = 41,                 /*!< 3DS modeller. */
-	ModellerRhino = 43,               /*!< Rhino modeller. */
-	ModellerXML = 44,                 /*!< XML modeller. */
-	Modeller3mf = 45,                 /*!< 3MF modeller. */
-	ModellerScs = 46,                 /*!< SCS modeller. */
-	Modeller3dHtml = 47,              /*!< 3DHTML modeller. */
-	ModellerHsf = 48,                 /*!< Hsf modeller. */
-	ModellerGltf = 49,                /*!< GL modeller. */
-	ModellerRevit = 50,               /*!< Revit modeller. */
-	ModellerFBX = 51,                 /*!< FBX modeller. */
-	ModellerStepXML = 52,             /*!< StepXML modeller. */
-	ModellerPLMXML = 53,              /*!< PLMXML modeller. */
-	ModellerNavisworks = 54,			/*!< For Future Use: Navisworks modeller. */
-	ModellerLast
-};
-
-struct FEntityMetaData
-{
-	TMap<FString, FString> MetaData;
-	bool bRemoved = false;
-	bool bShow = true;
-	bool bUnloaded = false;
-	FCADUUID ColorName = 0;
-	FCADUUID MaterialName = 0;
-	FFileDescriptor ExternalFile;
-};
-
-
-class FTechSoftInterface;
-
 class CADINTERFACES_API FTechSoftFileParser : public ICADFileParser
 {
 public:
diff --git a/Engine/Plugins/Enterprise/DatasmithCADImporter/Source/CADInterfaces/Private/TechSoftInterface.cpp b/Engine/Plugins/Enterprise/DatasmithCADImporter/Source/CADInterfaces/Private/TechSoftInterface.cpp
--- a/Engine/Plugins/Enterprise/DatasmithCADImporter/Source/CADInterfaces/Private/TechSoftInterface.cpp
+++ b/Engine/Plugins/Enterprise/DatasmithCADImporter/Source/CADInterfaces/Private/TechSoftInterface.cpp
@@ -113,6 +113,14 @@ FUniqueTechSoftModelFile LoadModelFileFromPrcFile(const A3DUTF8Char* CADFileName
 	return FUniqueTechSoftModelFile(ModelFile);
 }
 
+A3DStatus AdaptBRepInModelFile(A3DAsmModelFile* ModelFile, const A3DCopyAndAdaptBrepModelData& Setting, int32& ErrorCount, A3DCopyAndAdaptBrepModelErrorData** Errors)
+{
+	A3DUns32 NbErrors;
+	A3DStatus Status = A3DAdaptAndReplaceAllBrepInModelFileAdvanced(ModelFile, &Setting, &NbErrors, Errors);
+	ErrorCount = NbErrors;
+	return Status;
+}
+
 A3DStatus DeleteModelFile(A3DAsmModelFile* ModelFile)
 {
 	return A3DAsmModelFileDelete(ModelFile);
@@ -1415,8 +1423,8 @@ template<>
 const A3DEntity* TUniqueTSObj<A3DTopoUniqueVertexData>::GetDefaultIndexerValue() const { return nullptr; }
 template<>
 const A3DEntity* TUniqueTSObj<A3DTopoWireEdgeData>::GetDefaultIndexerValue() const { return nullptr; }
-template<>
 
+template<>
 uint32 TUniqueTSObjFromIndex<A3DGraphMaterialData>::GetDefaultIndexerValue() const { return A3D_DEFAULT_MATERIAL_INDEX; }
 template<>
 uint32 TUniqueTSObjFromIndex<A3DGraphPictureData>::GetDefaultIndexerValue() const { return A3D_DEFAULT_PICTURE_INDEX; }
diff --git a/Engine/Plugins/Enterprise/DatasmithCADImporter/Source/CADInterfaces/Private/TechSoftUtils.cpp b/Engine/Plugins/Enterprise/DatasmithCADImporter/Source/CADInterfaces/Private/TechSoftUtils.cpp
--- a/Engine/Plugins/Enterprise/DatasmithCADImporter/Source/CADInterfaces/Private/TechSoftUtils.cpp
+++ b/Engine/Plugins/Enterprise/DatasmithCADImporter/Source/CADInterfaces/Private/TechSoftUtils.cpp
@@ -582,7 +582,9 @@ A3DCrvNurbs* CreateTrimNurbsCurve(A3DCrvNurbs* CurveNurbsPtr, double UMin, doubl
 
 	return CADLibrary::TechSoftInterface::CreateCurveNurbs(*TransformedNurbsData);
 }
+#endif
 
+#ifdef USE_TECHSOFT_SDK
 void ExtractAttribute(const A3DMiscAttributeData& AttributeData, TMap<FString, FString>& OutMetaData)
 {
 	FString AttributeName;
@@ -605,16 +607,13 @@ void ExtractAttribute(const A3DMiscAttributeData& AttributeData, TMap<FString, F
 		case kA3DModellerAttributeTypeTime:
 		case kA3DModellerAttributeTypeInt:
 		{
-			A3DInt32 Value;
-			memcpy(&Value, AttributeData.m_asSingleAttributesData[Index].m_pcData, sizeof(A3DInt32));
+			A3DInt32 Value = *reinterpret_cast<A3DInt32*>(AttributeData.m_asSingleAttributesData[Index].m_pcData);
 			AttributeValue = FString::Printf(TEXT("%d"), Value);
 			break;
 		}
-
 		case kA3DModellerAttributeTypeReal:
 		{
-			A3DDouble Value;
-			memcpy(&Value, AttributeData.m_asSingleAttributesData[Index].m_pcData, sizeof(A3DDouble));
+			A3DDouble Value = *reinterpret_cast<A3DDouble*>(AttributeData.m_asSingleAttributesData[Index].m_pcData);
 			AttributeValue = FString::Printf(TEXT("%f"), Value);
 			break;
 		}
@@ -645,6 +644,7 @@ void ExtractAttribute(const A3DMiscAttributeData& AttributeData, TMap<FString, F
 		}
 	}
 }
+#endif
 
 FString CleanSdkName(const FString& Name)
 {
@@ -718,7 +718,6 @@ FString CleanCatiaReferenceName(const FString& Name)
 	}
 	return Name;
 }
-
 FString CleanNameByRemoving_prt(const FString& Name)
 {
 	int32 Position;
@@ -733,7 +732,6 @@ FString CleanNameByRemoving_prt(const FString& Name)
 	}
 	return Name;
 }
-
 bool CheckIfNameExists(TMap<FString, FString>& MetaData)
 {
 	FString* NamePtr = MetaData.Find(TEXT("Name"));
@@ -743,7 +741,6 @@ bool CheckIfNameExists(TMap<FString, FString>& MetaData)
 	}
 	return false;
 }
-
 bool ReplaceOrAddNameValue(TMap<FString, FString>& MetaData, const TCHAR* Key)
 {
 	FString* NamePtr = MetaData.Find(Key);
@@ -756,7 +753,6 @@ bool ReplaceOrAddNameValue(TMap<FString, FString>& MetaData, const TCHAR* Key)
 	return false;
 }
 
-#endif
 
 } // NS TechSoftUtils
 
diff --git a/Engine/Plugins/Enterprise/DatasmithCADImporter/Source/CADInterfaces/Private/TechSoftUtilsPrivate.h b/Engine/Plugins/Enterprise/DatasmithCADImporter/Source/CADInterfaces/Private/TechSoftUtilsPrivate.h
new file mode 100644
--- /dev/null
+++ b/Engine/Plugins/Enterprise/DatasmithCADImporter/Source/CADInterfaces/Private/TechSoftUtilsPrivate.h
@@ -0,0 +1,31 @@
+// Copyright Epic Games, Inc. All Rights Reserved.
+#pragma once
+
+#include "CoreMinimal.h"
+
+#include "TechSoftInterface.h"
+
+namespace CADLibrary
+{
+
+namespace TechSoftUtils
+{
+
+#ifdef USE_TECHSOFT_SDK
+void ExtractAttribute(const A3DMiscAttributeData& AttributeData, TMap<FString, FString>& OutMetaData);
+#endif
+
+FString CleanSdkName(const FString& Name);
+FString CleanCatiaInstanceSdkName(const FString& Name);
+FString CleanCatiaReferenceName(const FString& Name);
+FString Clean3dxmlReferenceSdkName(const FString& Name);
+FString CleanSwInstanceSdkName(const FString& Name);
+FString CleanSwReferenceSdkName(const FString& Name);
+FString CleanNameByRemoving_prt(const FString& Name);
+bool CheckIfNameExists(TMap<FString, FString>& MetaData);
+bool ReplaceOrAddNameValue(TMap<FString, FString>& MetaData, const TCHAR* Key);
+
+} // NS TechSoftUtils
+
+} // CADLibrary
+
diff --git a/Engine/Plugins/Enterprise/DatasmithCADImporter/Source/CADInterfaces/Public/TUniqueTechSoftObj.h b/Engine/Plugins/Enterprise/DatasmithCADImporter/Source/CADInterfaces/Public/TUniqueTechSoftObj.h
--- a/Engine/Plugins/Enterprise/DatasmithCADImporter/Source/CADInterfaces/Public/TUniqueTechSoftObj.h
+++ b/Engine/Plugins/Enterprise/DatasmithCADImporter/Source/CADInterfaces/Public/TUniqueTechSoftObj.h
@@ -645,17 +645,17 @@ CADINTERFACES_API const A3DEntity* TUniqueTSObj<A3DTopoWireEdgeData>::GetDefault
 // TUniqueTSObjFromIndex -----------------------------------
 
 template<>
-CADINTERFACES_API void TUniqueTSObj<A3DGraphMaterialData>::InitializeData();
+CADINTERFACES_API void TUniqueTSObjFromIndex<A3DGraphMaterialData>::InitializeData();
 template<>
-CADINTERFACES_API void TUniqueTSObj<A3DGraphPictureData>::InitializeData();
+CADINTERFACES_API void TUniqueTSObjFromIndex<A3DGraphPictureData>::InitializeData();
 template<>
-CADINTERFACES_API void TUniqueTSObj<A3DGraphRgbColorData>::InitializeData();
+CADINTERFACES_API void TUniqueTSObjFromIndex<A3DGraphRgbColorData>::InitializeData();
 template<>
-CADINTERFACES_API void TUniqueTSObj<A3DGraphStyleData>::InitializeData();
+CADINTERFACES_API void TUniqueTSObjFromIndex<A3DGraphStyleData>::InitializeData();
 template<>
-CADINTERFACES_API void TUniqueTSObj<A3DGraphTextureApplicationData>::InitializeData();
+CADINTERFACES_API void TUniqueTSObjFromIndex<A3DGraphTextureApplicationData>::InitializeData();
 template<>
-CADINTERFACES_API void TUniqueTSObj<A3DGraphTextureDefinitionData>::InitializeData();
+CADINTERFACES_API void TUniqueTSObjFromIndex<A3DGraphTextureDefinitionData>::InitializeData();
 
 template<>
 CADINTERFACES_API A3DStatus TUniqueTSObjFromIndex<A3DGraphMaterialData>::GetData(const uint32 InEntityIndex);
diff --git a/Engine/Plugins/Enterprise/DatasmithCADImporter/Source/CADInterfaces/Public/TechSoftInterface.h b/Engine/Plugins/Enterprise/DatasmithCADImporter/Source/CADInterfaces/Public/TechSoftInterface.h
--- a/Engine/Plugins/Enterprise/DatasmithCADImporter/Source/CADInterfaces/Public/TechSoftInterface.h
+++ b/Engine/Plugins/Enterprise/DatasmithCADImporter/Source/CADInterfaces/Public/TechSoftInterface.h
@@ -80,6 +80,8 @@ CADINTERFACES_API FString GetTechSoftVersion();
 CADINTERFACES_API FUniqueTechSoftModelFile LoadModelFileFromFile(const A3DImport& Import, A3DStatus& Status);
 CADINTERFACES_API FUniqueTechSoftModelFile LoadModelFileFromPrcFile(const A3DUTF8Char* CADFileName, A3DRWParamsPrcReadHelper** ReadHelper);
 
+CADINTERFACES_API A3DStatus AdaptBRepInModelFile(A3DAsmModelFile* ModelFile, const A3DCopyAndAdaptBrepModelData& Setting, int32& ErrorCount, A3DCopyAndAdaptBrepModelErrorData** Errors);
+
 CADINTERFACES_API A3DStatus DeleteModelFile(A3DAsmModelFile* ModelFile);
 CADINTERFACES_API A3DStatus DeleteEntity(A3DEntity* pEntity);
 
@@ -207,5 +209,60 @@ private:
 	A3DAsmModelFile* ModelFile;
 };
 
+enum EModellerType : uint32
+{
+	ModellerUnknown = 0,              /*!< User modeller. */
+	ModellerCatia = 2,                /*!< CATIA modeller. */
+	ModellerCatiaV5 = 3,              /*!< CATIA V5 modeller. */
+	ModellerCadds = 4,                /*!< CADDS modeller. */
+	ModellerUnigraphics = 5,          /*!< Unigraphics modeller. */
+	ModellerParasolid = 6,            /*!< Parasolid modeller. */
+	ModellerEuclid = 7,               /*!< Euclid modeller. */
+	ModellerIges = 9,                 /*!< IGES modeller. */
+	ModellerUnisurf = 10,             /*!< Unisurf modeller. */
+	ModellerVda = 11,                 /*!< VDA modeller. */
+	ModellerStl = 12,                 /*!< STL modeller. */
+	ModellerWrl = 13,                 /*!< WRL modeller. */
+	ModellerDxf = 14,                 /*!< DXF modeller. */
+	ModellerAcis = 15,                /*!< ACIS modeller. */
+	ModellerProE = 16,                /*!< Pro/E modeller. */
+	ModellerStep = 18,                /*!< STEP modeller. */
+	ModellerIdeas = 19,               /*!< I-DEAS modeller. */
+	ModellerJt = 20,                  /*!< JT modeller. */
+	ModellerSlw = 22,                 /*!< SolidWorks modeller. */
+	ModellerCgr = 23,                 /*!< CGR modeller. */
+	ModellerPrc = 24,                 /*!< PRC modeller. */
+	ModellerXvl = 25,                 /*!< XVL modeller. */
+	ModellerHpgl = 26,                /*!< HPGL modeller. */
+	ModellerTopSolid = 27,            /*!< TopSolid modeller. */
+	ModellerOneSpaceDesigner = 28,    /*!< OneSpace designer modeller. */
+	Modeller3dxml = 29,               /*!< 3DXML modeller. */
+	ModellerInventor = 30,            /*!< Inventor modeller. */
+	ModellerPostScript = 31,          /*!< Postscript modeller. */
+	ModellerPDF = 32,                 /*!< PDF modeller. */
+	ModellerU3D = 33,                 /*!< U3D modeller. */
+	ModellerIFC = 34,                 /*!< IFC modeller. */
+	ModellerDWG = 35,                 /*!< DWG modeller. */
+	ModellerDWF = 36,                 /*!< DWF modeller. */
+	ModellerSE = 37,                  /*!< SolidEdge modeller. */
+	ModellerOBJ = 38,                 /*!< OBJ modeller. */
+	ModellerKMZ = 39,                 /*!< KMZ modeller. */
+	ModellerDAE = 40,                 /*!< COLLADA modeller. */
+	Modeller3DS = 41,                 /*!< 3DS modeller. */
+	ModellerRhino = 43,               /*!< Rhino modeller. */
+	ModellerXML = 44,                 /*!< XML modeller. */
+	Modeller3mf = 45,                 /*!< 3MF modeller. */
+	ModellerScs = 46,                 /*!< SCS modeller. */
+	Modeller3dHtml = 47,              /*!< 3DHTML modeller. */
+	ModellerHsf = 48,                 /*!< Hsf modeller. */
+	ModellerGltf = 49,                /*!< GL modeller. */
+	ModellerRevit = 50,               /*!< Revit modeller. */
+	ModellerFBX = 51,                 /*!< FBX modeller. */
+	ModellerStepXML = 52,             /*!< StepXML modeller. */
+	ModellerPLMXML = 53,              /*!< PLMXML modeller. */
+	ModellerNavisworks = 54,			/*!< For Future Use: Navisworks modeller. */
+	ModellerLast
+};
+
 } // CADLibrary
 
diff --git a/Engine/Plugins/Enterprise/DatasmithCADImporter/Source/CADInterfaces/Public/TechSoftUtils.h b/Engine/Plugins/Enterprise/DatasmithCADImporter/Source/CADInterfaces/Public/TechSoftUtils.h
--- a/Engine/Plugins/Enterprise/DatasmithCADImporter/Source/CADInterfaces/Public/TechSoftUtils.h
+++ b/Engine/Plugins/Enterprise/DatasmithCADImporter/Source/CADInterfaces/Public/TechSoftUtils.h
@@ -3,6 +3,7 @@
 
 #include "CoreMinimal.h"
 
+#include "CADData.h"
 #include "TechSoftInterface.h"
 
 #define JSON_ENTRY_FILE_UNIT		TEXT("FileUnit")
@@ -31,6 +32,16 @@ namespace TechSoftInterface
 class FTechSoftInterface;
 }
 
+struct FEntityMetaData
+{
+	TMap<FString, FString> MetaData;
+	bool bRemoved = false;
+	bool bShow = true;
+	bool bUnloaded = false;
+	FCADUUID ColorName = 0;
+	FCADUUID MaterialName = 0;
+	FFileDescriptor ExternalFile;
+};
 
 namespace TechSoftUtils
 {
diff --git a/Engine/Source/Runtime/Datasmith/CADKernel/Private/CADKernel/Core/Session.cpp b/Engine/Source/Runtime/Datasmith/CADKernel/Private/CADKernel/Core/Session.cpp
--- a/Engine/Source/Runtime/Datasmith/CADKernel/Private/CADKernel/Core/Session.cpp
+++ b/Engine/Source/Runtime/Datasmith/CADKernel/Private/CADKernel/Core/Session.cpp
@@ -28,6 +28,24 @@ FModel& FSession::GetModel()
 	return Database.GetModel();
 }
 
+void FSession::SaveDatabase(const TCHAR* FileName, const TArray<FEntity*>& SelectedEntities)
+{
+	TArray<FIdent> EntityIds;
+	EntityIds.Reserve(SelectedEntities.Num());
+
+	SpawnEntityIdents(SelectedEntities, true);
+
+	for (const FEntity* Entity : SelectedEntities)
+	{
+		EntityIds.Add(Entity->GetId());
+	}
+
+	TSharedPtr<FCADKernelArchive> Archive = FCADKernelArchive::CreateArchiveWriter(*this, FileName);
+
+	Database.SerializeSelection(*Archive.Get(), EntityIds);
+	Archive->Close();
+}
+
 void FSession::SaveDatabase(const TCHAR* FileName, const TArray<TSharedPtr<FEntity>>& SelectedEntities)
 {
 	TArray<FIdent> EntityIds;
diff --git a/Engine/Source/Runtime/Datasmith/CADKernel/Public/CADKernel/Core/Session.h b/Engine/Source/Runtime/Datasmith/CADKernel/Public/CADKernel/Core/Session.h
--- a/Engine/Source/Runtime/Datasmith/CADKernel/Public/CADKernel/Core/Session.h
+++ b/Engine/Source/Runtime/Datasmith/CADKernel/Public/CADKernel/Core/Session.h
@@ -67,6 +67,25 @@ namespace CADKernel
 		 */
 		void SaveDatabase(const TCHAR* FileName, const TArray<TSharedPtr<FEntity>>& Entities);
 
+		/**
+		 * Save a selection and all the dependencies as a FAchive in a file
+		 * Mandatory: all entity have to have a defined ID
+		 * Use SpawnEntityIdent if needed
+		 */
+		void SaveDatabase(const TCHAR* FileName, const TArray<FEntity*>& Entities);
+
+		/**
+		 * Save a selection and all the dependencies as a FAchive in a file
+		 * Mandatory: all entity have to have a defined ID
+		 * Use SpawnEntityIdent if needed
+		 */
+		void SaveDatabase(const TCHAR* FileName, FEntity& Entity)
+		{
+			TArray<FEntity*> Entities;
+			Entities.Emplace(&Entity);
+			SaveDatabase(FileName, Entities);
+		}
+
 		/**
 		 * Save a selection and all the dependencies as a FAchive in a file
 		 * Mandatory: all entity have to have a defined ID
