[NO NON-PII COMMIT MESSAGE]

diff --git a/Engine/Plugins/Enterprise/DatasmithCADImporter/Source/CADInterfaces/Private/TechSoftFileParser.cpp b/Engine/Plugins/Enterprise/DatasmithCADImporter/Source/CADInterfaces/Private/TechSoftFileParser.cpp
--- a/Engine/Plugins/Enterprise/DatasmithCADImporter/Source/CADInterfaces/Private/TechSoftFileParser.cpp
+++ b/Engine/Plugins/Enterprise/DatasmithCADImporter/Source/CADInterfaces/Private/TechSoftFileParser.cpp
@@ -46,7 +46,7 @@ bool CheckIfNameExists(TMap<FString, FString>& MetaData)
 	return false;
 }
 
-bool ReplaceOrAddNameValue(TMap<FString, FString>& MetaData, TCHAR* Key)
+bool ReplaceOrAddNameValue(TMap<FString, FString>& MetaData, const TCHAR* Key)
 {
 	FString* NamePtr = MetaData.Find(Key);
 	if (NamePtr != nullptr)
@@ -1186,12 +1186,6 @@ void FTechSoftFileParser::ExtractMetaData(const A3DEntity* Entity, FEntityMetaDa
 	TUniqueTSObj<A3DRootBaseData> MetaData(Entity);
 	if (MetaData.IsValid())
 	{
-		if (false && MetaData->m_uiPersistentId > 0)
-		{
-			FString PersistentId = FString::FromInt(MetaData->m_uiPersistentId);
-			OutMetaData.MetaData.Emplace(TEXT("PersistentId"), PersistentId);
-		}
-
 		if (MetaData->m_pcName && MetaData->m_pcName[0] != '\0')
 		{
 			FString SDKName = UTF8_TO_TCHAR(MetaData->m_pcName);
@@ -1391,7 +1385,7 @@ void FTechSoftFileParser::ExtractSpecificMetaData(const A3DAsmProductOccurrence*
 	//----------- Export Specific information per CAD format -----------
 	switch (ModellerType)
 	{
-	case kA3DModellerSlw:
+	case ModellerSlw:
 	{
 		TUniqueTSObj<A3DAsmProductOccurrenceDataSLW> SolidWorksSpecificData(Occurrence);
 		if (SolidWorksSpecificData.IsValid())
@@ -1406,7 +1400,7 @@ void FTechSoftFileParser::ExtractSpecificMetaData(const A3DAsmProductOccurrence*
 		}
 		break;
 	}
-	case kA3DModellerUnigraphics:
+	case ModellerUnigraphics:
 	{
 		TUniqueTSObj<A3DAsmProductOccurrenceDataUg> UnigraphicsSpecificData(Occurrence);
 		if (UnigraphicsSpecificData.IsValid())
@@ -1431,7 +1425,7 @@ void FTechSoftFileParser::ExtractSpecificMetaData(const A3DAsmProductOccurrence*
 		break;
 	}
 
-	case kA3DModellerCatiaV5:
+	case ModellerCatiaV5:
 	{
 		TUniqueTSObj<A3DAsmProductOccurrenceDataCV5> CatiaV5SpecificData(Occurrence);
 		if (CatiaV5SpecificData.IsValid())
