[TechSoft][DWG] Apply the same logic when parsing the occurrence to ExternalData as to Part or Occurences

diff --git a/Engine/Plugins/Enterprise/DatasmithCADImporter/Source/CADInterfaces/Private/TechSoftFileParser.cpp b/Engine/Plugins/Enterprise/DatasmithCADImporter/Source/CADInterfaces/Private/TechSoftFileParser.cpp
--- a/Engine/Plugins/Enterprise/DatasmithCADImporter/Source/CADInterfaces/Private/TechSoftFileParser.cpp
+++ b/Engine/Plugins/Enterprise/DatasmithCADImporter/Source/CADInterfaces/Private/TechSoftFileParser.cpp
@@ -766,32 +766,13 @@ FCadId FTechSoftFileParser::TraverseOccurrence(const A3DAsmProductOccurrence* Oc
 	
 	FArchiveComponent& Component = AddComponent(InstanceMetaData, Instance);
 
-	// Add External data
-	if (OccurrenceData->m_pExternalData != nullptr)
-	{
-		TUniqueTSObj<A3DAsmProductOccurrenceData> ExternalData(OccurrenceData->m_pExternalData);
-		if (ExternalData->m_pPart != nullptr)
-		{
-			A3DAsmPartDefinition* PartDefinition = ExternalData->m_pPart;
-			TraversePartDefinition(PartDefinition, Component);
-		}
-
-		uint32 ExternalChildrenCount = ExternalData->m_uiPOccurrencesSize;
-		A3DAsmProductOccurrence** ExternalChildren = ExternalData->m_ppPOccurrences;
-		for (uint32 Index = 0; Index < ExternalChildrenCount; ++Index)
-		{
-			int32 ChildrenId = TraverseOccurrence(ExternalChildren[Index]);
-			Component.Children.Add(ChildrenId);
-		}
-	}
-
-	while (OccurrenceData->m_pPrototype != nullptr && OccurrenceData->m_pPart == nullptr && OccurrenceData->m_uiPOccurrencesSize == 0)
+	while (OccurrenceData->m_pPrototype != nullptr && OccurrenceData->m_pPart == nullptr && OccurrenceData->m_uiPOccurrencesSize == 0 && OccurrenceData->m_pExternalData == nullptr)
 	{
 		CachedOccurrencePtr = OccurrenceData->m_pPrototype;
 		OccurrenceData.FillFrom(OccurrenceData->m_pPrototype);
 	}
 
-	if(OccurrenceData->m_pPart == nullptr && OccurrenceData->m_uiPOccurrencesSize == 0)
+	if(OccurrenceData->m_pPart == nullptr && OccurrenceData->m_uiPOccurrencesSize == 0 && OccurrenceData->m_pExternalData == nullptr)
 	{
 		return Instance.ObjectId;
 	}
@@ -822,6 +803,30 @@ FCadId FTechSoftFileParser::TraverseOccurrence(const A3DAsmProductOccurrence* Oc
 		Component.Children.Add(ChildrenId);
 	}
 
+	// Add External data
+	OccurrenceData.FillFrom(CachedOccurrencePtr);
+	while (OccurrenceData->m_pPrototype != nullptr && OccurrenceData->m_pExternalData == nullptr)
+	{
+		OccurrenceData.FillFrom(OccurrenceData->m_pPrototype);
+	}
+	if (OccurrenceData->m_pExternalData != nullptr)
+	{
+		TUniqueTSObj<A3DAsmProductOccurrenceData> ExternalData(OccurrenceData->m_pExternalData);
+		if (ExternalData->m_pPart != nullptr)
+		{
+			A3DAsmPartDefinition* PartDefinition = ExternalData->m_pPart;
+			TraversePartDefinition(PartDefinition, Component);
+		}
+
+		uint32 ExternalChildrenCount = ExternalData->m_uiPOccurrencesSize;
+		A3DAsmProductOccurrence** ExternalChildren = ExternalData->m_ppPOccurrences;
+		for (uint32 Index = 0; Index < ExternalChildrenCount; ++Index)
+		{
+			int32 ChildrenId = TraverseOccurrence(ExternalChildren[Index]);
+			Component.Children.Add(ChildrenId);
+		}
+	}
+
 	return Instance.ObjectId;
 }
 
@@ -833,38 +838,62 @@ void FTechSoftFileParser::CountUnderOccurrence(const A3DAsmProductOccurrence* Oc
 		ComponentCount[EComponentType::Occurrence]++;
 		ComponentCount[EComponentType::Reference]++;
 
-		if (OccurrenceData->m_pExternalData != nullptr)
+		const A3DAsmProductOccurrence* CachedOccurrencePtr = Occurrence;
+		while (OccurrenceData->m_pPrototype != nullptr && OccurrenceData->m_pPart == nullptr && OccurrenceData->m_uiPOccurrencesSize == 0 && OccurrenceData->m_pExternalData == nullptr)
 		{
-			CountUnderOccurrence(OccurrenceData->m_pExternalData);
+			CachedOccurrencePtr = OccurrenceData->m_pPrototype;
+			OccurrenceData.FillFrom(OccurrenceData->m_pPrototype);
 		}
 
-		A3DAsmProductOccurrence* PrototypePtr = OccurrenceData->m_pPrototype;
-		A3DAsmPartDefinition* PartDefinition = OccurrenceData->m_pPart;
+		if (OccurrenceData->m_pPart == nullptr && OccurrenceData->m_uiPOccurrencesSize == 0 && OccurrenceData->m_pExternalData == nullptr)
+		{
+			return;
+		}
 
-		while (!PartDefinition && PrototypePtr)
+		// count under part
+		while (OccurrenceData->m_pPrototype != nullptr && OccurrenceData->m_pPart == nullptr)
+		{
+			OccurrenceData.FillFrom(OccurrenceData->m_pPrototype);
+		}
+		if (OccurrenceData->m_pPart != nullptr)
 		{
-			TUniqueTSObj<A3DAsmProductOccurrenceData> PrototypeOccurrenceData(PrototypePtr);
-			PartDefinition = PrototypeOccurrenceData->m_pPart;
-			PrototypePtr = PrototypeOccurrenceData->m_pPrototype;
+			CountUnderPartDefinition(OccurrenceData->m_pPart);
 		}
 
-		CountUnderPartDefinition(PartDefinition);
+		// count under Occurrence
+		OccurrenceData.FillFrom(CachedOccurrencePtr);
+		while (OccurrenceData->m_pPrototype != nullptr && OccurrenceData->m_uiPOccurrencesSize == 0)
+		{
+			OccurrenceData.FillFrom(OccurrenceData->m_pPrototype);
+		}
 
 		uint32 ChildrenCount = OccurrenceData->m_uiPOccurrencesSize;
 		A3DAsmProductOccurrence** Children = OccurrenceData->m_ppPOccurrences;
-		PrototypePtr = OccurrenceData->m_pPrototype;
-
-		while (ChildrenCount == 0 && PrototypePtr)
+		for (uint32 Index = 0; Index < ChildrenCount; ++Index)
 		{
-			TUniqueTSObj<A3DAsmProductOccurrenceData> PrototypeOccurrenceData(PrototypePtr);
-			ChildrenCount = PrototypeOccurrenceData->m_uiPOccurrencesSize;
-			Children = PrototypeOccurrenceData->m_ppPOccurrences;
-			PrototypePtr = PrototypeOccurrenceData->m_pPrototype;
+			CountUnderOccurrence(Children[Index]);
 		}
 
-		for (uint32 Index = 0; Index < ChildrenCount; ++Index)
+		// count under External data
+		OccurrenceData.FillFrom(CachedOccurrencePtr);
+		while (OccurrenceData->m_pPrototype != nullptr && OccurrenceData->m_pExternalData == nullptr)
 		{
-			CountUnderOccurrence(Children[Index]);
+			OccurrenceData.FillFrom(OccurrenceData->m_pPrototype);
+		}
+		if (OccurrenceData->m_pExternalData != nullptr)
+		{
+			TUniqueTSObj<A3DAsmProductOccurrenceData> ExternalData(OccurrenceData->m_pExternalData);
+			if (ExternalData->m_pPart != nullptr)
+			{
+				CountUnderPartDefinition(ExternalData->m_pPart);
+			}
+
+			uint32 ExternalChildrenCount = ExternalData->m_uiPOccurrencesSize;
+			A3DAsmProductOccurrence** ExternalChildren = ExternalData->m_ppPOccurrences;
+			for (uint32 Index = 0; Index < ExternalChildrenCount; ++Index)
+			{
+				CountUnderOccurrence(ExternalChildren[Index]);
+			}
 		}
 	}
 }
@@ -907,7 +936,7 @@ void FTechSoftFileParser::ProcessPrototype(const A3DAsmProductOccurrence* InProt
 			return;
 		}
 
-		if (PrototypeData->m_pPart != nullptr || PrototypeData->m_uiPOccurrencesSize != 0 || PrototypeData->m_pPrototype == nullptr)
+		if (PrototypeData->m_pPart != nullptr || PrototypeData->m_uiPOccurrencesSize != 0 || PrototypeData->m_pExternalData != nullptr || PrototypeData->m_pPrototype == nullptr)
 		{
 			ExtractMetaData(PrototypePtr, OutPrototypeMetaData);
 			ExtractSpecificMetaData(PrototypePtr, OutPrototypeMetaData);
@@ -915,23 +944,23 @@ void FTechSoftFileParser::ProcessPrototype(const A3DAsmProductOccurrence* InProt
 
 			TUniqueTSObj<A3DUTF8Char*> FilePathUTF8Ptr;
 			FilePathUTF8Ptr.FillWith(&TechSoftInterface::GetFilePathName, PrototypePtr);
-			if (!FilePathUTF8Ptr.IsValid())
+			if (!FilePathUTF8Ptr.IsValid() || *FilePathUTF8Ptr == nullptr)
 			{
 				FilePathUTF8Ptr.FillWith(&TechSoftInterface::GetOriginalFilePathName, PrototypePtr);
 			}
-			if (FilePathUTF8Ptr.IsValid())
+			if (FilePathUTF8Ptr.IsValid() && *FilePathUTF8Ptr != nullptr)
 			{
 				FString FilePath = UTF8_TO_TCHAR(*FilePathUTF8Ptr);
 				FPaths::NormalizeFilename(FilePath);
 				FString FileName = FPaths::GetCleanFilename(FilePath);
-				if (FileName != CADFileData.GetCADFileDescription().GetFileName())
+				if (!FileName.IsEmpty() && FileName != CADFileData.GetCADFileDescription().GetFileName())
 				{
 					OutPrototypeMetaData.ExternalFile = FFileDescriptor(*FilePath, nullptr, *CADFileData.GetCADFileDescription().GetRootFolder());
 				}
 			}
 		}
 		
-		if (PrototypeData->m_pPart != nullptr || PrototypeData->m_uiPOccurrencesSize != 0)
+		if (PrototypeData->m_pPart != nullptr|| PrototypeData->m_uiPOccurrencesSize != 0 || PrototypeData->m_pExternalData != nullptr )
 		{
 			OutPrototypeMetaData.bUnloaded = false;
 			PrototypePtr = nullptr;
