Remove Old CAD Library

diff --git a/Engine/Plugins/Enterprise/DatasmithCADImporter/Source/DatasmithCADTranslator/Private/DatasmithCADTranslator.cpp b/Engine/Plugins/Enterprise/DatasmithCADImporter/Source/DatasmithCADTranslator/Private/DatasmithCADTranslator.cpp
--- a/Engine/Plugins/Enterprise/DatasmithCADImporter/Source/DatasmithCADTranslator/Private/DatasmithCADTranslator.cpp
+++ b/Engine/Plugins/Enterprise/DatasmithCADImporter/Source/DatasmithCADTranslator/Private/DatasmithCADTranslator.cpp
@@ -76,8 +76,8 @@ void FDatasmithCADTranslator::Initialize(FDatasmithTranslatorCapabilities& OutCa
 
 	if (CADLibrary::FImportParameters::GCADLibrary.Equals(TEXT("TechSoft")))
 	{
-		OutCapabilities.SupportedFileFormats.Add(FFileFormatInfo{ TEXT("ifc"), TEXT("IFC files") });
 		OutCapabilities.SupportedFileFormats.Add(FFileFormatInfo{ TEXT("hsf"), TEXT("HOOPS stream files") });
+		OutCapabilities.SupportedFileFormats.Add(FFileFormatInfo{ TEXT("prc"), TEXT("HOOPS stream files") });
 	}
 	else
 	{
