The tessellator was no plugged anymore after the old cad library was removed

diff --git a/Engine/Plugins/Enterprise/DatasmithCADImporter/Source/DatasmithOpenNurbsTranslator/Private/DatasmithOpenNurbsTranslator.cpp b/Engine/Plugins/Enterprise/DatasmithCADImporter/Source/DatasmithOpenNurbsTranslator/Private/DatasmithOpenNurbsTranslator.cpp
--- a/Engine/Plugins/Enterprise/DatasmithCADImporter/Source/DatasmithOpenNurbsTranslator/Private/DatasmithOpenNurbsTranslator.cpp
+++ b/Engine/Plugins/Enterprise/DatasmithCADImporter/Source/DatasmithOpenNurbsTranslator/Private/DatasmithOpenNurbsTranslator.cpp
@@ -652,9 +652,18 @@ public:
 
 		if (CADLibrary::FImportParameters::bGDisableCADKernelTessellation)
 		{
-			TSharedRef<FOpenNurbsBRepToCoretechConverter> OpenNurbsBRepToCoretechConverter = MakeShared<FOpenNurbsBRepToCoretechConverter>(TEXT("Al2CTSharedSession"), ImportParameters);
-			CADModelConverter = OpenNurbsBRepToCoretechConverter;
-			OpenNurbsBRepConverter = OpenNurbsBRepToCoretechConverter;
+			if (CADLibrary::FImportParameters::GCADLibrary == TEXT("TechSoft"))
+			{
+				TSharedRef<FOpenNurbsBRepToTechSoftConverter> OpenNurbsBRepToCoretechConverter = MakeShared<FOpenNurbsBRepToTechSoftConverter>(ImportParameters);
+				CADModelConverter = OpenNurbsBRepToCoretechConverter;
+				OpenNurbsBRepConverter = OpenNurbsBRepToCoretechConverter;
+			}
+			else if (CADLibrary::FImportParameters::GCADLibrary == TEXT("KernelIO"))
+			{
+				TSharedRef<FOpenNurbsBRepToCoretechConverter> OpenNurbsBRepToCoretechConverter = MakeShared<FOpenNurbsBRepToCoretechConverter>(TEXT("Al2CTSharedSession"), ImportParameters);
+				CADModelConverter = OpenNurbsBRepToCoretechConverter;
+				OpenNurbsBRepConverter = OpenNurbsBRepToCoretechConverter;
+			}
 		}
 		else
 		{
