Plug of CADKernel into TechSoft Step 3:
 - Add new class TechSoftFileParserCADKernelTessellator to tessellate with CADKernel

diff --git a/Engine/Plugins/Enterprise/DatasmithCADImporter/Source/CADInterfaces/Private/TechSoftBridge.cpp b/Engine/Plugins/Enterprise/DatasmithCADImporter/Source/CADInterfaces/Private/TechSoftBridge.cpp
--- a/Engine/Plugins/Enterprise/DatasmithCADImporter/Source/CADInterfaces/Private/TechSoftBridge.cpp
+++ b/Engine/Plugins/Enterprise/DatasmithCADImporter/Source/CADInterfaces/Private/TechSoftBridge.cpp
@@ -5,6 +5,8 @@
 #ifdef USE_TECHSOFT_SDK
 
 #include "TechSoftInterface.h"
+#include "TechSoftFileParser.h"
+#include "TechSoftUtils.h"
 #include "TUniqueTechSoftObj.h"
 #include "CADFileReport.h"
 
@@ -40,6 +42,7 @@
 #include "CADKernel/Geo/Surfaces/TabulatedCylinderSurface.h"
 #include "CADKernel/Geo/Surfaces/TorusSurface.h"
 #include "CADKernel/Topo/Body.h"
+#include "CADKernel/Topo/Model.h"
 #include "CADKernel/Topo/Shell.h"
 #include "CADKernel/Topo/TopologicalEdge.h"
 #include "CADKernel/Topo/TopologicalFace.h"
@@ -50,7 +53,7 @@
 
 #include "CADKernel/UI/Display.h"
 
-namespace CADKernel
+namespace CADLibrary
 {
 namespace TechSoftUtils
 {
@@ -60,458 +63,20 @@ const FUVReparameterization FUVReparameterization::Identity(1., 0, 1., 0);
 template<typename... InArgTypes>
 A3DStatus GetCurveAsNurbs(const A3DCrvBase* A3DCurve, A3DCrvNurbsData* DataPtr, InArgTypes&&... Args)
 {
-	return CADLibrary::TechSoftInterface::GetCurveAsNurbs(A3DCurve, DataPtr, Forward<InArgTypes>(Args)...);
+	return TechSoftInterface::GetCurveAsNurbs(A3DCurve, DataPtr, Forward<InArgTypes>(Args)...);
 };
 
 template<typename... InArgTypes>
 A3DStatus GetSurfaceAsNurbs(const A3DSurfBase* A3DSurface, A3DSurfNurbsData* DataPtr, InArgTypes&&... Args)
 {
-	return CADLibrary::TechSoftInterface::GetSurfaceAsNurbs(A3DSurface, DataPtr, Forward<InArgTypes>(Args)...);
+	return TechSoftInterface::GetSurfaceAsNurbs(A3DSurface, DataPtr, Forward<InArgTypes>(Args)...);
 }; 
 
-FString GetFileName(const FString& InFilePath)
+CADKernel::FMatrixH CreateCoordinateSystem(const A3DMiscCartesianTransformationData& Transformation, double UnitScale = 1.0)
 {
-	int32 Position = -1;
-	if (!InFilePath.FindLastChar(TEXT('\\'), Position))
-	{
-		InFilePath.FindLastChar(TEXT('/'), Position);
-	}
-
-	if (Position >= 0)
-	{
-		return InFilePath.Right(Position + 1);
-	}
-
-	return InFilePath;
-}
-
-FString GetOriginalName(const FString& Name)
-{
-	int32 Position;
-	if (Name.FindLastChar(TEXT('['), Position))
-	{
-		return Name.Left(Position);
-	}
-	return Name;
-}
-
-FString Get3dxmlOriginalName(const FString& Name)
-{
-	int32 Position;
-	if (Name.FindLastChar(TEXT('('), Position))
-	{
-		FString NewName = Name.Right(Position + 1);
-		if (NewName.FindLastChar(TEXT(')'), Position))
-		{
-			NewName[Position] = 0;
-		}
-		return NewName;
-	}
-	return Name;
-}
-
-FString GetSWOriginalName(const FString& Name)
-{
-	int32 Position;
-	if (Name.FindLastChar(TEXT('-'), Position))
-	{
-		FString NewName = Name.Left(Position) + TEXT("<") + Name.Right(Position + 1) + TEXT(">");
-		return NewName;
-	}
-	return Name;
-}
-
-int TraverseAttribute(const A3DMiscAttribute* Attribute, TMap<FString, FString>& OutMetaData)
-{
-	A3DInt32 Ret = A3D_SUCCESS;
-
-	CADLibrary::TUniqueTSObj<A3DMiscAttributeData> AttributeData(Attribute);
-	if (AttributeData.IsValid())
-	{
-		FString AttributeName;
-		if (AttributeData->m_bTitleIsInt)
-		{
-			A3DUns32 uiVal = 0;
-			memcpy(&uiVal, AttributeData->m_pcTitle, sizeof(A3DUns32));
-			AttributeName = Utils::ToFString((uint32)uiVal);
-		}
-		else if (AttributeData->m_pcTitle && AttributeData->m_pcTitle[0] != '\0')
-		{
-			AttributeName = UTF8_TO_TCHAR(AttributeData->m_pcTitle);
-		}
-		else
-		{
-			return A3D_ERROR;
-		}
-
-		for (A3DUns32 Index = 0; Index < AttributeData->m_uiSize; ++Index)
-		{
-			//single->SetAttribute("m_bTitleIsInt", (int)AttributeData->m_asSingleAttributesData[Index].m_bTitleIsInt);
-			if (AttributeData->m_asSingleAttributesData[Index].m_pcTitle == NULL)
-			{
-				//single->SetAttribute("m_pcTitle", "NULL");
-			}
-			else if (AttributeData->m_asSingleAttributesData[Index].m_bTitleIsInt)
-			{
-				A3DUns32 uiVal;
-				memcpy(&uiVal, AttributeData->m_asSingleAttributesData[Index].m_pcTitle, sizeof(A3DUns32));
-				//single->SetAttribute("m_pcTitle", (int)uiVal);
-			}
-			else
-			{
-				if (AttributeData->m_asSingleAttributesData[Index].m_pcTitle && AttributeData->m_asSingleAttributesData[Index].m_pcTitle[0] != '\0')
-				{
-					//single->SetAttribute("m_pcTitle", AttributeData->m_asSingleAttributesData[Index].m_pcTitle);
-				}
-			}
-
-			FString AttributeValue;
-
-			switch (AttributeData->m_asSingleAttributesData[Index].m_eType)
-			{
-			case kA3DModellerAttributeTypeTime:
-			case kA3DModellerAttributeTypeInt:
-			{
-				A3DInt32 Value;
-				memcpy(&Value, AttributeData->m_asSingleAttributesData[Index].m_pcData, sizeof(A3DInt32));
-				AttributeValue = Utils::ToFString((int32) Value);
-				break;
-			}
-
-			case kA3DModellerAttributeTypeReal:
-			{
-				A3DDouble Value;
-				memcpy(&Value, AttributeData->m_asSingleAttributesData[Index].m_pcData, sizeof(A3DDouble));
-				AttributeValue = Utils::ToFString((double) Value);
-				break;
-			}
-
-			case kA3DModellerAttributeTypeString:
-				if (AttributeData->m_asSingleAttributesData[Index].m_pcData && AttributeData->m_asSingleAttributesData[Index].m_pcData[0] != '\0')
-				{
-					AttributeValue = UTF8_TO_TCHAR(AttributeData->m_asSingleAttributesData[Index].m_pcData);
-				}
-				break;
-
-			default:
-				break;
-			}
-
-			if (AttributeName.Len())
-			{
-				if (Index)
-				{
-					FString Key = FString::Printf(TEXT("%s_%u"), *AttributeName, Index);
-					OutMetaData.Emplace(Key, AttributeValue);
-				}
-				else
-				{
-					OutMetaData.Emplace(AttributeName, AttributeValue);
-				}
-			}
-			else
-			{
-				Ret = A3D_ERROR;
-			}
-		}
-	}
-
-	return Ret;
-}
-
-int TraverseBase(const ECADType FileType, const A3DEntity* Entity, FEntityData& OutEntityData, bool bIsOccurrence)
-{
-	CADLibrary::TUniqueTSObj<A3DRootBaseData> RootBaseData(Entity);
-	if (RootBaseData.IsValid())
-	{
-		if (RootBaseData->m_uiPersistentId > 0)
-		{
-			FString PersistentId = Utils::ToFString((uint32)RootBaseData->m_uiPersistentId);
-			OutEntityData.MetaData.Emplace(TEXT("PersistentId"), PersistentId);
-		}
-
-		if (RootBaseData->m_pcPersistentId)
-		{
-			FString PersistentUuid = UTF8_TO_TCHAR(RootBaseData->m_pcPersistentId);
-			Utils::RemoveUnwantedChar(PersistentUuid, L'-');
-			OutEntityData.MetaData.Emplace(TEXT("UUID"), PersistentUuid);
-		}
-
-		for (A3DUns32 Index = 0; Index < RootBaseData->m_uiSize; ++Index)
-		{
-			TraverseAttribute(RootBaseData->m_ppAttributes[Index], OutEntityData.MetaData);
-		}
-
-		CADLibrary::TUniqueTSObj<A3DBoundingBoxData> sBoundingBox(Entity);
-		if (sBoundingBox.IsValid())
-		{
-			if ((sBoundingBox->m_sMin.m_dX + sBoundingBox->m_sMin.m_dY + sBoundingBox->m_sMin.m_dZ + sBoundingBox->m_sMax.m_dX + sBoundingBox->m_sMax.m_dY + sBoundingBox->m_sMax.m_dZ) != 0.0)
-			{
-				FString BBox = FString::Printf(TEXT("Min(%f,%f,%f) Max(%f,%f,%f)"), sBoundingBox->m_sMin.m_dX, sBoundingBox->m_sMin.m_dY, sBoundingBox->m_sMin.m_dZ, sBoundingBox->m_sMax.m_dX, sBoundingBox->m_sMax.m_dY, sBoundingBox->m_sMax.m_dZ);
-				OutEntityData.MetaData.Emplace(TEXT("BBox"), BBox);
-			}
-		}
-	}
-
-	if (bIsOccurrence)
-	{
-		FString* InstanceName = OutEntityData.MetaData.Find(TEXT("InstanceName"));
-		FString* Name = OutEntityData.MetaData.Find(TEXT("Name"));
-		FString* OriginalName = OutEntityData.MetaData.Find(TEXT("OriginalName"));
-		FString* PartNumber = OutEntityData.MetaData.Find(TEXT("PartNumber"));
-
-		if (!Name && !OriginalName &&
-			RootBaseData->m_pcName && RootBaseData->m_pcName[0] != '\0')
-		{
-			FString NameStr = UTF8_TO_TCHAR(RootBaseData->m_pcName);
-			FString PartNumberStr = GetOriginalName(NameStr);
-			if (FileType == ECADType::CATIA_3DXML)
-			{
-				PartNumberStr = Get3dxmlOriginalName(PartNumberStr);
-			}
-			if (FileType == ECADType::SOLIDWORKS)
-			{
-				PartNumberStr = GetSWOriginalName(PartNumberStr);
-			}
-			OutEntityData.MetaData.Add(TEXT("PartNumber"), PartNumberStr);
-			PartNumber = OutEntityData.MetaData.Find(TEXT("PartNumber"));
-		}
-
-		if (InstanceName)
-		{
-			OutEntityData.MetaData.Add(TEXT("Name"), *InstanceName);
-		}
-		else if (PartNumber)
-		{
-			OutEntityData.MetaData.Add(TEXT("Name"), *PartNumber);
-		}
-		else if (OriginalName)
-		{
-			OutEntityData.MetaData.Add(TEXT("Name"), *OriginalName);
-		}
-		OutEntityData.Name = OutEntityData.MetaData[L"Name"];
-	}
-	else
-	{
-		FString* Name = OutEntityData.MetaData.Find(TEXT("Name"));
-		FString* OriginalName = OutEntityData.MetaData.Find(TEXT("OriginalName"));
-		FString* PartNumber = OutEntityData.MetaData.Find(TEXT("PartNumber"));
-
-		if (!OriginalName && !PartNumber &&
-			RootBaseData->m_pcName && RootBaseData->m_pcName[0] != '\0')
-		{
-			FString NameStr = UTF8_TO_TCHAR(RootBaseData->m_pcName);
-			FString OriginalNameStr = GetOriginalName(NameStr);
-			OutEntityData.MetaData.Emplace(TEXT("OriginalName"), OriginalNameStr);
-			OriginalName = OutEntityData.MetaData.Find(TEXT("OriginalName"));
-		}
-
-		if (PartNumber)
-		{
-			OutEntityData.MetaData.Add(TEXT("Name"), *PartNumber);
-		}
-		else if (OriginalName)
-		{
-			OutEntityData.MetaData.Add(TEXT("Name"), *OriginalName);
-		}
-
-		Name = OutEntityData.MetaData.Find(TEXT("Name"));
-		if (Name)
-		{
-			OutEntityData.Name = *Name;
-			if (!OriginalName)
-			{
-				OutEntityData.MetaData.Add(TEXT("OriginalName"), OutEntityData.Name);
-			}
-		}
-	}
-
-	return A3D_SUCCESS;
-}
-
-int TraverseGraphics(const A3DGraphics* Graphics)
-{
-	CADLibrary::TUniqueTSObj<A3DGraphicsData> GraphicsData(Graphics);
-	if (GraphicsData.IsValid())
-	{
-		if (GraphicsData->m_uiLayerIndex != A3D_DEFAULT_LAYER)
-		{
-			//graphics->SetAttribute("m_uiLayerIndex", (int)GraphicsData.m_uiLayerIndex);
-		}
-		if (GraphicsData->m_uiStyleIndex != A3D_DEFAULT_STYLE_INDEX)
-		{
-			//graphics->SetAttribute("m_uiStyleIndex", (int)GraphicsData.m_uiStyleIndex);
-		}
-
-		//graphics->SetAttribute("m_usBehaviour", (int)GraphicsData.m_usBehaviour);
-	}
-
-	return A3D_SUCCESS;
-}
-
-int TraverseBaseWithGraphics(const ECADType FileType, const A3DEntity* Entity, FEntityData& OutEntityData, bool bIsOccurrence)
-{
-	CADLibrary::TUniqueTSObj<A3DRootBaseWithGraphicsData> BaseWithGraphicsData(Entity);
-	if (BaseWithGraphicsData.IsValid())
-	{
-		TraverseBase(FileType, Entity, OutEntityData, bIsOccurrence);
-		if (BaseWithGraphicsData->m_pGraphics != NULL)
-		{
-			TraverseGraphics(BaseWithGraphicsData->m_pGraphics);
-		}
-	}
-
-	return A3D_SUCCESS;
-}
-
-int TraverseSource(const A3DEntity* Entity, FEntityData& OutEntityData, bool bIsOccurrence, const ECADType FileType)
-{
-	if (CADLibrary::TechSoftInterface::IsEntityBaseWithGraphicsType(Entity))
-	{
-		return TraverseBaseWithGraphics(FileType, Entity, OutEntityData, bIsOccurrence);
-	}
-	else if (CADLibrary::TechSoftInterface::IsEntityBaseType(Entity))
-	{
-		return TraverseBase(FileType, Entity, OutEntityData, bIsOccurrence);
-	}
-	return A3D_SUCCESS;
-}
-
-int TraverseSpecificMetaData(A3DEModellerType ModellerType, const A3DAsmProductOccurrence* Occurrence, FEntityData& MetaData)
-{
-	//----------- Export Specific information per CAD format -----------
-	switch (ModellerType)
-	{
-	case kA3DModellerSlw:
-	{
-		CADLibrary::TUniqueTSObj<A3DAsmProductOccurrenceDataSLW> SolidWorksSpecificData(Occurrence);
-		if (SolidWorksSpecificData.IsValid())
-		{
-			//A3DUns16 m_usStructSize;									// Reserved: will be initialized by \ref A3D_INITIALIZE_DATA.
-			//A3DUTF8Char* SolidWorksSpecificData->m_psNodeSlwID;									/*!< ID comming from Solidworks, to use to resolve the "PathsInAssemblyTree" */
-
-			//A3DInt32 SolidWorksSpecificData->m_iIndexCfg;										/*!< Cfg Index */
-			//A3DUTF8Char* SolidWorksSpecificData->m_psCfgName;									/*!< Cfg Name in the file*/
-			if (SolidWorksSpecificData->m_psCfgName)
-			{
-				FString ConfigurationName = UTF8_TO_TCHAR(SolidWorksSpecificData->m_psCfgName);
-				MetaData.MetaData.Emplace(TEXT("ConfigurationName"), ConfigurationName);
-				FString ConfigurationIndex = Utils::ToFString((int32) SolidWorksSpecificData->m_iIndexCfg);
-				MetaData.MetaData.Emplace(TEXT("ConfigurationIndex"), ConfigurationIndex);
-				FString Configuration = TEXT("Configuration_") + ConfigurationIndex;
-				MetaData.MetaData.Emplace(Configuration, ConfigurationName);
-			}
-
-			//A3DEProductOccurrenceTypeSLW SolidWorksSpecificData->m_usType;						/*!< Product Occurrence Type */
-
-			//A3DUns32 SolidWorksSpecificData->m_uiAttachementsSize;								/*!< The size of \ref m_psAttachements. */
-			//A3DAsmAttachmentsInfosSLW* SolidWorksSpecificData->m_psAttachements;				/*!< Storage to put assembly attachments information on product occurrence in assembly tree. */
-
-		}
-		break;
-	}
-	case kA3DModellerUnigraphics:
-	{
-		CADLibrary::TUniqueTSObj<A3DAsmProductOccurrenceDataUg> UnigraphicsSpecificData(Occurrence);
-		if (UnigraphicsSpecificData.IsValid())
-		{
-			if (UnigraphicsSpecificData->m_psPartUID)
-			{
-				FString PartUID = UTF8_TO_TCHAR(UnigraphicsSpecificData->m_psPartUID);
-				MetaData.MetaData.Emplace(TEXT("UnigraphicsPartUID"), PartUID);
-			}
-			if (UnigraphicsSpecificData->m_psInstanceFileName)
-			{
-				FString InstanceFileName = UTF8_TO_TCHAR(UnigraphicsSpecificData->m_psInstanceFileName);
-				MetaData.MetaData.Emplace(TEXT("UnigraphicsPartUID"), InstanceFileName);
-			}
-
-			if (UnigraphicsSpecificData->m_uiInstanceTag)
-			{
-				FString InstanceTag = Utils::ToFString((uint32)UnigraphicsSpecificData->m_uiInstanceTag);
-				MetaData.MetaData.Emplace(TEXT("UnigraphicsInstanceTag"), InstanceTag);
-			}
-
-			//// Treat RefSets
-			//if (UnigraphicsSpecificData.m_uiChildrenByRefsetsSize || UnigraphicsSpecificData.m_uiSolidsByRefsetsSize)
-			//{
-			//	// Children by RefSet
-			//	for (A3DUns32 iRefSet = 0; iRefSet < UnigraphicsSpecificData.m_uiChildrenByRefsetsSize; ++iRefSet)
-			//	{
-			//		const A3DElementsByRefsetUg& pRefsetData = UnigraphicsSpecificData.m_asChildrenByRefsets[iRefSet];
-			//		// Init Xml Element if first occurrence
-			//		//refset->SetAttribute("Name", pRefsetData.m_psRefset);
-			//		// Add children
-			//		for (A3DUns32 i = 0; i < pRefsetData.m_uiElementsSize; ++i)
-			//		{
-			//			//refsetChild->SetAttribute("id", pRefsetData.m_auiElements[i]);
-			//		}
-			//	}
-			//	// Solid by RefSet
-			//	for (A3DUns32 iRefSet = 0; iRefSet < UnigraphicsSpecificData.m_uiSolidsByRefsetsSize; ++iRefSet)
-			//	{
-			//		const A3DElementsByRefsetUg& pRefsetData = UnigraphicsSpecificData.m_asSolidsByRefsets[iRefSet];
-			//		//refset->SetAttribute("Name", pRefsetData.m_psRefset);
-			//		// Add Solid
-			//		for (A3DUns32 i = 0; i < pRefsetData.m_uiElementsSize; ++i)
-			//		{
-			//			//refsetChild->SetAttribute("id", pRefsetData.m_auiElements[i]);
-			//		}
-			//	}
-			//}
-		}
-		break;
-	}
-
-	case kA3DModellerCatiaV5:
-	{
-		CADLibrary::TUniqueTSObj<A3DAsmProductOccurrenceDataCV5> CatiaV5SpecificData(Occurrence);
-		if (CatiaV5SpecificData.IsValid())
-		{
-			if (CatiaV5SpecificData->m_iCurrentAsmProductIdentifier)
-			{
-				FString InstanceTag = Utils::ToFString((int32) CatiaV5SpecificData->m_iCurrentAsmProductIdentifier);
-				MetaData.MetaData.Emplace(TEXT("CatiaAsmProductIdentifier"), InstanceTag);
-			}
-
-			if (CatiaV5SpecificData->m_iNotUpdatedAsmProductIdentifier)
-			{
-				FString InstanceTag = Utils::ToFString((int32) CatiaV5SpecificData->m_iNotUpdatedAsmProductIdentifier);
-				MetaData.MetaData.Emplace(TEXT("CatiaNotUpdatedAsmProductIdentifier"), InstanceTag);
-			}
-
-			if (CatiaV5SpecificData->m_aiCLSID[0])
-			{
-				FString InstanceTag = FString::Printf(TEXT("%d %d %d %d"), CatiaV5SpecificData->m_aiCLSID[0], CatiaV5SpecificData->m_aiCLSID[1], CatiaV5SpecificData->m_aiCLSID[2], CatiaV5SpecificData->m_aiCLSID[3]);
-				MetaData.MetaData.Emplace(TEXT("CatiaCLSID"), InstanceTag);
-			}
-
-			if (CatiaV5SpecificData->m_psVersion)
-			{
-				FString Version = UTF8_TO_TCHAR(CatiaV5SpecificData->m_psVersion);
-				MetaData.MetaData.Emplace(TEXT("CatiaVersion"), Version);
-			}
-
-			if (CatiaV5SpecificData->m_psPartNumber)
-			{
-				FString PartNumber = UTF8_TO_TCHAR(CatiaV5SpecificData->m_psPartNumber);
-				MetaData.MetaData.Emplace(TEXT("CatiaPartNumber"), PartNumber);
-			}
-		}
-		break;
-	}
-
-	default:
-		break;
-	}
-	return A3D_SUCCESS;
-}
-
-FMatrixH CreateCoordinateSystem(const A3DMiscCartesianTransformationData& Transformation, double UnitScale = 1.0)
-{
-	FPoint Origin(&Transformation.m_sOrigin.m_dX);
-	FPoint Ox(&Transformation.m_sXVector.m_dX);
-	FPoint Oy(&Transformation.m_sYVector.m_dX);
+	CADKernel::FPoint Origin(&Transformation.m_sOrigin.m_dX);
+	CADKernel::FPoint Ox(&Transformation.m_sXVector.m_dX);
+	CADKernel::FPoint Oy(&Transformation.m_sYVector.m_dX);
 
 	Ox.Normalize();
 	Oy.Normalize();
@@ -520,87 +85,18 @@ FMatrixH CreateCoordinateSystem(const A3DMiscCartesianTransformationData& Transf
 	{
 		Origin *= UnitScale;
 	}
-	FPoint Oz = Ox ^ Oy;
+	CADKernel::FPoint Oz = Ox ^ Oy;
 
-	FMatrixH Matrix = FMatrixH(Origin, Ox, Oy, Oz);
+	CADKernel::FMatrixH Matrix = CADKernel::FMatrixH(Origin, Ox, Oy, Oz);
 
 	if (!FMath::IsNearlyEqual(Transformation.m_sScale.m_dX, 1.) || !FMath::IsNearlyEqual(Transformation.m_sScale.m_dY, 1.) || !FMath::IsNearlyEqual(Transformation.m_sScale.m_dZ, 1.))
 	{
-		FMatrixH Scale = FMatrixH::MakeScaleMatrix(Transformation.m_sScale.m_dX, Transformation.m_sScale.m_dY, Transformation.m_sScale.m_dZ);
+		CADKernel::FMatrixH Scale = CADKernel::FMatrixH::MakeScaleMatrix(Transformation.m_sScale.m_dX, Transformation.m_sScale.m_dY, Transformation.m_sScale.m_dZ);
 		Matrix *= Scale;
 	}
 	return Matrix;
 }
 
-FMatrixH TraverseTransformation3D(const A3DMiscCartesianTransformation* CartesianTransformation, double UnitScale)
-{
-	FMatrixH Matrix;
-	if (CartesianTransformation)
-	{
-		CADLibrary::TUniqueTSObj<A3DMiscCartesianTransformationData> CartesianTransformationData(CartesianTransformation);
-		if (CartesianTransformationData.IsValid())
-		{
-			Matrix = CreateCoordinateSystem(*CartesianTransformationData, UnitScale);
-		}
-	}
-	return Matrix;
-}
-
-FMatrixH TraverseCoordinateSystem(const A3DRiCoordinateSystem* CoordinateSystem, double UnitScale)
-{
-
-	FEntityData MetaData;
-	TechSoftUtils::TraverseSource(CoordinateSystem, MetaData);
-
-	FMatrixH Matrix;
-	CADLibrary::TUniqueTSObj<A3DRiCoordinateSystemData> CoordinateSystemData(CoordinateSystem);
-	if (CoordinateSystemData.IsValid())
-	{
-		Matrix = TraverseTransformation3D(CoordinateSystemData->m_pTransformation, UnitScale);
-	}
-	return Matrix;
-}
-
-FMatrixH  TraverseGeneralTransformation(const A3DMiscGeneralTransformation* GeneralTransformation, double UnitScale)
-{
-
-	FMatrixH Matrix;
-	CADLibrary::TUniqueTSObj<A3DMiscGeneralTransformationData> GeneralTransformationData(GeneralTransformation);
-	if (GeneralTransformationData.IsValid())
-	{
-		for (int32 Index = 0; Index < 16; ++Index)
-		{
-			Matrix[Index] = GeneralTransformationData->m_adCoeff[Index];
-		}
-		ensureCADKernel(!FMath::IsNearlyEqual(UnitScale, 1.));
-	}
-	return Matrix;
-}
-
-FMatrixH TraverseTransformation(const A3DMiscCartesianTransformation* Transformation3D, double UnitScale)
-{
-	if (Transformation3D == NULL)
-	{
-		return FMatrixH();
-	}
-
-	FEntityData MetaData;
-	TraverseSource(Transformation3D, MetaData, false);
-
-	A3DEEntityType Type = kA3DTypeUnknown;
-	CADLibrary::TechSoftInterface::GetEntityType(Transformation3D, &Type);
-
-	if (Type == kA3DTypeMiscCartesianTransformation)
-	{
-		return TraverseTransformation3D(Transformation3D, UnitScale);
-	}
-	else if (Type == kA3DTypeMiscGeneralTransformation)
-	{
-		return TraverseGeneralTransformation(Transformation3D, UnitScale);
-	}
-	return FMatrixH();
-}
-
 void FillInt32Array(const int32 Count, const A3DInt32* Values, TArray<int32>& OutInt32Array)
 {
 	OutInt32Array.Reserve(Count);
@@ -632,7 +128,7 @@ void FillDoubleArray(const int32 UCount, const int32 VCount, const double* Value
 	}
 }
 
-void FillPointArray(const int32 Count, const A3DVector3dData* Points, TArray<FPoint>& OutPointsArray, double UnitScale = 1.0)
+void FillPointArray(const int32 Count, const A3DVector3dData* Points, TArray<CADKernel::FPoint>& OutPointsArray, double UnitScale = 1.0)
 {
 	OutPointsArray.Reserve(Count);
 	for (int32 Index = 0; Index < Count; Index++)
@@ -642,14 +138,14 @@ void FillPointArray(const int32 Count, const A3DVector3dData* Points, TArray<FPo
 
 	if (!FMath::IsNearlyEqual(UnitScale, 1.))
 	{
-		for (FPoint& Point : OutPointsArray)
+		for (CADKernel::FPoint& Point : OutPointsArray)
 		{
 			Point *= UnitScale;
 		}
 	}
 };
 
-void FillPointArray(const int32 UCount, const int32 VCount, const A3DVector3dData* Points, TArray<FPoint>& OutPointsArray, double UnitScale = 1.0)
+void FillPointArray(const int32 UCount, const int32 VCount, const A3DVector3dData* Points, TArray<CADKernel::FPoint>& OutPointsArray, double UnitScale = 1.0)
 {
 	OutPointsArray.SetNum(UCount * VCount);
 	for (int32 Undex = 0, PointIndex = 0; Undex < UCount; ++Undex)
@@ -663,19 +159,19 @@ void FillPointArray(const int32 UCount, const int32 VCount, const A3DVector3dDat
 
 	if (!FMath::IsNearlyEqual(UnitScale, 1.))
 	{
-		for (FPoint& Point : OutPointsArray)
+		for (CADKernel::FPoint& Point : OutPointsArray)
 		{
 			Point *= UnitScale;
 		}
 	}
 };
 
-FSurfacicBoundary GetSurfacicBoundary(A3DDomainData& Domain, double UnitScale = 1.0, bool bSwapUV = false)
+CADKernel::FSurfacicBoundary GetSurfacicBoundary(A3DDomainData& Domain, double UnitScale = 1.0, bool bSwapUV = false)
 {
-	EIso UIndex = bSwapUV ? EIso::IsoV : EIso::IsoU;
-	EIso VIndex = bSwapUV ? EIso::IsoU : EIso::IsoV;
+	CADKernel::EIso UIndex = bSwapUV ? CADKernel::EIso::IsoV : CADKernel::EIso::IsoU;
+	CADKernel::EIso VIndex = bSwapUV ? CADKernel::EIso::IsoU : CADKernel::EIso::IsoV;
 
-	FSurfacicBoundary Boundary;
+	CADKernel::FSurfacicBoundary Boundary;
 	Boundary[UIndex].Min = Domain.m_sMin.m_dX * UnitScale;
 	Boundary[UIndex].Max = Domain.m_sMax.m_dX * UnitScale;
 	Boundary[VIndex].Min = Domain.m_sMin.m_dY * UnitScale;
@@ -683,23 +179,25 @@ FSurfacicBoundary GetSurfacicBoundary(A3DDomainData& Domain, double UnitScale =
 	return Boundary;
 }
 
-FLinearBoundary GetLinearBoundary(A3DIntervalData& A3DDomain)
+CADKernel::FLinearBoundary GetLinearBoundary(A3DIntervalData& A3DDomain)
 {
-	FLinearBoundary Domain(A3DDomain.m_dMin, A3DDomain.m_dMax);
+	CADKernel::FLinearBoundary Domain(A3DDomain.m_dMin, A3DDomain.m_dMax);
 	return Domain;
 }
 
-FLinearBoundary GetLinearBoundary(const A3DCrvBase* A3DCurve)
+CADKernel::FLinearBoundary GetLinearBoundary(const A3DCrvBase* A3DCurve)
 {
-	CADLibrary::TUniqueTSObj<A3DIntervalData> A3DDomain(A3DCurve);
-	FLinearBoundary Domain(A3DDomain->m_dMin, A3DDomain->m_dMax);
+	TUniqueTSObj<A3DIntervalData> A3DDomain(A3DCurve);
+	CADKernel::FLinearBoundary Domain(A3DDomain->m_dMin, A3DDomain->m_dMax);
 	return Domain;
 }
 
 } // ns TechSoftUtils
 
-FTechSoftBridge::FTechSoftBridge(FSession& InSession, FCADFileReport& InReport)
-	: Session(InSession)
+FTechSoftBridge::FTechSoftBridge(FTechSoftFileParser& InParser, CADKernel::FSession& InSession, CADKernel::FCADFileReport& InReport)
+	: Parser(InParser)
+	, Session(InSession)
+	, Model(InSession.GetModel())
 	, Report(InReport)
 	, GeometricTolerance(Session.GetGeometricTolerance())
 	, SquareGeometricTolerance(FMath::Square(Session.GetGeometricTolerance()))
@@ -707,51 +205,84 @@ FTechSoftBridge::FTechSoftBridge(FSession& InSession, FCADFileReport& InReport)
 {
 }
 
-TSharedRef<FBody> FTechSoftBridge::AddBody(A3DRiBrepModel* BrepModel, FEntityData& ParentData, double FileUnit)
+const A3DRiBrepModel* FTechSoftBridge::GetA3DBody(CADKernel::FBody* BRepModel)
 {
-	Report.BodyCount++;
+	const A3DEntity** BodyPtr = CADKernelToTechSoft.Find(BRepModel);
+	if (BodyPtr)
+	{
+		return (A3DRiBrepModel*) *BodyPtr;
+	}
+	return nullptr;
+}
 
-	UnitScale = FileUnit;
+CADKernel::FBody* FTechSoftBridge::GetBody(A3DRiBrepModel* A3DBRepModel)
+{
+	CADKernel::FBody** BodyPtr = TechSoftToCADKernel.Find(A3DBRepModel);
+	if (BodyPtr &&!(*BodyPtr)->IsDeleted())
+	{
+		return *BodyPtr;
+	}
+	return nullptr;
+}
 
-	FEntityData BrepMetaData;
-	TechSoftUtils::TraverseSource(BrepModel, BrepMetaData, false);
+CADKernel::FBody* FTechSoftBridge::AddBody(A3DRiBrepModel* A3DBRepModel, TMap<FString, FString> MetaData)
+{
+	Report.BodyCount++;
+
+	UnitScale = Parser.GetFileUnit();
 
-	TSharedPtr<FEntity>* EntityPtr = TechSoftToEntity.Find(BrepModel);
-	if (EntityPtr)
+	FEntityMetaData BRepMetaData;
+	Parser.ExtractMetaData(A3DBRepModel, BRepMetaData);
+	FString* Name = MetaData.Find(TEXT("Name"));
+	if(Name != nullptr)
 	{
-		TSharedPtr<FBody> Body = StaticCastSharedPtr<FBody>(*EntityPtr);
-		if (!Body.IsValid())
+		BRepMetaData.MetaData.FindOrAdd(TEXT("Name")) = *Name;
+	}
+	
+	CADKernel::FBody** BodyPtr = TechSoftToCADKernel.Find(A3DBRepModel);
+	if (BodyPtr)
+	{
+		if((*BodyPtr)->IsDeleted())
 		{
-			return Body.ToSharedRef();
+			return nullptr;
 		}
+		return *BodyPtr;
 	}
 
-	TSharedRef<FBody> Body = FEntity::MakeShared<FBody>();
-	TechSoftToEntity.Add(BrepModel, Body);
-
-	AddMetadata(BrepMetaData, *Body);
+	TSharedRef<CADKernel::FBody> Body = CADKernel::FEntity::MakeShared<CADKernel::FBody>();
+	AddMetadata(BRepMetaData, *Body);
 
-	CADLibrary::TUniqueTSObj<A3DRiBrepModelData> BRepModelData(BrepModel);
+	TUniqueTSObj<A3DRiBrepModelData> BRepModelData(A3DBRepModel);
 	if (BRepModelData.IsValid())
 	{
 		TraverseBrepData(BRepModelData->m_pBrepData, Body);
 	}
 
-	return Body;
+	if(Body->FaceCount() == 0)
+	{
+		Body->SetDeleted();
+		return nullptr;
+	}
+
+	Model.Add(Body);
+	TechSoftToCADKernel.Add(A3DBRepModel, &*Body);
+	CADKernelToTechSoft.Add(&*Body, A3DBRepModel);
+
+	return &*Body;
 }
 
-void FTechSoftBridge::TraverseBrepData(const A3DTopoBrepData* A3DBrepData, TSharedRef<FBody>& Body)
+void FTechSoftBridge::TraverseBrepData(const A3DTopoBrepData* A3DBrepData, TSharedRef<CADKernel::FBody>& Body)
 {
-	FEntityData MetaData;
-	TechSoftUtils::TraverseSource(A3DBrepData, MetaData, false);
+	FEntityMetaData MetaData;
+	Parser.ExtractMetaData(A3DBrepData, MetaData);
 
 	{
-		CADLibrary::TUniqueTSObj<A3DTopoBodyData> TopoBodyData(A3DBrepData);
+		TUniqueTSObj<A3DTopoBodyData> TopoBodyData(A3DBrepData);
 		if (TopoBodyData.IsValid())
 		{
 			if (TopoBodyData->m_pContext)
 			{
-				CADLibrary::TUniqueTSObj<A3DTopoContextData> TopoContextData(TopoBodyData->m_pContext);
+				TUniqueTSObj<A3DTopoContextData> TopoContextData(TopoBodyData->m_pContext);
 				if (TopoContextData.IsValid())
 				{
 					if (TopoContextData->m_bHaveScale)
@@ -763,7 +294,7 @@ void FTechSoftBridge::TraverseBrepData(const A3DTopoBrepData* A3DBrepData, TShar
 		}
 	}
 
-	CADLibrary::TUniqueTSObj<A3DTopoBrepDataData> TopoBrepData(A3DBrepData);
+	TUniqueTSObj<A3DTopoBrepDataData> TopoBrepData(A3DBrepData);
 	if (TopoBrepData.IsValid())
 	{
 		for (A3DUns32 Index = 0; Index < TopoBrepData->m_uiConnexSize; ++Index)
@@ -773,12 +304,12 @@ void FTechSoftBridge::TraverseBrepData(const A3DTopoBrepData* A3DBrepData, TShar
 	}
 }
 
-void FTechSoftBridge::TraverseConnex(const A3DTopoConnex* A3DTopoConnex, TSharedRef<FBody>& Body)
+void FTechSoftBridge::TraverseConnex(const A3DTopoConnex* A3DTopoConnex, TSharedRef<CADKernel::FBody>& Body)
 {
-	FEntityData MetaData;
-	TechSoftUtils::TraverseSource(A3DTopoConnex, MetaData);
+	FEntityMetaData MetaData;
+	Parser.ExtractMetaData(A3DTopoConnex, MetaData);
 
-	CADLibrary::TUniqueTSObj<A3DTopoConnexData> TopoConnexData(A3DTopoConnex);
+	TUniqueTSObj<A3DTopoConnexData> TopoConnexData(A3DTopoConnex);
 	if (TopoConnexData.IsValid())
 	{
 		for (A3DUns32 Index = 0; Index < TopoConnexData->m_uiShellSize; ++Index)
@@ -788,33 +319,33 @@ void FTechSoftBridge::TraverseConnex(const A3DTopoConnex* A3DTopoConnex, TShared
 	}
 }
 
-void FTechSoftBridge::TraverseShell(const A3DTopoShell* A3DShell, TSharedRef<FBody>& Body)
+void FTechSoftBridge::TraverseShell(const A3DTopoShell* A3DShell, TSharedRef<CADKernel::FBody>& Body)
 {
-	FEntityData MetaData;
-	TechSoftUtils::TraverseSource(A3DShell, MetaData);
+	FEntityMetaData MetaData;
+	Parser.ExtractMetaData(A3DShell, MetaData);
 
-	TSharedRef<FShell> Shell = FEntity::MakeShared<FShell>();
+	TSharedRef<CADKernel::FShell> Shell = CADKernel::FEntity::MakeShared<CADKernel::FShell>();
 	Body->AddShell(Shell);
 	Report.ShellCount++;
 
 	AddMetadata(MetaData, *Shell);
 
-	CADLibrary::TUniqueTSObj<A3DTopoShellData> ShellData(A3DShell);
+	TUniqueTSObj<A3DTopoShellData> ShellData(A3DShell);
 	if (ShellData.IsValid())
 	{
-		for (A3DUns32 ui = 0; ui < ShellData->m_uiFaceSize; ++ui)
+		for (A3DUns32 Index = 0; Index < ShellData->m_uiFaceSize; ++Index)
 		{
-			AddFace(ShellData->m_ppFaces[ui], Shell);
+			AddFace(ShellData->m_ppFaces[Index], ShellData->m_pucOrientationWithShell[Index] == 1 ? CADKernel::Front : CADKernel::Back, Shell);
 		}
 	}
 }
 
 static bool bUseCurveAsNurbs = true;
-TSharedPtr<FCurve> FTechSoftBridge::AddCurve(const A3DCrvBase* A3DCurve, const TechSoftUtils::FUVReparameterization& UVReparameterization)
+TSharedPtr<CADKernel::FCurve> FTechSoftBridge::AddCurve(const A3DCrvBase* A3DCurve, const TechSoftUtils::FUVReparameterization& UVReparameterization)
 {
-	TSharedPtr<FCurve> Curve = TSharedPtr<FCurve>();
+	TSharedPtr<CADKernel::FCurve> Curve = TSharedPtr<CADKernel::FCurve>();
 	A3DEEntityType eType;
-	A3DInt32 Ret = CADLibrary::TechSoftInterface::GetEntityType(A3DCurve, &eType);
+	A3DInt32 Ret = TechSoftInterface::GetEntityType(A3DCurve, &eType);
 	if (Ret == A3D_SUCCESS)
 	{
 		Report.CurveCount++;
@@ -857,12 +388,12 @@ TSharedPtr<FCurve> FTechSoftBridge::AddCurve(const A3DCrvBase* A3DCurve, const T
 		}
 	}
 
-	FLinearBoundary Boundary = TechSoftUtils::GetLinearBoundary(A3DCurve);
+	CADKernel::FLinearBoundary Boundary = TechSoftUtils::GetLinearBoundary(A3DCurve);
 
 	return Curve;
 }
 
-TSharedPtr<FCurve> FTechSoftBridge::AddCurveLine(const A3DCrvLine* A3DCurve, const TechSoftUtils::FUVReparameterization& UVReparameterization)
+TSharedPtr<CADKernel::FCurve> FTechSoftBridge::AddCurveLine(const A3DCrvLine* A3DCurve, const TechSoftUtils::FUVReparameterization& UVReparameterization)
 {
 	Report.CurveLineCount++;
 
@@ -872,23 +403,18 @@ TSharedPtr<FCurve> FTechSoftBridge::AddCurveLine(const A3DCrvLine* A3DCurve, con
 	}
 
 
-	CADLibrary::TUniqueTSObj<A3DCrvLineData> CrvLineData(A3DCurve);
+	TUniqueTSObj<A3DCrvLineData> CrvLineData(A3DCurve);
 	if (!CrvLineData.IsValid())
 	{
-		return TSharedPtr<FCurve>();
+		return TSharedPtr<CADKernel::FCurve>();
 	}
 
 	bool bIs2D = (bool)CrvLineData->m_bIs2D;
-
-	//setting->SetAttribute("m_bIs2D", (int)sData.m_bIs2D);
-	//traverseTransfo3d(sData.m_sTrsf, crv);
-	//traverseParam(&sData.m_sParam, crv);
-	//A3DCrvLineGet(NULL, &sData);
-
-	return TSharedPtr<FCurve>();
+	// Todo
+	return TSharedPtr<CADKernel::FCurve>();
 }
 
-TSharedPtr<FCurve> FTechSoftBridge::AddCurveCircle(const A3DCrvCircle* A3DCurve, const TechSoftUtils::FUVReparameterization& UVReparameterization)
+TSharedPtr<CADKernel::FCurve> FTechSoftBridge::AddCurveCircle(const A3DCrvCircle* A3DCurve, const TechSoftUtils::FUVReparameterization& UVReparameterization)
 {
 	Report.CurveCircleCount++;
 
@@ -897,23 +423,18 @@ TSharedPtr<FCurve> FTechSoftBridge::AddCurveCircle(const A3DCrvCircle* A3DCurve,
 		return AddCurveAsNurbs(A3DCurve, UVReparameterization);
 	}
 
-	CADLibrary::TUniqueTSObj<A3DCrvCircleData> CrvCircleData(A3DCurve);
+	TUniqueTSObj<A3DCrvCircleData> CrvCircleData(A3DCurve);
 	if (!CrvCircleData.IsValid())
 	{
-		return TSharedPtr<FCurve>();
+		return TSharedPtr<CADKernel::FCurve>();
 	}
 
 	bool bIs2D = (bool)CrvCircleData->m_bIs2D;
-
-	//setting->SetAttribute("m_bIs2D", (int)sData.m_bIs2D);
-	//traverseTransfo3d(sData.m_sTrsf, crv);
-	//traverseParam(&sData.m_sParam, crv);
-	//_SetDoubleAttribute(crv, "m_dRadius", sData.m_dRadius);
-
-	return TSharedPtr<FCurve>();
+	// Todo
+	return TSharedPtr<CADKernel::FCurve>();
 }
 
-TSharedPtr<FCurve> FTechSoftBridge::AddCurveEllipse(const A3DCrvEllipse* A3DCurve, const TechSoftUtils::FUVReparameterization& UVReparameterization)
+TSharedPtr<CADKernel::FCurve> FTechSoftBridge::AddCurveEllipse(const A3DCrvEllipse* A3DCurve, const TechSoftUtils::FUVReparameterization& UVReparameterization)
 {
 	Report.CurveEllipseCount++;
 
@@ -922,24 +443,18 @@ TSharedPtr<FCurve> FTechSoftBridge::AddCurveEllipse(const A3DCrvEllipse* A3DCurv
 		return AddCurveAsNurbs(A3DCurve, UVReparameterization);
 	}
 
-	CADLibrary::TUniqueTSObj<A3DCrvEllipseData> CrvEllipseData(A3DCurve);
+	TUniqueTSObj<A3DCrvEllipseData> CrvEllipseData(A3DCurve);
 	if (!CrvEllipseData.IsValid())
 	{
-		return TSharedPtr<FCurve>();
+		return TSharedPtr<CADKernel::FCurve>();
 	}
 
 	bool bIs2D = (bool)CrvEllipseData->m_bIs2D;
-
-	//setting->SetAttribute("m_bIs2D", (int)sData.m_bIs2D);
-	//traverseTransfo3d(sData.m_sTrsf, crv);
-	//traverseParam(&sData.m_sParam, crv);
-	//_SetDoubleAttribute(crv, "m_dXRadius", sData.m_dXRadius);
-	//_SetDoubleAttribute(crv, "m_dYRadius", sData.m_dYRadius);
-
-	return TSharedPtr<FCurve>();
+	// Todo
+	return TSharedPtr<CADKernel::FCurve>();
 }
 
-TSharedPtr<FCurve> FTechSoftBridge::AddCurveParabola(const A3DCrvParabola* A3DCurve, const TechSoftUtils::FUVReparameterization& UVReparameterization)
+TSharedPtr<CADKernel::FCurve> FTechSoftBridge::AddCurveParabola(const A3DCrvParabola* A3DCurve, const TechSoftUtils::FUVReparameterization& UVReparameterization)
 {
 	Report.CurveParabolaCount++;
 
@@ -948,24 +463,18 @@ TSharedPtr<FCurve> FTechSoftBridge::AddCurveParabola(const A3DCrvParabola* A3DCu
 		return AddCurveAsNurbs(A3DCurve, UVReparameterization);
 	}
 
-	CADLibrary::TUniqueTSObj<A3DCrvParabolaData> CrvParabolaData(A3DCurve);
+	TUniqueTSObj<A3DCrvParabolaData> CrvParabolaData(A3DCurve);
 	if (!CrvParabolaData.IsValid())
 	{
-		return TSharedPtr<FCurve>();
+		return TSharedPtr<CADKernel::FCurve>();
 	}
 
 	bool bIs2D = (bool)CrvParabolaData->m_bIs2D;
-
-	//setting->SetAttribute("m_bIs2D", (int)sData.m_bIs2D);
-	//	traverseTransfo3d(sData.m_sTrsf, crv);
-	//	traverseParam(&sData.m_sParam, crv);
-	//	_SetDoubleAttribute(crv, "m_dFocal", sData.m_dFocal);
-	//	crv->SetAttribute("m_cParamType", (int)sData.m_cParamType);
-
-	return TSharedPtr<FCurve>();
+	// Todo
+	return TSharedPtr<CADKernel::FCurve>();
 }
 
-TSharedPtr<FCurve> FTechSoftBridge::AddCurveHyperbola(const A3DCrvHyperbola* A3DCurve, const TechSoftUtils::FUVReparameterization& UVReparameterization)
+TSharedPtr<CADKernel::FCurve> FTechSoftBridge::AddCurveHyperbola(const A3DCrvHyperbola* A3DCurve, const TechSoftUtils::FUVReparameterization& UVReparameterization)
 {
 	Report.CurveHyperbolaCount++;
 
@@ -974,25 +483,18 @@ TSharedPtr<FCurve> FTechSoftBridge::AddCurveHyperbola(const A3DCrvHyperbola* A3D
 		return AddCurveAsNurbs(A3DCurve, UVReparameterization);
 	}
 
-	CADLibrary::TUniqueTSObj<A3DCrvHyperbolaData> CrvHyperbolaData(A3DCurve);
+	TUniqueTSObj<A3DCrvHyperbolaData> CrvHyperbolaData(A3DCurve);
 	if (!CrvHyperbolaData.IsValid())
 	{
-		return TSharedPtr<FCurve>();
+		return TSharedPtr<CADKernel::FCurve>();
 	}
 
 	bool bIs2D = (bool)CrvHyperbolaData->m_bIs2D;
-
-	//setting->SetAttribute("m_bIs2D", (int)sData.m_bIs2D);
-	//traverseTransfo3d(sData.m_sTrsf, crv);
-	//traverseParam(&sData.m_sParam, crv);
-	//_SetDoubleAttribute(crv, "m_dSemiAxis", sData.m_dSemiAxis);
-	//_SetDoubleAttribute(crv, "m_dSemiImageAxis", sData.m_dSemiImageAxis);
-	//crv->SetAttribute("m_cParamType", (int)sData.m_cParamType);
-
-	return TSharedPtr<FCurve>();
+	// Todo
+	return TSharedPtr<CADKernel::FCurve>();
 }
 
-TSharedPtr<FCurve> FTechSoftBridge::AddCurveHelix(const A3DCrvHelix* A3DCurve, const TechSoftUtils::FUVReparameterization& UVReparameterization)
+TSharedPtr<CADKernel::FCurve> FTechSoftBridge::AddCurveHelix(const A3DCrvHelix* A3DCurve, const TechSoftUtils::FUVReparameterization& UVReparameterization)
 {
 	Report.CurveHelixCount++;
 
@@ -1001,18 +503,18 @@ TSharedPtr<FCurve> FTechSoftBridge::AddCurveHelix(const A3DCrvHelix* A3DCurve, c
 		return AddCurveAsNurbs(A3DCurve, UVReparameterization);
 	}
 
-	CADLibrary::TUniqueTSObj<A3DCrvHelixData> CrvHelixData(A3DCurve);
+	TUniqueTSObj<A3DCrvHelixData> CrvHelixData(A3DCurve);
 	if (!CrvHelixData.IsValid())
 	{
-		return TSharedPtr<FCurve>();
+		return TSharedPtr<CADKernel::FCurve>();
 	}
 
 	bool bIs2D = (bool)CrvHelixData->m_bIs2D;
 
-	return TSharedPtr<FCurve>();
+	return TSharedPtr<CADKernel::FCurve>();
 }
 
-TSharedPtr<FCurve> FTechSoftBridge::AddCurvePolyLine(const A3DCrvPolyLine* A3DCurve, const TechSoftUtils::FUVReparameterization& UVReparameterization)
+TSharedPtr<CADKernel::FCurve> FTechSoftBridge::AddCurvePolyLine(const A3DCrvPolyLine* A3DCurve, const TechSoftUtils::FUVReparameterization& UVReparameterization)
 {
 	Report.CurvePolyLineCount++;
 
@@ -1021,24 +523,18 @@ TSharedPtr<FCurve> FTechSoftBridge::AddCurvePolyLine(const A3DCrvPolyLine* A3DCu
 		return AddCurveAsNurbs(A3DCurve, UVReparameterization);
 	}
 
-	CADLibrary::TUniqueTSObj<A3DCrvPolyLineData> CrvPolyLineData(A3DCurve);
+	TUniqueTSObj<A3DCrvPolyLineData> CrvPolyLineData(A3DCurve);
 	if (!CrvPolyLineData.IsValid())
 	{
-		return TSharedPtr<FCurve>();
+		return TSharedPtr<CADKernel::FCurve>();
 	}
 
 	bool bIs2D = (bool)CrvPolyLineData->m_bIs2D;
-
-	//setting->SetAttribute("m_bIs2D", (int)sData.m_bIs2D);
-	//traverseTransfo3d(sData.m_sTrsf, crv);
-	//traverseParam(&sData.m_sParam, crv);
-	//traversePoints("m_pPts", sData.m_uiSize, sData.m_pPts, crv);
-
-
-	return TSharedPtr<FCurve>();
+	// Todo
+	return TSharedPtr<CADKernel::FCurve>();
 }
 
-TSharedPtr<FCurve> FTechSoftBridge::AddCurveComposite(const A3DCrvComposite* A3DCurve, const TechSoftUtils::FUVReparameterization& UVReparameterization)
+TSharedPtr<CADKernel::FCurve> FTechSoftBridge::AddCurveComposite(const A3DCrvComposite* A3DCurve, const TechSoftUtils::FUVReparameterization& UVReparameterization)
 {
 	Report.CurveCompositeCount++;
 
@@ -1047,30 +543,20 @@ TSharedPtr<FCurve> FTechSoftBridge::AddCurveComposite(const A3DCrvComposite* A3D
 		return AddCurveAsNurbs(A3DCurve, UVReparameterization);
 	}
 
-	CADLibrary::TUniqueTSObj<A3DCrvCompositeData> CrvCompositeData(A3DCurve);
+	TUniqueTSObj<A3DCrvCompositeData> CrvCompositeData(A3DCurve);
 	if (!CrvCompositeData.IsValid())
 	{
-		return TSharedPtr<FCurve>();
+		return TSharedPtr<CADKernel::FCurve>();
 	}
 
 	bool bIs2D = (bool)CrvCompositeData->m_bIs2D;
-
-	//setting->SetAttribute("m_bIs2D", (int)sData.m_bIs2D);
-	//traverseTransfo3d(sData.m_sTrsf, crv);
-	//traverseParam(&sData.m_sParam, crv);
-	//crv->SetAttribute("m_bClosed", (int)sData.m_bClosed);
-	//traverseBools("m_pbSenses", sData.m_uiSize, sData.m_pbSenses, crv);
-
-	//A3DUns32 ui, uiSize = sData.m_uiSize;
-	//for (ui = 0; ui < uiSize; ++ui)
-	//	traverseCurve(sData.m_ppCurves[ui], crv);
-
-	return TSharedPtr<FCurve>();
+	// Todo
+	return TSharedPtr<CADKernel::FCurve>();
 }
 
-TSharedPtr<FCurve> AddCurveNurbsFromData(A3DCrvNurbsData& A3DNurbs, const TechSoftUtils::FUVReparameterization& UVReparameterization)
+TSharedPtr<CADKernel::FCurve> AddCurveNurbsFromData(A3DCrvNurbsData& A3DNurbs, const TechSoftUtils::FUVReparameterization& UVReparameterization)
 {
-	FNurbsCurveData Nurbs;
+	CADKernel::FNurbsCurveData Nurbs;
 	Nurbs.Dimension = A3DNurbs.m_bIs2D ? 2 : 3;
 	Nurbs.bIsRational = (bool)A3DNurbs.m_bRational;
 	Nurbs.Degree = A3DNurbs.m_uiDegree;
@@ -1078,7 +564,7 @@ TSharedPtr<FCurve> AddCurveNurbsFromData(A3DCrvNurbsData& A3DNurbs, const TechSo
 	TechSoftUtils::FillPointArray(A3DNurbs.m_uiCtrlSize, A3DNurbs.m_pCtrlPts, Nurbs.Poles);
 	if (UVReparameterization.NeedApply())
 	{
-		for (FPoint& Point : Nurbs.Poles)
+		for (CADKernel::FPoint& Point : Nurbs.Poles)
 		{
 			UVReparameterization.Apply(Point);
 		}
@@ -1090,49 +576,29 @@ TSharedPtr<FCurve> AddCurveNurbsFromData(A3DCrvNurbsData& A3DNurbs, const TechSo
 		TechSoftUtils::FillDoubleArray(A3DNurbs.m_uiCtrlSize, A3DNurbs.m_pdWeights, Nurbs.Weights);
 	}
 
-	//A3DNurbs.m_eKnotType;
-	//kA3DKnotTypeUniformKnots,			/*!< Uniform. */
-	//kA3DKnotTypeUnspecified,			/*!< No particularity. */
-	//kA3DKnotTypeQuasiUniformKnots,	/*!< Quasi-uniform. */
-	//kA3DKnotTypePieceWiseBezierKnots	/*!< Extrema with multiplicities of degree + 1, internal is degree. */
-	ensureCADKernel(A3DNurbs.m_eKnotType != kA3DKnotTypePieceWiseBezierKnots);
-
-	//A3DNurbs.m_eCurveForm;
-	//kA3DBSplineCurveFormUnspecified,	/*!< No particularity. */
-	//kA3DBSplineCurveFormPolyline,		/*!< Polyline. */
-	//kA3DBSplineCurveFormCircularArc,	/*!< Circle arc. */
-	//kA3DBSplineCurveFormEllipticArc,	/*!< Elliptic arc. */
-	//kA3DBSplineCurveFormParabolicArc,	/*!< Parabolic arc. */
-	//kA3DBSplineCurveFormHyperbolicArc	/*!< Hyperbolic arc. */
-	ensureCADKernel(A3DNurbs.m_eCurveForm != kA3DBSplineCurveFormPolyline);
-	ensureCADKernel(A3DNurbs.m_eCurveForm != kA3DBSplineCurveFormCircularArc);
-	ensureCADKernel(A3DNurbs.m_eCurveForm != kA3DBSplineCurveFormEllipticArc);
-	ensureCADKernel(A3DNurbs.m_eCurveForm != kA3DBSplineCurveFormParabolicArc);
-	ensureCADKernel(A3DNurbs.m_eCurveForm != kA3DBSplineCurveFormHyperbolicArc);
-
 	A3DCrvNurbsGet(NULL, &A3DNurbs);
 
-	return FEntity::MakeShared<FNURBSCurve>(Nurbs);
+	return CADKernel::FEntity::MakeShared<CADKernel::FNURBSCurve>(Nurbs);
 }
 
-TSharedPtr<FCurve> FTechSoftBridge::AddCurveNurbs(const A3DCrvNurbs* A3DNurbs, const TechSoftUtils::FUVReparameterization& UVReparameterization)
+TSharedPtr<CADKernel::FCurve> FTechSoftBridge::AddCurveNurbs(const A3DCrvNurbs* A3DNurbs, const TechSoftUtils::FUVReparameterization& UVReparameterization)
 {
 	Report.CurveLineCount++;
 
-	CADLibrary::TUniqueTSObj<A3DCrvNurbsData> CrvNurbsData(A3DNurbs);
+	TUniqueTSObj<A3DCrvNurbsData> CrvNurbsData(A3DNurbs);
 	if (!CrvNurbsData.IsValid())
 	{
-		return TSharedPtr<FCurve>();
+		return TSharedPtr<CADKernel::FCurve>();
 	}
 
 	return AddCurveNurbsFromData(*CrvNurbsData, UVReparameterization);
 }
 
-TSharedPtr<FCurve> FTechSoftBridge::AddCurveAsNurbs(const A3DCrvBase* A3DCurve, const TechSoftUtils::FUVReparameterization& UVReparameterization)
+TSharedPtr<CADKernel::FCurve> FTechSoftBridge::AddCurveAsNurbs(const A3DCrvBase* A3DCurve, const TechSoftUtils::FUVReparameterization& UVReparameterization)
 {
 	Report.CurveAsNurbsCount++;
 
-	CADLibrary::TUniqueTSObj<A3DCrvNurbsData> NurbsData;
+	TUniqueTSObj<A3DCrvNurbsData> NurbsData;
 
 	A3DDouble Tolerance = 1e-3;
 	A3DBool bUseSameParameterization = true;
@@ -1140,69 +606,43 @@ TSharedPtr<FCurve> FTechSoftBridge::AddCurveAsNurbs(const A3DCrvBase* A3DCurve,
 
 	if (!NurbsData.IsValid())
 	{
-		return TSharedPtr<FCurve>();
+		return TSharedPtr<CADKernel::FCurve>();
 	}
 
 	return AddCurveNurbsFromData(*NurbsData, UVReparameterization);
 }
 
-TSharedPtr<FTopologicalEdge> FTechSoftBridge::AddEdge(const A3DTopoCoEdge* A3DCoedge, const TSharedRef<FSurface>& Surface, const TechSoftUtils::FUVReparameterization& UVReparameterization, EOrientation& OutOrientation)
+TSharedPtr<CADKernel::FTopologicalEdge> FTechSoftBridge::AddEdge(const A3DTopoCoEdge* A3DCoedge, const TSharedRef<CADKernel::FSurface>& Surface, const TechSoftUtils::FUVReparameterization& UVReparameterization, CADKernel::EOrientation& OutOrientation)
 {
 	Report.EdgeCount++;
 
-	CADLibrary::TUniqueTSObj<A3DTopoCoEdgeData> CoEdgeData(A3DCoedge);
+	TUniqueTSObj<A3DTopoCoEdgeData> CoEdgeData(A3DCoedge);
 	if (!CoEdgeData.IsValid())
 	{
-		return TSharedPtr<FTopologicalEdge>();
+		return TSharedPtr<CADKernel::FTopologicalEdge>();
 	}
 
 	ensureCADKernel(CoEdgeData->m_pUVCurve);
 
-	TSharedPtr<FCurve> Curve = AddCurve(CoEdgeData->m_pUVCurve, UVReparameterization);
+	TSharedPtr<CADKernel::FCurve> Curve = AddCurve(CoEdgeData->m_pUVCurve, UVReparameterization);
 	if (!Curve.IsValid())
 	{
-		return TSharedPtr<FTopologicalEdge>();
+		return TSharedPtr<CADKernel::FTopologicalEdge>();
 	}
 
-	TSharedRef<FRestrictionCurve> RestrictionCurve = FEntity::MakeShared<FRestrictionCurve>(Surface, Curve.ToSharedRef());
-
-	TSharedPtr<FTopologicalEdge> Edge;
-	if (CoEdgeData->m_pEdge)
-	{
-		CADLibrary::TUniqueTSObj<A3DTopoEdgeData> A3DEdgeData(CoEdgeData->m_pEdge);
-		if (A3DEdgeData.IsValid())
-		{
-			double Tolerance = A3DEdgeData->m_dTolerance;
-			bool bHasTrimDomain = (bool)A3DEdgeData->m_bHasTrimDomain;
-
-			if (A3DEdgeData->m_p3dCurve)
-			{
-				//traverseCurve(A3DEdgeData.m_p3dCurve, edge);
-			}
-
-			if (A3DEdgeData->m_bHasTrimDomain)
-			{
-				FLinearBoundary Boundary = TechSoftUtils::GetLinearBoundary(A3DEdgeData->m_sInterval);
-				Edge = FTopologicalEdge::Make(RestrictionCurve, Boundary);
-			}
-		}
-	}
+	TSharedRef<CADKernel::FRestrictionCurve> RestrictionCurve = CADKernel::FEntity::MakeShared<CADKernel::FRestrictionCurve>(Surface, Curve.ToSharedRef());
 
-	if (!Edge.IsValid())
-	{
-		Edge = FTopologicalEdge::Make(RestrictionCurve);
-	}
+	TSharedPtr<CADKernel::FTopologicalEdge> Edge;
+	Edge = CADKernel::FTopologicalEdge::Make(RestrictionCurve);
 
 	if (!Edge.IsValid())
 	{
-		return TSharedPtr<FTopologicalEdge>();
+		return TSharedPtr<CADKernel::FTopologicalEdge>();
 	}
 
 	A3DEdgeToEdge.Emplace(A3DCoedge, Edge);
 
-	//CoEdgeData->m_ucOrientationUVWithLoop;
-	//CoEdgeData->m_ucOrientationWithLoop;
-	OutOrientation = CoEdgeData->m_ucOrientationUVWithLoop > 0 ? EOrientation::Front : EOrientation::Back;
+	OutOrientation = CoEdgeData->m_ucOrientationUVWithLoop > 0 ? CADKernel::EOrientation::Front : CADKernel::EOrientation::Back;
 
 	// Link edges
 	if (CoEdgeData->m_pNeighbor)
@@ -1210,13 +650,13 @@ TSharedPtr<FTopologicalEdge> FTechSoftBridge::AddEdge(const A3DTopoCoEdge* A3DCo
 		const A3DTopoCoEdge* Neighbor = CoEdgeData->m_pNeighbor;
 		while (Neighbor && Neighbor != A3DCoedge)
 		{
-			TSharedPtr<FTopologicalEdge>* TwinEdge = A3DEdgeToEdge.Find(Neighbor);
+			TSharedPtr<CADKernel::FTopologicalEdge>* TwinEdge = A3DEdgeToEdge.Find(Neighbor);
 			if (TwinEdge != nullptr)
 			{
 				Edge->Link(*TwinEdge->Get(), SquareJoiningVertexTolerance);
 			}
 
-			CADLibrary::TUniqueTSObj<A3DTopoCoEdgeData> NeighborData(Neighbor);
+			TUniqueTSObj<A3DTopoCoEdgeData> NeighborData(Neighbor);
 			if (NeighborData.IsValid())
 			{
 				Neighbor = NeighborData->m_pNeighbor;
@@ -1231,25 +671,25 @@ TSharedPtr<FTopologicalEdge> FTechSoftBridge::AddEdge(const A3DTopoCoEdge* A3DCo
 	return Edge;
 }
 
-TSharedPtr<FTopologicalLoop> FTechSoftBridge::AddLoop(const A3DTopoLoop* A3DLoop, const TSharedRef<FSurface>& Surface, const TechSoftUtils::FUVReparameterization& UVReparameterization)
+TSharedPtr<CADKernel::FTopologicalLoop> FTechSoftBridge::AddLoop(const A3DTopoLoop* A3DLoop, const TSharedRef<CADKernel::FSurface>& Surface, const TechSoftUtils::FUVReparameterization& UVReparameterization)
 {
 	Report.LoopCount++;
 
-	TArray<TSharedPtr<FTopologicalEdge>> Edges;
-	TArray<EOrientation> Directions;
+	TArray<TSharedPtr<CADKernel::FTopologicalEdge>> Edges;
+	TArray<CADKernel::EOrientation> Directions;
 
-	CADLibrary::TUniqueTSObj<A3DTopoLoopData> TopoLoopData(A3DLoop);
+	TUniqueTSObj<A3DTopoLoopData> TopoLoopData(A3DLoop);
 	if (!TopoLoopData.IsValid())
 	{
 		Report.DegeneratedLoopCount++;
-		return TSharedPtr<FTopologicalLoop>();
+		return TSharedPtr<CADKernel::FTopologicalLoop>();
 	}
 
 	bool bLoopOrientation = (bool)TopoLoopData->m_ucOrientationWithSurface;
 	for (A3DUns32 Index = 0; Index < TopoLoopData->m_uiCoEdgeSize; ++Index)
 	{
-		EOrientation Orientation;
-		TSharedPtr<FTopologicalEdge> Edge = AddEdge(TopoLoopData->m_ppCoEdges[Index], Surface, UVReparameterization, Orientation);
+		CADKernel::EOrientation Orientation;
+		TSharedPtr<CADKernel::FTopologicalEdge> Edge = AddEdge(TopoLoopData->m_ppCoEdges[Index], Surface, UVReparameterization, Orientation);
 		if (!Edge.IsValid())
 		{
 			Report.DegeneratedEdgeCount++;
@@ -1260,42 +700,23 @@ TSharedPtr<FTopologicalLoop> FTechSoftBridge::AddLoop(const A3DTopoLoop* A3DLoop
 		Directions.Emplace(Orientation);
 	}
 
-	//int32 CoedgeCount = CTCoedgeIds.Count();
-	//Edges.Reserve(CoedgeCount);
-	//Directions.Reserve(CoedgeCount);
-
-	//CTCoedgeIds.IteratorInitialize();
-
-	//CT_OBJECT_ID CoEdgeId;
-	//while ((CoEdgeId = CTCoedgeIds.IteratorIter()) != 0)
-	//{
-	//	TSharedPtr<FTopologicalEdge> Edge = AddEdge(CoEdgeId, Surface);
-	//	if (!Edge.IsValid())
-	//	{
-	//		continue;
-	//	}
-
-	//	Edges.Emplace(Edge);
-	//	Directions.Emplace(EOrientation::Front);
-	//}
-
 	if (Edges.Num() == 0)
 	{
 		Report.DegeneratedLoopCount++;
-		return TSharedPtr<FTopologicalLoop>();
+		return TSharedPtr<CADKernel::FTopologicalLoop>();
 	}
 
-	return FTopologicalLoop::Make(Edges, Directions, GeometricTolerance);
+	return CADKernel::FTopologicalLoop::Make(Edges, Directions, GeometricTolerance);
 }
 
-void FTechSoftBridge::AddFace(const A3DTopoFace* A3DFace, TSharedRef<FShell>& Shell)
+void FTechSoftBridge::AddFace(const A3DTopoFace* A3DFace, CADKernel::EOrientation Orientation, TSharedRef<CADKernel::FShell>& Shell)
 {
 	Report.FaceCount++;
 
-	FEntityData MetaData;
-	TechSoftUtils::TraverseSource(A3DFace, MetaData);
+	FEntityMetaData MetaData;
+	Parser.ExtractMetaData(A3DFace, MetaData);
 
-	CADLibrary::TUniqueTSObj<A3DTopoFaceData> TopoFaceData(A3DFace);
+	TUniqueTSObj<A3DTopoFaceData> TopoFaceData(A3DFace);
 	if (!TopoFaceData.IsValid())
 	{
 		Report.FailedFaceCount++;
@@ -1304,20 +725,20 @@ void FTechSoftBridge::AddFace(const A3DTopoFace* A3DFace, TSharedRef<FShell>& Sh
 
 	const A3DSurfBase* A3DSurface = TopoFaceData->m_pSurface;
 	TechSoftUtils::FUVReparameterization UVReparameterization = TechSoftUtils::FUVReparameterization::Identity;
-	TSharedPtr<FSurface> SurfacePtr = AddSurface(A3DSurface, UVReparameterization);
+	TSharedPtr<CADKernel::FSurface> SurfacePtr = AddSurface(A3DSurface, UVReparameterization);
 	if (!SurfacePtr.IsValid())
 	{
 		Report.DegeneratedSurfaceCount++;
 		Report.FailedFaceCount++;
 		return;
 	}
-	TSharedRef<FSurface> Surface = SurfacePtr.ToSharedRef();
+	TSharedRef<CADKernel::FSurface> Surface = SurfacePtr.ToSharedRef();
 
-	TSharedRef<FTopologicalFace> Face = FEntity::MakeShared<FTopologicalFace>(Surface);
+	TSharedRef<CADKernel::FTopologicalFace> Face = CADKernel::FEntity::MakeShared<CADKernel::FTopologicalFace>(Surface);
 
 	if (TopoFaceData->m_bHasTrimDomain)
 	{
-		FSurfacicBoundary SurfaceBoundary = TechSoftUtils::GetSurfacicBoundary(TopoFaceData->m_sSurfaceDomain, UnitScale);
+		CADKernel::FSurfacicBoundary SurfaceBoundary = TechSoftUtils::GetSurfacicBoundary(TopoFaceData->m_sSurfaceDomain, UnitScale);
 		Surface->TrimBoundaryTo(SurfaceBoundary);
 	}
 
@@ -1327,25 +748,25 @@ void FTechSoftBridge::AddFace(const A3DTopoFace* A3DFace, TSharedRef<FShell>& Sh
 	}
 	else
 	{
-		TArray<TSharedPtr<FTopologicalLoop>> Loops;
+		TArray<TSharedPtr<CADKernel::FTopologicalLoop>> Loops;
 
 		for (A3DUns32 Index = 0; Index < TopoFaceData->m_uiLoopSize; ++Index)
 		{
-			TSharedPtr<FTopologicalLoop> Loop = FTechSoftBridge::AddLoop(TopoFaceData->m_ppLoops[Index], Surface, UVReparameterization);
+			TSharedPtr<CADKernel::FTopologicalLoop> Loop = FTechSoftBridge::AddLoop(TopoFaceData->m_ppLoops[Index], Surface, UVReparameterization);
 			if (!Loop.IsValid())
 			{
 				continue;
 			}
 
-			TArray<FPoint2D> LoopSampling;
+			TArray<CADKernel::FPoint2D> LoopSampling;
 			Loop->Get2DSampling(LoopSampling);
-			FAABB2D Boundary;
+			CADKernel::FAABB2D Boundary;
 			Boundary += LoopSampling;
 			Loop->Boundary.Set(Boundary.GetMin(), Boundary.GetMax());
 
 			// Check if the loop is not composed with only degenerated edge
 			bool bDegeneratedLoop = true;
-			for (const FOrientedEdge& Edge : Loop->GetEdges())
+			for (const CADKernel::FOrientedEdge& Edge : Loop->GetEdges())
 			{
 				if (!Edge.Entity->IsDegenerated())
 				{
@@ -1363,10 +784,8 @@ void FTechSoftBridge::AddFace(const A3DTopoFace* A3DFace, TSharedRef<FShell>& Sh
 
 		if (Loops.Num() == 0)
 		{
-			FMessage::Printf(Log, TEXT("The Face has no loop, this face is bounded by its natural boundaries\n"));
-			//Face->ApplyNaturalLoops();
 			Report.FailedFaceCount++;
-			return; //?
+			return;
 		}
 		else
 		{
@@ -1375,23 +794,20 @@ void FTechSoftBridge::AddFace(const A3DTopoFace* A3DFace, TSharedRef<FShell>& Sh
 	}
 
 	AddMetadata(MetaData, *Face);
-	//Face->SetPatchId((int32)CTFaceId);
-
-	EOrientation Orientation = EOrientation::Front; //CTOrientation == CT_FORWARD ? EOrientation::Front : EOrientation::Back;
 	Shell->Add(Face, Orientation);
 }
 
 static bool bUseSurfaceAsNurbs = true;
 
-TSharedPtr<FSurface> FTechSoftBridge::AddSurface(const A3DSurfBase* A3DSurface, TechSoftUtils::FUVReparameterization& OutUVReparameterization)
+TSharedPtr<CADKernel::FSurface> FTechSoftBridge::AddSurface(const A3DSurfBase* A3DSurface, TechSoftUtils::FUVReparameterization& OutUVReparameterization)
 {
 	Report.SurfaceCount++;
 
-	FEntityData MetaData;
-	TechSoftUtils::TraverseSource(A3DSurface, MetaData);
+	FEntityMetaData MetaData;
+	Parser.ExtractMetaData(A3DSurface, MetaData);
 
 	A3DEEntityType Type;
-	int32 Ret = CADLibrary::TechSoftInterface::GetEntityType(A3DSurface, &Type);
+	int32 Ret = TechSoftInterface::GetEntityType(A3DSurface, &Type);
 	if (Ret == A3D_SUCCESS)
 	{
 		switch (Type)
@@ -1458,46 +874,46 @@ TSharedPtr<FSurface> FTechSoftBridge::AddSurface(const A3DSurfBase* A3DSurface,
 	{
 		return AddSurfaceAsNurbs(A3DSurface, OutUVReparameterization);
 	}
-	return TSharedPtr<FSurface>();
+	return TSharedPtr<CADKernel::FSurface>();
 }
 
-TSharedPtr<FSurface> FTechSoftBridge::AddConeSurface(const A3DSurfBase* Surface, TechSoftUtils::FUVReparameterization& OutUVReparameterization)
+TSharedPtr<CADKernel::FSurface> FTechSoftBridge::AddConeSurface(const A3DSurfBase* Surface, TechSoftUtils::FUVReparameterization& OutUVReparameterization)
 {
 	Report.ConeSurfaceCount++;
 
-	CADLibrary::TUniqueTSObj<A3DSurfConeData> A3DConeData(Surface);
+	TUniqueTSObj<A3DSurfConeData> A3DConeData(Surface);
 	if (!A3DConeData.IsValid())
 	{
-		return TSharedPtr<FSurface>();
+		return TSharedPtr<CADKernel::FSurface>();
 	}
 
 	OutUVReparameterization.AddUVTransform(A3DConeData->m_sParam);
-	FMatrixH CoordinateSystem = TechSoftUtils::CreateCoordinateSystem(A3DConeData->m_sTrsf, UnitScale);
-	FSurfacicBoundary Boundary = TechSoftUtils::GetSurfacicBoundary(A3DConeData->m_sParam.m_sUVDomain, UnitScale);
-	TSharedPtr<FSurface> Cone = FEntity::MakeShared<FConeSurface>(GeometricTolerance, CoordinateSystem, A3DConeData->m_dRadius * UnitScale, A3DConeData->m_dSemiAngle, Boundary);
+	CADKernel::FMatrixH CoordinateSystem = TechSoftUtils::CreateCoordinateSystem(A3DConeData->m_sTrsf, UnitScale);
+	CADKernel::FSurfacicBoundary Boundary = TechSoftUtils::GetSurfacicBoundary(A3DConeData->m_sParam.m_sUVDomain, UnitScale);
+	TSharedPtr<CADKernel::FSurface> Cone = CADKernel::FEntity::MakeShared<CADKernel::FConeSurface>(GeometricTolerance, CoordinateSystem, A3DConeData->m_dRadius * UnitScale, A3DConeData->m_dSemiAngle, Boundary);
 
 	return Cone;
 }
 
-TSharedPtr<FSurface> FTechSoftBridge::AddCylinderSurface(const A3DSurfBase* Surface, TechSoftUtils::FUVReparameterization& OutUVReparameterization)
+TSharedPtr<CADKernel::FSurface> FTechSoftBridge::AddCylinderSurface(const A3DSurfBase* Surface, TechSoftUtils::FUVReparameterization& OutUVReparameterization)
 {
 	Report.CylinderSurfaceCount++;
 
-	CADLibrary::TUniqueTSObj<A3DSurfCylinderData> A3DCylinderData(Surface);
+	TUniqueTSObj<A3DSurfCylinderData> A3DCylinderData(Surface);
 	if (!A3DCylinderData.IsValid())
 	{
-		return TSharedPtr<FSurface>();
+		return TSharedPtr<CADKernel::FSurface>();
 	}
 
 	OutUVReparameterization.AddUVTransform(A3DCylinderData->m_sParam);
-	FMatrixH CoordinateSystem = TechSoftUtils::CreateCoordinateSystem(A3DCylinderData->m_sTrsf, UnitScale);
-	FSurfacicBoundary Boundary = TechSoftUtils::GetSurfacicBoundary(A3DCylinderData->m_sParam.m_sUVDomain, UnitScale);
-	TSharedPtr<FSurface> Cylinder = FEntity::MakeShared<FCylinderSurface>(GeometricTolerance, CoordinateSystem, A3DCylinderData->m_dRadius * UnitScale, Boundary);
+	CADKernel::FMatrixH CoordinateSystem = TechSoftUtils::CreateCoordinateSystem(A3DCylinderData->m_sTrsf, UnitScale);
+	CADKernel::FSurfacicBoundary Boundary = TechSoftUtils::GetSurfacicBoundary(A3DCylinderData->m_sParam.m_sUVDomain, UnitScale);
+	TSharedPtr<CADKernel::FSurface> Cylinder = CADKernel::FEntity::MakeShared<CADKernel::FCylinderSurface>(GeometricTolerance, CoordinateSystem, A3DCylinderData->m_dRadius * UnitScale, Boundary);
 
 	return Cylinder;
 }
 
-TSharedPtr<FSurface> FTechSoftBridge::AddLinearTransfoSurface(const A3DSurfBase* Surface, TechSoftUtils::FUVReparameterization& OutUVReparameterization)
+TSharedPtr<CADKernel::FSurface> FTechSoftBridge::AddLinearTransfoSurface(const A3DSurfBase* Surface, TechSoftUtils::FUVReparameterization& OutUVReparameterization)
 {
 	Report.LinearTransfoSurfaceCount++;
 
@@ -1505,33 +921,24 @@ TSharedPtr<FSurface> FTechSoftBridge::AddLinearTransfoSurface(const A3DSurfBase*
 	{
 		return AddSurfaceAsNurbs(Surface, OutUVReparameterization);
 	}
-
-	//TUniqueTSObj<A3DSurfLinearTransfoData A3DLinearTransfoData;
-	//A3D_INITIALIZE_DATA(A3DSurfLinearTransfoData, A3DLinearTransfoData);
-	//A3DInt32 Ret = A3DSurfLinearTransfoGet(Surface, &A3DLinearTransfoData);
-	//if (Ret != A3D_SUCCESS)
-	//{
-	//	return TSharedPtr<FSurface>();
-	//}
-
-	//A3DSurfLinearTransfoGet(NULL, &A3DLinearTransfoData);
-	return TSharedPtr<FSurface>();
+	// Todo
+	return TSharedPtr<CADKernel::FSurface>();
 }
 
-TSharedPtr<FSurface> FTechSoftBridge::AddNurbsSurface(const A3DSurfNurbs* Nurbs, TechSoftUtils::FUVReparameterization& OutUVReparameterization)
+TSharedPtr<CADKernel::FSurface> FTechSoftBridge::AddNurbsSurface(const A3DSurfNurbs* Nurbs, TechSoftUtils::FUVReparameterization& OutUVReparameterization)
 {
 	Report.NurbsSurfaceCount++;
 
-	CADLibrary::TUniqueTSObj<A3DSurfNurbsData> A3DNurbsData(Nurbs);
+	TUniqueTSObj<A3DSurfNurbsData> A3DNurbsData(Nurbs);
 	if (!A3DNurbsData.IsValid())
 	{
-		return TSharedPtr<FSurface>();
+		return TSharedPtr<CADKernel::FSurface>();
 	}
 
 	return AddSurfaceNurbs(*A3DNurbsData, OutUVReparameterization);
 }
 
-TSharedPtr<FSurface> FTechSoftBridge::AddOffsetSurface(const A3DSurfBase* Surface, TechSoftUtils::FUVReparameterization& OutUVReparameterization)
+TSharedPtr<CADKernel::FSurface> FTechSoftBridge::AddOffsetSurface(const A3DSurfBase* Surface, TechSoftUtils::FUVReparameterization& OutUVReparameterization)
 {
 	Report.OffsetSurfaceCount++;
 
@@ -1540,28 +947,28 @@ TSharedPtr<FSurface> FTechSoftBridge::AddOffsetSurface(const A3DSurfBase* Surfac
 		return AddSurfaceAsNurbs(Surface, OutUVReparameterization);
 	}
 
-	return TSharedPtr<FSurface>();
+	return TSharedPtr<CADKernel::FSurface>();
 }
 
-TSharedPtr<FSurface> FTechSoftBridge::AddPlaneSurface(const A3DSurfPlane* Surface, TechSoftUtils::FUVReparameterization& OutUVReparameterization)
+TSharedPtr<CADKernel::FSurface> FTechSoftBridge::AddPlaneSurface(const A3DSurfPlane* Surface, TechSoftUtils::FUVReparameterization& OutUVReparameterization)
 {
 	Report.PlaneSurfaceCount++;
 
-	CADLibrary::TUniqueTSObj<A3DSurfPlaneData> A3DPlaneData(Surface);
+	TUniqueTSObj<A3DSurfPlaneData> A3DPlaneData(Surface);
 	if (!A3DPlaneData.IsValid())
 	{
-		return TSharedPtr<FSurface>();
+		return TSharedPtr<CADKernel::FSurface>();
 	}
 
 	OutUVReparameterization.AddUVTransform(A3DPlaneData->m_sParam);
-	FMatrixH CoordinateSystem = TechSoftUtils::CreateCoordinateSystem(A3DPlaneData->m_sTrsf, UnitScale);
-	FSurfacicBoundary Boundary = TechSoftUtils::GetSurfacicBoundary(A3DPlaneData->m_sParam.m_sUVDomain, UnitScale);
-	TSharedPtr<FSurface> Plane = FEntity::MakeShared<FPlaneSurface>(GeometricTolerance, CoordinateSystem, Boundary);
+	CADKernel::FMatrixH CoordinateSystem = TechSoftUtils::CreateCoordinateSystem(A3DPlaneData->m_sTrsf, UnitScale);
+	CADKernel::FSurfacicBoundary Boundary = TechSoftUtils::GetSurfacicBoundary(A3DPlaneData->m_sParam.m_sUVDomain, UnitScale);
+	TSharedPtr<CADKernel::FSurface> Plane = CADKernel::FEntity::MakeShared<CADKernel::FPlaneSurface>(GeometricTolerance, CoordinateSystem, Boundary);
 
 	return Plane;
 }
 
-TSharedPtr<FSurface> FTechSoftBridge::AddRevolutionSurface(const A3DSurfRevolution* Surface, TechSoftUtils::FUVReparameterization& OutUVReparameterization)
+TSharedPtr<CADKernel::FSurface> FTechSoftBridge::AddRevolutionSurface(const A3DSurfRevolution* Surface, TechSoftUtils::FUVReparameterization& OutUVReparameterization)
 {
 	Report.RevolutionSurfaceCount++;
 
@@ -1570,16 +977,16 @@ TSharedPtr<FSurface> FTechSoftBridge::AddRevolutionSurface(const A3DSurfRevoluti
 		return AddSurfaceAsNurbs(Surface, OutUVReparameterization);
 	}
 
-	CADLibrary::TUniqueTSObj<A3DSurfRevolutionData> A3DRevolutionData(Surface);
+	TUniqueTSObj<A3DSurfRevolutionData> A3DRevolutionData(Surface);
 	if (!A3DRevolutionData.IsValid())
 	{
-		return TSharedPtr<FSurface>();
+		return TSharedPtr<CADKernel::FSurface>();
 	}
 
-	return TSharedPtr<FSurface>();
+	return TSharedPtr<CADKernel::FSurface>();
 }
 
-TSharedPtr<FSurface> FTechSoftBridge::AddRuledSurface(const A3DSurfBase* Surface, TechSoftUtils::FUVReparameterization& OutUVReparameterization)
+TSharedPtr<CADKernel::FSurface> FTechSoftBridge::AddRuledSurface(const A3DSurfBase* Surface, TechSoftUtils::FUVReparameterization& OutUVReparameterization)
 {
 	Report.RuledSurfaceCount++;
 
@@ -1588,16 +995,16 @@ TSharedPtr<FSurface> FTechSoftBridge::AddRuledSurface(const A3DSurfBase* Surface
 		return AddSurfaceAsNurbs(Surface, OutUVReparameterization);
 	}
 
-	CADLibrary::TUniqueTSObj<A3DSurfRuledData> A3DRuledData(Surface);
+	TUniqueTSObj<A3DSurfRuledData> A3DRuledData(Surface);
 	if (!A3DRuledData.IsValid())
 	{
-		return TSharedPtr<FSurface>();
+		return TSharedPtr<CADKernel::FSurface>();
 	}
 
-	return TSharedPtr<FSurface>();
+	return TSharedPtr<CADKernel::FSurface>();
 }
 
-TSharedPtr<FSurface> FTechSoftBridge::AddSphereSurface(const A3DSurfBase* Surface, TechSoftUtils::FUVReparameterization& OutUVReparameterization)
+TSharedPtr<CADKernel::FSurface> FTechSoftBridge::AddSphereSurface(const A3DSurfBase* Surface, TechSoftUtils::FUVReparameterization& OutUVReparameterization)
 {
 	Report.SphereSurfaceCount++;
 
@@ -1606,39 +1013,39 @@ TSharedPtr<FSurface> FTechSoftBridge::AddSphereSurface(const A3DSurfBase* Surfac
 		return AddSurfaceAsNurbs(Surface, OutUVReparameterization);
 	}
 
-	CADLibrary::TUniqueTSObj<A3DSurfSphereData> A3DSphereData(Surface);
+	TUniqueTSObj<A3DSurfSphereData> A3DSphereData(Surface);
 	if (!A3DSphereData.IsValid())
 	{
-		return TSharedPtr<FSurface>();
+		return TSharedPtr<CADKernel::FSurface>();
 	}
 
 	OutUVReparameterization.AddUVTransform(A3DSphereData->m_sParam);
-	FMatrixH CoordinateSystem = TechSoftUtils::CreateCoordinateSystem(A3DSphereData->m_sTrsf, UnitScale);
-	FSurfacicBoundary Boundary = TechSoftUtils::GetSurfacicBoundary(A3DSphereData->m_sParam.m_sUVDomain, UnitScale);
-	TSharedRef<FSphericalSurface> Sphere = FEntity::MakeShared<FSphericalSurface>(GeometricTolerance, CoordinateSystem, A3DSphereData->m_dRadius * UnitScale, Boundary);
+	CADKernel::FMatrixH CoordinateSystem = TechSoftUtils::CreateCoordinateSystem(A3DSphereData->m_sTrsf, UnitScale);
+	CADKernel::FSurfacicBoundary Boundary = TechSoftUtils::GetSurfacicBoundary(A3DSphereData->m_sParam.m_sUVDomain, UnitScale);
+	TSharedRef<CADKernel::FSphericalSurface> Sphere = CADKernel::FEntity::MakeShared<CADKernel::FSphericalSurface>(GeometricTolerance, CoordinateSystem, A3DSphereData->m_dRadius * UnitScale, Boundary);
 
 	return Sphere;
 }
 
-TSharedPtr<FSurface> FTechSoftBridge::AddTorusSurface(const A3DSurfBase* Surface, TechSoftUtils::FUVReparameterization& OutUVReparameterization)
+TSharedPtr<CADKernel::FSurface> FTechSoftBridge::AddTorusSurface(const A3DSurfBase* Surface, TechSoftUtils::FUVReparameterization& OutUVReparameterization)
 {
 	Report.TorusSurfaceCount++;
 
-	CADLibrary::TUniqueTSObj<A3DSurfTorusData> A3DTorusData(Surface);
+	TUniqueTSObj<A3DSurfTorusData> A3DTorusData(Surface);
 	if (!A3DTorusData.IsValid())
 	{
-		return TSharedPtr<FSurface>();
+		return TSharedPtr<CADKernel::FSurface>();
 	}
 
 	OutUVReparameterization.AddUVTransform(A3DTorusData->m_sParam);
-	FMatrixH CoordinateSystem = TechSoftUtils::CreateCoordinateSystem(A3DTorusData->m_sTrsf, UnitScale);
-	FSurfacicBoundary Boundary = TechSoftUtils::GetSurfacicBoundary(A3DTorusData->m_sParam.m_sUVDomain, UnitScale);
-	TSharedPtr<FSurface> Torus = FEntity::MakeShared<FTorusSurface>(GeometricTolerance, CoordinateSystem, A3DTorusData->m_dMajorRadius * UnitScale, A3DTorusData->m_dMinorRadius * UnitScale, Boundary);
+	CADKernel::FMatrixH CoordinateSystem = TechSoftUtils::CreateCoordinateSystem(A3DTorusData->m_sTrsf, UnitScale);
+	CADKernel::FSurfacicBoundary Boundary = TechSoftUtils::GetSurfacicBoundary(A3DTorusData->m_sParam.m_sUVDomain, UnitScale);
+	TSharedPtr<CADKernel::FSurface> Torus = CADKernel::FEntity::MakeShared<CADKernel::FTorusSurface>(GeometricTolerance, CoordinateSystem, A3DTorusData->m_dMajorRadius * UnitScale, A3DTorusData->m_dMinorRadius * UnitScale, Boundary);
 
 	return Torus;
 }
 
-TSharedPtr<FSurface> FTechSoftBridge::AddBlend01Surface(const A3DSurfBase* Surface, TechSoftUtils::FUVReparameterization& OutUVReparameterization)
+TSharedPtr<CADKernel::FSurface> FTechSoftBridge::AddBlend01Surface(const A3DSurfBase* Surface, TechSoftUtils::FUVReparameterization& OutUVReparameterization)
 {
 	Report.Blend01SurfaceCount++;
 
@@ -1647,16 +1054,16 @@ TSharedPtr<FSurface> FTechSoftBridge::AddBlend01Surface(const A3DSurfBase* Surfa
 		return AddSurfaceAsNurbs(Surface, OutUVReparameterization);
 	}
 
-	CADLibrary::TUniqueTSObj<A3DSurfBlend01Data> A3DBlend01Data(Surface);
+	TUniqueTSObj<A3DSurfBlend01Data> A3DBlend01Data(Surface);
 	if (!A3DBlend01Data.IsValid())
 	{
-		return TSharedPtr<FSurface>();
+		return TSharedPtr<CADKernel::FSurface>();
 	}
 
-	return TSharedPtr<FSurface>();
+	return TSharedPtr<CADKernel::FSurface>();
 }
 
-TSharedPtr<FSurface> FTechSoftBridge::AddBlend02Surface(const A3DSurfBase* Surface, TechSoftUtils::FUVReparameterization& OutUVReparameterization)
+TSharedPtr<CADKernel::FSurface> FTechSoftBridge::AddBlend02Surface(const A3DSurfBase* Surface, TechSoftUtils::FUVReparameterization& OutUVReparameterization)
 {
 	Report.Blend02SurfaceCount++;
 
@@ -1665,16 +1072,16 @@ TSharedPtr<FSurface> FTechSoftBridge::AddBlend02Surface(const A3DSurfBase* Surfa
 		return AddSurfaceAsNurbs(Surface, OutUVReparameterization);
 	}
 
-	CADLibrary::TUniqueTSObj<A3DSurfBlend02Data> A3DBlend02Data(Surface);
+	TUniqueTSObj<A3DSurfBlend02Data> A3DBlend02Data(Surface);
 	if (!A3DBlend02Data.IsValid())
 	{
-		return TSharedPtr<FSurface>();
+		return TSharedPtr<CADKernel::FSurface>();
 	}
 
-	return TSharedPtr<FSurface>();
+	return TSharedPtr<CADKernel::FSurface>();
 }
 
-TSharedPtr<FSurface> FTechSoftBridge::AddBlend03Surface(const A3DSurfBase* Surface, TechSoftUtils::FUVReparameterization& OutUVReparameterization)
+TSharedPtr<CADKernel::FSurface> FTechSoftBridge::AddBlend03Surface(const A3DSurfBase* Surface, TechSoftUtils::FUVReparameterization& OutUVReparameterization)
 {
 	Report.Blend03SurfaceCount++;
 
@@ -1683,100 +1090,10 @@ TSharedPtr<FSurface> FTechSoftBridge::AddBlend03Surface(const A3DSurfBase* Surfa
 		return AddSurfaceAsNurbs(Surface, OutUVReparameterization);
 	}
 
-#ifdef WIP
-	CADLibrary::TUniqueTSObj<A3DSurfBlend03Data> A3DBlend03Data(Surface);
-	if (!A3DBlend03Data.IsValid())
-	{
-		return TSharedPtr<FSurface>();
-	}
-
-	FBlend03Data BlendData;
-
-	F3DDebugSession _(TEXT("Blend03Surface"));
-	{
-
-		BlendData.CenterCurve.Reserve(A3DBlend03Data->m_uiNbOfElement);
-		BlendData.RailCurve1.Reserve(A3DBlend03Data->m_uiNbOfElement);
-		BlendData.RailCurve2.Reserve(A3DBlend03Data->m_uiNbOfElement);
-
-		BlendData.Parameters.Reserve(A3DBlend03Data->m_uiNbOfElement);
-		BlendData.Multiplicities.Reserve(A3DBlend03Data->m_uiNbOfElement);
-		BlendData.RailToAnglesV.Reserve(A3DBlend03Data->m_uiNbOfElement);
-		BlendData.Tangents.Reserve(A3DBlend03Data->m_uiNbOfElement);
-		BlendData.RailToDerivativesV.Reserve(A3DBlend03Data->m_uiNbOfElement);
-		BlendData.SecondDerivatives.Reserve(A3DBlend03Data->m_uiNbOfElement * 3);
-		BlendData.RailToSecondDerivatives.Reserve(A3DBlend03Data->m_uiNbOfElement);
-
-		BlendData.RailToParameterV = A3DBlend03Data->m_dRail2ParameterV;
-		BlendData.TrimVMin= A3DBlend03Data->m_dTrimVMin;
-		BlendData.TrimVMax= A3DBlend03Data->m_dTrimVMax;
-		BlendData.InitialMaxV = A3DBlend03Data->m_dInitialMaxV;
-
-		TArray<FPoint> Positions;
-		A3DBlend03Data->m_pPositions;
-		TechSoftUtils::FillPointArray(A3DBlend03Data->m_uiNbOfElement*3, A3DBlend03Data->m_pPositions, Positions, UnitScale);
-
-		for (uint32 Index = 0; Index < (3 * A3DBlend03Data->m_uiNbOfElement);)
-		{
-			BlendData.RailCurve1.Add(Positions[Index++]);
-			BlendData.RailCurve2.Add(Positions[Index++]);
-			BlendData.CenterCurve.Add(Positions[Index++]);
-		}
-
-		TechSoftUtils::FillPointArray(A3DBlend03Data->m_uiNbOfElement * 3, A3DBlend03Data->m_pTangent, Positions, 1);
-		TechSoftUtils::FillPointArray(A3DBlend03Data->m_uiNbOfElement * 3, A3DBlend03Data->m_pSecondDerivatives, Positions, 1);
-
-		TechSoftUtils::FillDoubleArray(A3DBlend03Data->m_uiNbOfElement, A3DBlend03Data->m_pdParameters, BlendData.Parameters);
-		TechSoftUtils::FillInt32Array(A3DBlend03Data->m_uiNbOfElement, A3DBlend03Data->m_piMultiplicities, BlendData.Multiplicities);
-		TechSoftUtils::FillDoubleArray(A3DBlend03Data->m_uiNbOfElement, A3DBlend03Data->m_pdRail2AnglesV, BlendData.RailToAnglesV);
-		TechSoftUtils::FillDoubleArray(A3DBlend03Data->m_uiNbOfElement, A3DBlend03Data->m_pdRail2DerivativesV, BlendData.RailToDerivativesV);
-		TechSoftUtils::FillDoubleArray(A3DBlend03Data->m_uiNbOfElement, A3DBlend03Data->m_pdRail2SecondDerivatives, BlendData.RailToSecondDerivatives);
-
-		DisplayPoints(TEXT("Center"), BlendData.CenterCurve);
-		DisplayPoints(TEXT("RailCurve1"), BlendData.RailCurve1, EVisuProperty::YellowPoint);
-		DisplayPoints(TEXT("RailCurve2"), BlendData.RailCurve2, EVisuProperty::GreenPoint);
-
-		TechSoftUtils::FillPointArray(A3DBlend03Data->m_uiNbOfElement * 3, A3DBlend03Data->m_pPositions, Positions, UnitScale);
-
-		FNurbsCurveData NurbsCurve;
-		NurbsCurve.Dimension = 3;
-		NurbsCurve.Degree = 5;
-		for (uint32 Index = 0; Index < A3DBlend03Data->m_uiNbOfElement; ++Index)
-		{
-			for (int32 Mndex = 0; Mndex < BlendData.Multiplicities[Index]; ++Mndex)
-			{
-				NurbsCurve.NodalVector.Add(BlendData.Parameters[Index]);
-			}
-		}
-		NurbsCurve.bIsRational = false;
-		NurbsCurve.Poles = BlendData.CenterCurve;
-
-		TSharedPtr<FNURBSCurve> CenterNurbs = FEntity::MakeShared<FNURBSCurve>(NurbsCurve);
-		{
-			F3DDebugSession _(TEXT("Center"));
-			Display(*CenterNurbs);
-		}
-
-		NurbsCurve.Poles = BlendData.RailCurve1;
-		TSharedPtr<FNURBSCurve> Rail1Nurbs = FEntity::MakeShared<FNURBSCurve>(NurbsCurve);
-		{
-			F3DDebugSession _(TEXT("Rail1Nurbs"));
-			Display(*Rail1Nurbs);
-		}
-
-		NurbsCurve.Poles = BlendData.RailCurve2;
-		TSharedPtr<FNURBSCurve> Rail2Nurbs = FEntity::MakeShared<FNURBSCurve>(NurbsCurve);
-		{
-			F3DDebugSession _(TEXT("Rail2Nurbs"));
-			Display(*Rail2Nurbs);
-		}
-
-	}
-#endif
-	return TSharedPtr<FSurface>();
+	return TSharedPtr<CADKernel::FSurface>();
 }
 
-TSharedPtr<FSurface> FTechSoftBridge::AddCylindricalSurface(const A3DSurfBase* Surface, TechSoftUtils::FUVReparameterization& OutUVReparameterization)
+TSharedPtr<CADKernel::FSurface> FTechSoftBridge::AddCylindricalSurface(const A3DSurfBase* Surface, TechSoftUtils::FUVReparameterization& OutUVReparameterization)
 {
 	Report.CylindricalSurfaceCount++;
 
@@ -1785,16 +1102,16 @@ TSharedPtr<FSurface> FTechSoftBridge::AddCylindricalSurface(const A3DSurfBase* S
 		return AddSurfaceAsNurbs(Surface, OutUVReparameterization);
 	}
 
-	CADLibrary::TUniqueTSObj<A3DSurfCylindricalData> A3DCylindricalData(Surface);
+	TUniqueTSObj<A3DSurfCylindricalData> A3DCylindricalData(Surface);
 	if (!A3DCylindricalData.IsValid())
 	{
-		return TSharedPtr<FSurface>();
+		return TSharedPtr<CADKernel::FSurface>();
 	}
 
-	return TSharedPtr<FSurface>();
+	return TSharedPtr<CADKernel::FSurface>();
 }
 
-TSharedPtr<FSurface> FTechSoftBridge::AddPipeSurface(const A3DSurfBase* Surface, TechSoftUtils::FUVReparameterization& OutUVReparameterization)
+TSharedPtr<CADKernel::FSurface> FTechSoftBridge::AddPipeSurface(const A3DSurfBase* Surface, TechSoftUtils::FUVReparameterization& OutUVReparameterization)
 {
 	Report.PipeSurfaceCount++;
 
@@ -1803,16 +1120,16 @@ TSharedPtr<FSurface> FTechSoftBridge::AddPipeSurface(const A3DSurfBase* Surface,
 		return AddSurfaceAsNurbs(Surface, OutUVReparameterization);
 	}
 
-	CADLibrary::TUniqueTSObj<A3DSurfPipeData> A3DPipeData(Surface);
+	TUniqueTSObj<A3DSurfPipeData> A3DPipeData(Surface);
 	if (!A3DPipeData.IsValid())
 	{
-		return TSharedPtr<FSurface>();
+		return TSharedPtr<CADKernel::FSurface>();
 	}
 
-	return TSharedPtr<FSurface>();
+	return TSharedPtr<CADKernel::FSurface>();
 }
 
-TSharedPtr<FSurface> FTechSoftBridge::AddExtrusionSurface(const A3DSurfBase* Surface, TechSoftUtils::FUVReparameterization& OutUVReparameterization)
+TSharedPtr<CADKernel::FSurface> FTechSoftBridge::AddExtrusionSurface(const A3DSurfBase* Surface, TechSoftUtils::FUVReparameterization& OutUVReparameterization)
 {
 	Report.ExtrusionSurfaceCount++;
 
@@ -1821,16 +1138,16 @@ TSharedPtr<FSurface> FTechSoftBridge::AddExtrusionSurface(const A3DSurfBase* Sur
 		return AddSurfaceAsNurbs(Surface, OutUVReparameterization);
 	}
 
-	CADLibrary::TUniqueTSObj<A3DSurfExtrusionData> A3DExtrusionData(Surface);
+	TUniqueTSObj<A3DSurfExtrusionData> A3DExtrusionData(Surface);
 	if (!A3DExtrusionData.IsValid())
 	{
-		return TSharedPtr<FSurface>();
+		return TSharedPtr<CADKernel::FSurface>();
 	}
 
-	return TSharedPtr<FSurface>();
+	return TSharedPtr<CADKernel::FSurface>();
 }
 
-TSharedPtr<FSurface> FTechSoftBridge::AddSurfaceFromCurves(const A3DSurfBase* Surface, TechSoftUtils::FUVReparameterization& OutUVReparameterization)
+TSharedPtr<CADKernel::FSurface> FTechSoftBridge::AddSurfaceFromCurves(const A3DSurfBase* Surface, TechSoftUtils::FUVReparameterization& OutUVReparameterization)
 {
 	Report.SurfaceFromCurvesCount++;
 
@@ -1839,16 +1156,16 @@ TSharedPtr<FSurface> FTechSoftBridge::AddSurfaceFromCurves(const A3DSurfBase* Su
 		return AddSurfaceAsNurbs(Surface, OutUVReparameterization);
 	}
 
-	CADLibrary::TUniqueTSObj<A3DSurfFromCurvesData> A3DFromCurvesData(Surface);
+	TUniqueTSObj<A3DSurfFromCurvesData> A3DFromCurvesData(Surface);
 	if (!A3DFromCurvesData.IsValid())
 	{
-		return TSharedPtr<FSurface>();
+		return TSharedPtr<CADKernel::FSurface>();
 	}
 
-	return TSharedPtr<FSurface>();
+	return TSharedPtr<CADKernel::FSurface>();
 }
 
-TSharedPtr<FSurface> FTechSoftBridge::AddTransformSurface(const A3DSurfBase* Surface, TechSoftUtils::FUVReparameterization& OutUVReparameterization)
+TSharedPtr<CADKernel::FSurface> FTechSoftBridge::AddTransformSurface(const A3DSurfBase* Surface, TechSoftUtils::FUVReparameterization& OutUVReparameterization)
 {
 	Report.TransformSurfaceCount++;
 
@@ -1857,37 +1174,18 @@ TSharedPtr<FSurface> FTechSoftBridge::AddTransformSurface(const A3DSurfBase* Sur
 		return AddSurfaceAsNurbs(Surface, OutUVReparameterization);
 	}
 
-	CADLibrary::TUniqueTSObj<A3DSurfFromCurvesData> A3DTransformData(Surface);
+	TUniqueTSObj<A3DSurfFromCurvesData> A3DTransformData(Surface);
 	if (!A3DTransformData.IsValid())
 	{
-		return TSharedPtr<FSurface>();
+		return TSharedPtr<CADKernel::FSurface>();
 	}
 
-	return TSharedPtr<FSurface>();
+	return TSharedPtr<CADKernel::FSurface>();
 }
 
-TSharedPtr<FSurface> FTechSoftBridge::AddSurfaceNurbs(const A3DSurfNurbsData& A3DNurbsData, TechSoftUtils::FUVReparameterization& OutUVReparameterization)
+TSharedPtr<CADKernel::FSurface> FTechSoftBridge::AddSurfaceNurbs(const A3DSurfNurbsData& A3DNurbsData, TechSoftUtils::FUVReparameterization& OutUVReparameterization)
 {
-	FNurbsSurfaceData NurbsData;
-
-	A3DNurbsData.m_eKnotType;
-	//kA3DKnotTypeUniformKnots,			/*!< Uniform. */
-	//kA3DKnotTypeUnspecified,			/*!< No particularity. */
-	//kA3DKnotTypeQuasiUniformKnots,	/*!< Quasi-uniform. */
-	//kA3DKnotTypePieceWiseBezierKnots	/*!< Extrema with multiplicities of degree + 1, internal is degree. */
-
-	A3DNurbsData.m_eSurfaceForm;
-	//kA3DBSplineSurfaceFormPlane,				/*!< Planar surface. */
-	//kA3DBSplineSurfaceFormCylindrical,		/*!< Cylindrical surface. */
-	//kA3DBSplineSurfaceFormConical,			/*!< Conical surface. */
-	//kA3DBSplineSurfaceFormSpherical,			/*!< Spherical surface. */
-	//kA3DBSplineSurfaceFormRevolution,			/*!< Surface of revolution. */
-	//kA3DBSplineSurfaceFormRuled,				/*!< Ruled surface. */
-	//kA3DBSplineSurfaceFormGeneralizedCone,	/*!< Cone. */
-	//kA3DBSplineSurfaceFormQuadric,			/*!< Quadric surface. */
-	//kA3DBSplineSurfaceFormLinearExtrusion,	/*!< Surface of extrusion. */
-	//kA3DBSplineSurfaceFormUnspecified,		/*!< Unspecified. */
-	//kA3DBSplineSurfaceFormPolynomial			/*!< Polynomial surface. */
+	CADKernel::FNurbsSurfaceData NurbsData;
 
 	NurbsData.PoleUCount = A3DNurbsData.m_uiUCtrlSize;
 	NurbsData.PoleVCount = A3DNurbsData.m_uiVCtrlSize;
@@ -1899,11 +1197,11 @@ TSharedPtr<FSurface> FTechSoftBridge::AddSurfaceNurbs(const A3DSurfNurbsData& A3
 	TechSoftUtils::FillDoubleArray(A3DNurbsData.m_uiUKnotSize, A3DNurbsData.m_pdUKnots, NurbsData.UNodalVector);
 	TechSoftUtils::FillDoubleArray(A3DNurbsData.m_uiVKnotSize, A3DNurbsData.m_pdVKnots, NurbsData.VNodalVector);
 
-	TArray<FPoint> Poles;
+	TArray<CADKernel::FPoint> Poles;
 	TechSoftUtils::FillPointArray(NurbsData.PoleUCount, NurbsData.PoleVCount, A3DNurbsData.m_pCtrlPts, NurbsData.Poles);
 	if (!FMath::IsNearlyEqual(UnitScale, 1.))
 	{
-		for (FPoint& Point : NurbsData.Poles)
+		for (CADKernel::FPoint& Point : NurbsData.Poles)
 		{
 			Point *= UnitScale;
 		}
@@ -1916,14 +1214,14 @@ TSharedPtr<FSurface> FTechSoftBridge::AddSurfaceNurbs(const A3DSurfNurbsData& A3
 		TechSoftUtils::FillDoubleArray(NurbsData.PoleUCount, NurbsData.PoleVCount, A3DNurbsData.m_pdWeights, NurbsData.Weights);
 	}
 
-	return FEntity::MakeShared<FNURBSSurface>(GeometricTolerance, NurbsData);
+	return CADKernel::FEntity::MakeShared<CADKernel::FNURBSSurface>(GeometricTolerance, NurbsData);
 }
 
-TSharedPtr<FSurface> FTechSoftBridge::AddSurfaceAsNurbs(const A3DSurfBase* Surface, TechSoftUtils::FUVReparameterization& OutUVReparameterization)
+TSharedPtr<CADKernel::FSurface> FTechSoftBridge::AddSurfaceAsNurbs(const A3DSurfBase* Surface, TechSoftUtils::FUVReparameterization& OutUVReparameterization)
 {
 	Report.SurfaceAsNurbsCount++;
 
-	CADLibrary::TUniqueTSObj<A3DSurfNurbsData> A3DNurbsData;
+	TUniqueTSObj<A3DSurfNurbsData> A3DNurbsData;
 
 	A3DDouble Tolerance = 1e-3;
 	A3DBool bUseSameParameterization = true;
@@ -1931,21 +1229,22 @@ TSharedPtr<FSurface> FTechSoftBridge::AddSurfaceAsNurbs(const A3DSurfBase* Surfa
 
 	if (!A3DNurbsData.IsValid())
 	{
-		return TSharedPtr<FSurface>();
+		return TSharedPtr<CADKernel::FSurface>();
 	}
 
 	return AddSurfaceNurbs(*A3DNurbsData, OutUVReparameterization);
 
 }
 
-void FTechSoftBridge::AddMetadata(FEntityData& MetaData, FTopologicalShapeEntity& Entity)
+void FTechSoftBridge::AddMetadata(FEntityMetaData& MetaData, CADKernel::FTopologicalShapeEntity& Entity)
 {
-	//Entity->SetHostId(CTNodeId);
-	Entity.SetName(*MetaData.Name);
-	//Entity.SetLayer(LayerId);
-	//uint32 ColorHId = CADLibrary::BuildColorId(ColorId, Alpha);
-	//Entity.SetColorId(ColorHId);
-	//Entity.SetPatchId();
+	FString* Name = MetaData.MetaData.Find(TEXT("Name"));
+	if(Name != nullptr)
+	{
+		Entity.SetName(*Name);
+	}
+	Entity.SetColorId(MetaData.ColorName);
+	Entity.SetMaterialId(MetaData.MaterialName);
 }
 
 }
diff --git a/Engine/Plugins/Enterprise/DatasmithCADImporter/Source/CADInterfaces/Private/TechSoftBridge.h b/Engine/Plugins/Enterprise/DatasmithCADImporter/Source/CADInterfaces/Private/TechSoftBridge.h
--- a/Engine/Plugins/Enterprise/DatasmithCADImporter/Source/CADInterfaces/Private/TechSoftBridge.h
+++ b/Engine/Plugins/Enterprise/DatasmithCADImporter/Source/CADInterfaces/Private/TechSoftBridge.h
@@ -12,117 +12,35 @@
 
 #include "TechSoftInterface.h"
 
-typedef void A3DAsmProductOccurrence;
-typedef void A3DCrvBase;
-typedef void A3DCrvBase;
-typedef void A3DCrvBlend02Boundary;
-typedef void A3DCrvCircle;
-typedef void A3DCrvComposite;
-typedef void A3DCrvEllipse;
-typedef void A3DCrvEquation;
-typedef void A3DCrvHelix;
-typedef void A3DCrvHyperbola;
-typedef void A3DCrvIntersection;
-typedef void A3DCrvLine;
-typedef void A3DCrvNurbs;
-typedef void A3DCrvOffset;
-typedef void A3DCrvOnSurf;
-typedef void A3DCrvParabola;
-typedef void A3DCrvPolyLine;
-typedef void A3DCrvTransform;
-typedef void A3DEntity;
-typedef void A3DGraphics;
-typedef void A3DMiscCartesianTransformation;
-typedef void A3DMiscCartesianTransformation;
-typedef void A3DMiscGeneralTransformation;
-typedef void A3DRiBrepModel;
-typedef void A3DRiCoordinateSystem;
-typedef void A3DRiRepresentationItem;
-typedef void A3DSurfBase;
-typedef void A3DSurfBase;
-typedef void A3DSurfBase;
-typedef void A3DSurfBase;
-typedef void A3DSurfBlend01;
-typedef void A3DSurfBlend02;
-typedef void A3DSurfBlend03;
-typedef void A3DSurfCone;
-typedef void A3DSurfCone;
-typedef void A3DSurfCylinder;
-typedef void A3DSurfCylinder;
-typedef void A3DSurfCylindrical;
-typedef void A3DSurfExtrusion;
-typedef void A3DSurfFromCurves;
-typedef void A3DSurfNurbs;
-typedef void A3DSurfOffset;
-typedef void A3DSurfPipe;
-typedef void A3DSurfPlane;
-typedef void A3DSurfRevolution;
-typedef void A3DSurfRuled;
-typedef void A3DSurfSphere;
-typedef void A3DSurfTorus;
-typedef void A3DSurfTransform;
-typedef void A3DTopoBrepData;
-typedef void A3DTopoCoEdge;
-typedef void A3DTopoConnex;
-typedef void A3DTopoFace;
-typedef void A3DTopoLoop;
-typedef void A3DTopoShell;
-
 namespace CADKernel
 {
 
-enum class EEntityType
-{
-	Instance,
-	Component,
-	Part,
-	Assembly,
-	UnloadedComponent,
-	Body,
-	BRepModel,
-	Undefined
-};
+class FBody;
+class FCriterion;
+class FCurve;
+class FEntity;
+class FFaceMesh;
+class FModel;
+class FPoint;
+class FRestrictionCurve;
+class FSession;
+class FShell;
+class FSurface;
+class FSurfacicBoundary;
+class FTopologicalEdge;
+class FTopologicalLoop;
+class FTopologicalShapeEntity;
 
-struct FEntityData
-{
-	unsigned int Uuid = 0;
-	EEntityType EntityType;
-	FString Name;
-	TMap<FString, FString> MetaData;
-	TArray<A3DAsmProductOccurrence*> Childs;
-};
+}
 
-struct FBodyGeometricData
+namespace CADLibrary
 {
-	unsigned int Uuid = 0;
-	FString Name;
-	//int FaceCount = 0;
-	//int MeshTriangleCount = 0;
-	//double MeshArea = 0.;
-	//FBBox MeshBBox;
-	//FBBox CADBBox;
-	//int MaxFaceTriangleCount = 0;
-	//int MaxCtrlHullSize = 0;
-	//bool bTestBBox = true;
-	//bool bTestArea = true;
-	//bool bTestVertexCount = true;
-	//bool bMeshIsExtract = false;
-
-	A3DRiRepresentationItem* RepresentationItem = nullptr;
-};
+
+class FTechSoftFileParser;
+struct FEntityMetaData;
 
 namespace TechSoftUtils
 {
-int32 TraverseSource(const A3DEntity* Entity, FEntityData& OutEntityData, bool bIsOccurrence = false, const ECADType FileType = ECADType::OTHER);
-int32 TraverseSpecificMetaData(A3DEModellerType ModellerType, const A3DAsmProductOccurrence* Occurrence, FEntityData& MetaData);
-
-int32 TraverseGraphics(const A3DGraphics* Graphics);
-FMatrixH TraverseCoordinateSystem(const A3DRiCoordinateSystem* CoordinateSystem, double UnitScale);
-FMatrixH TraverseTransformation(const A3DMiscCartesianTransformation* Transformation3d, double UnitScale);
-FMatrixH TraverseGeneralTransformation(const A3DMiscGeneralTransformation* GeneralTransformation, double UnitScale);
-FMatrixH TraverseTransformation3D(const A3DMiscCartesianTransformation* CartesianTransformation, double UnitScale);
-
-FString GetFileName(const FString& InFilePath);
 
 class FUVReparameterization
 {
@@ -142,20 +60,20 @@ public:
 
 	void SetScale(const double InUScale, const double InVScale)
 	{
-		Scale[EIso::IsoU] = InUScale;
-		Scale[EIso::IsoV] = InVScale;
-		Offset[EIso::IsoU] = 0.;
-		Offset[EIso::IsoV] = 0.;
+		Scale[CADKernel::EIso::IsoU] = InUScale;
+		Scale[CADKernel::EIso::IsoV] = InVScale;
+		Offset[CADKernel::EIso::IsoU] = 0.;
+		Offset[CADKernel::EIso::IsoV] = 0.;
 
 		SetNeedApply();
 	}
 
 	void SetCoef(const double InUScale, const double InUOffset, const double InVScale, const double InVOffset)
 	{
-		Scale[EIso::IsoU] = InUScale;
-		Scale[EIso::IsoV] = InVScale;
-		Offset[EIso::IsoU] = InUOffset;
-		Offset[EIso::IsoV] = InVOffset;
+		Scale[CADKernel::EIso::IsoU] = InUScale;
+		Scale[CADKernel::EIso::IsoV] = InVScale;
+		Offset[CADKernel::EIso::IsoU] = InUOffset;
+		Offset[CADKernel::EIso::IsoV] = InVOffset;
 		SetNeedApply();
 	}
 
@@ -166,7 +84,7 @@ public:
 
 	void SetNeedApply()
 	{
-		if (!FMath::IsNearlyEqual(Scale[EIso::IsoU], 1.) || !FMath::IsNearlyEqual(Scale[EIso::IsoV], 1.) || !FMath::IsNearlyEqual(Offset[EIso::IsoU], 0.) || !FMath::IsNearlyEqual(Offset[EIso::IsoV], 0.))
+		if (!FMath::IsNearlyEqual(Scale[CADKernel::EIso::IsoU], 1.) || !FMath::IsNearlyEqual(Scale[CADKernel::EIso::IsoV], 1.) || !FMath::IsNearlyEqual(Offset[CADKernel::EIso::IsoU], 0.) || !FMath::IsNearlyEqual(Offset[CADKernel::EIso::IsoV], 0.))
 		{
 			bNeedApply = true;
 		}
@@ -178,17 +96,17 @@ public:
 
 	void ScaleTransform(double InUScale, double InVScale)
 	{
-		Scale[EIso::IsoU] *= InUScale;
-		Scale[EIso::IsoV] *= InVScale;
-		Offset[EIso::IsoU] *= InUScale;
-		Offset[EIso::IsoV] *= InVScale;
+		Scale[CADKernel::EIso::IsoU] *= InUScale;
+		Scale[CADKernel::EIso::IsoV] *= InVScale;
+		Offset[CADKernel::EIso::IsoU] *= InUScale;
+		Offset[CADKernel::EIso::IsoV] *= InVScale;
 		SetNeedApply();
 	}
 
-	void Apply(FPoint& Point) const
+	void Apply(CADKernel::FPoint& Point) const
 	{
-		Point.X = Scale[EIso::IsoU] * Point.X + Offset[EIso::IsoU];
-		Point.Y = Scale[EIso::IsoV] * Point.Y + Offset[EIso::IsoV];
+		Point.X = Scale[CADKernel::EIso::IsoU] * Point.X + Offset[CADKernel::EIso::IsoU];
+		Point.Y = Scale[CADKernel::EIso::IsoV] * Point.Y + Offset[CADKernel::EIso::IsoV];
 	}
 
 	void AddUVTransform(A3DUVParameterizationData& Transform, bool bSwapUV = false)
@@ -203,106 +121,92 @@ public:
 		SetNeedApply();
 	}
 };
-}
-
-
-class FSurfacicBoundary;
-
-class FBody;
-class FCriterion;
-class FCurve;
-class FEntity;
-class FFaceMesh;
-class FModel;
-class FPoint;
-class FRestrictionCurve;
-class FSession;
-class FShell;
-class FSurface;
-class FTopologicalEdge;
-class FTopologicalLoop;
-class FTopologicalShapeEntity;
+} // ns TechSoftUtils
 
 class FTechSoftBridge
 {
 private:
+	FTechSoftFileParser& Parser;
 
-	FSession& Session;
-	FCADFileReport& Report;
+	CADKernel::FSession& Session;
+	CADKernel::FModel& Model;
+	CADKernel::FCADFileReport& Report;
 
 	const double GeometricTolerance;
 	const double SquareGeometricTolerance;
 	const double SquareJoiningVertexTolerance;
 
-	TMap<const A3DEntity*, TSharedPtr<FEntity>> TechSoftToEntity;
-	TMap<const A3DTopoCoEdge*, TSharedPtr<FTopologicalEdge>> A3DEdgeToEdge;
+	TMap<const A3DEntity*, CADKernel::FBody*> TechSoftToCADKernel;
+	TMap<CADKernel::FBody*, const A3DEntity*> CADKernelToTechSoft;
+	TMap<const A3DTopoCoEdge*, TSharedPtr<CADKernel::FTopologicalEdge>> A3DEdgeToEdge;
 
 	double UnitScale = 1;
 
 public:
-	FTechSoftBridge(FSession& InSession, FCADFileReport& InReport);
-
-	TSharedRef<FBody> AddBody(A3DRiBrepModel* A3DBrepModel, FEntityData& EntityData, double FileUnit);
+	FTechSoftBridge(FTechSoftFileParser& InParser, CADKernel::FSession& InSession, CADKernel::FCADFileReport& InReport);
+	CADKernel::FBody* AddBody(A3DRiBrepModel* A3DBRepModel, TMap<FString, FString> MetaData);
+	CADKernel::FBody* GetBody(A3DRiBrepModel* A3DBRepModel);
+	const A3DRiBrepModel* GetA3DBody(CADKernel::FBody* BRepModel);
 
 private:
 
-	void TraverseBrepData(const A3DTopoBrepData* A3DBrepData, TSharedRef<FBody>& Body);
-
-	void TraverseConnex(const A3DTopoConnex* A3DConnex, TSharedRef<FBody>& Body);
-	void TraverseShell(const A3DTopoShell* A3DShell, TSharedRef<FBody>& Body);
-	void AddFace(const A3DTopoFace* A3DFace, TSharedRef<FShell>& Body);
-
-	TSharedPtr<FTopologicalLoop> AddLoop(const A3DTopoLoop* A3DLoop, const TSharedRef<FSurface>& Surface, const TechSoftUtils::FUVReparameterization& UVReparameterization);
-	TSharedPtr<FTopologicalEdge> AddEdge(const A3DTopoCoEdge* A3DCoedge, const TSharedRef<FSurface>& Surface, const TechSoftUtils::FUVReparameterization& UVReparameterization, EOrientation& OutOrientation);
-
-	TSharedPtr<FSurface> AddSurface(const A3DSurfBase* A3DSurface, TechSoftUtils::FUVReparameterization& OutUVReparameterization);
-	TSharedPtr<FSurface> AddConeSurface(const A3DSurfCone* A3DSurface, TechSoftUtils::FUVReparameterization& OutUVReparameterization);
-	TSharedPtr<FSurface> AddCylinderSurface(const A3DSurfCylinder* A3DSurface, TechSoftUtils::FUVReparameterization& OutUVReparameterization);
-	TSharedPtr<FSurface> AddLinearTransfoSurface(const A3DSurfBase* A3DSurface, TechSoftUtils::FUVReparameterization& OutUVReparameterization);
-	TSharedPtr<FSurface> AddNurbsSurface(const A3DSurfNurbs* A3DSurface, TechSoftUtils::FUVReparameterization& OutUVReparameterization);
-	TSharedPtr<FSurface> AddOffsetSurface(const A3DSurfOffset* A3DSurface, TechSoftUtils::FUVReparameterization& OutUVReparameterization);
-	TSharedPtr<FSurface> AddPlaneSurface(const A3DSurfPlane* A3DSurface, TechSoftUtils::FUVReparameterization& OutUVReparameterization);
-	TSharedPtr<FSurface> AddRevolutionSurface(const A3DSurfRevolution* A3DSurface, TechSoftUtils::FUVReparameterization& OutUVReparameterization);
-	TSharedPtr<FSurface> AddRuledSurface(const A3DSurfRuled* A3DSurface, TechSoftUtils::FUVReparameterization& OutUVReparameterization);
-	TSharedPtr<FSurface> AddSphereSurface(const A3DSurfSphere* A3DSurface, TechSoftUtils::FUVReparameterization& OutUVReparameterization);
-	TSharedPtr<FSurface> AddTorusSurface(const A3DSurfTorus* A3DSurface, TechSoftUtils::FUVReparameterization& OutUVReparameterization);
-
-	TSharedPtr<FSurface> AddBlend01Surface(const A3DSurfBlend01* A3DSurface, TechSoftUtils::FUVReparameterization& OutUVReparameterization);
-	TSharedPtr<FSurface> AddBlend02Surface(const A3DSurfBlend02* A3DSurface, TechSoftUtils::FUVReparameterization& OutUVReparameterization);
-	TSharedPtr<FSurface> AddBlend03Surface(const A3DSurfBlend03* A3DSurface, TechSoftUtils::FUVReparameterization& OutUVReparameterization);
-	TSharedPtr<FSurface> AddCylindricalSurface(const A3DSurfCylindrical* A3DSurface, TechSoftUtils::FUVReparameterization& OutUVReparameterization);
-	TSharedPtr<FSurface> AddPipeSurface(const A3DSurfPipe* A3DSurface, TechSoftUtils::FUVReparameterization& OutUVReparameterization);
-	TSharedPtr<FSurface> AddExtrusionSurface(const A3DSurfExtrusion* A3DSurface, TechSoftUtils::FUVReparameterization& OutUVReparameterization);
-	TSharedPtr<FSurface> AddSurfaceFromCurves(const A3DSurfFromCurves* A3DSurface, TechSoftUtils::FUVReparameterization& OutUVReparameterization);
-	TSharedPtr<FSurface> AddTransformSurface(const A3DSurfTransform* A3DSurface, TechSoftUtils::FUVReparameterization& OutUVReparameterization);
-
-	TSharedPtr<FSurface> AddSurfaceAsNurbs(const A3DSurfBase* A3DSurface, TechSoftUtils::FUVReparameterization& OutUVReparameterization);
-	TSharedPtr<FSurface> AddSurfaceNurbs(const A3DSurfNurbsData& A3DNurbsData, TechSoftUtils::FUVReparameterization& OutUVReparameterization);
-
-	TSharedPtr<FRestrictionCurve> AddRestrictionCurve(const A3DCrvBase* A3DCurve, const TSharedRef<FSurface>& Surface);
-
-	TSharedPtr<FCurve> AddCurve(const A3DCrvBase* A3DCurve, const TechSoftUtils::FUVReparameterization& UVReparameterization = TechSoftUtils::FUVReparameterization::Identity);
-
-	TSharedPtr<FCurve> AddCurveCircle(const A3DCrvCircle* A3DCurve, const TechSoftUtils::FUVReparameterization& UVReparameterization);
-	TSharedPtr<FCurve> AddCurveComposite(const A3DCrvComposite* A3DCurve, const TechSoftUtils::FUVReparameterization& UVReparameterization);
-	TSharedPtr<FCurve> AddCurveEllipse(const A3DCrvEllipse* A3DCurve, const TechSoftUtils::FUVReparameterization& UVReparameterization);
-	TSharedPtr<FCurve> AddCurveHelix(const A3DCrvHelix* A3DCurve, const TechSoftUtils::FUVReparameterization& UVReparameterization);
-	TSharedPtr<FCurve> AddCurveHyperbola(const A3DCrvHyperbola* A3DCurve, const TechSoftUtils::FUVReparameterization& UVReparameterization);
-	TSharedPtr<FCurve> AddCurveLine(const A3DCrvLine* A3DCurve, const TechSoftUtils::FUVReparameterization& UVReparameterization);
-	TSharedPtr<FCurve> AddCurveNurbs(const A3DCrvNurbs* A3DCurve, const TechSoftUtils::FUVReparameterization& UVReparameterization);
-	TSharedPtr<FCurve> AddCurveParabola(const A3DCrvParabola* A3DCurve, const TechSoftUtils::FUVReparameterization& UVReparameterization);
-	TSharedPtr<FCurve> AddCurvePolyLine(const A3DCrvPolyLine* A3DCurve, const TechSoftUtils::FUVReparameterization& UVReparameterization);
-
-	TSharedPtr<FCurve> AddCurveBlend02Boundary(const A3DCrvBlend02Boundary* A3DCurve, const TechSoftUtils::FUVReparameterization& UVReparameterization);
-	TSharedPtr<FCurve> AddCurveEquation(const A3DCrvEquation* A3DCurve, const TechSoftUtils::FUVReparameterization& UVReparameterization);
-	TSharedPtr<FCurve> AddCurveIntersection(const A3DCrvIntersection* A3DCurve, const TechSoftUtils::FUVReparameterization& UVReparameterization);
-	TSharedPtr<FCurve> AddCurveOffset(const A3DCrvOffset* A3DCurve, const TechSoftUtils::FUVReparameterization& UVReparameterization);
-	TSharedPtr<FCurve> AddCurveOnSurf(const A3DCrvOnSurf* A3DCurve, const TechSoftUtils::FUVReparameterization& UVReparameterization);
-	TSharedPtr<FCurve> AddCurveTransform(const A3DCrvTransform* A3DCurve, const TechSoftUtils::FUVReparameterization& UVReparameterization);
-
-	TSharedPtr<FCurve> AddCurveAsNurbs(const A3DCrvBase* A3DCurve, const TechSoftUtils::FUVReparameterization& UVReparameterization);
-
-	void AddMetadata(FEntityData& EntityData, FTopologicalShapeEntity& Entity);
+	void TraverseBrepData(const A3DTopoBrepData* A3DBrepData, TSharedRef<CADKernel::FBody>& Body);
+
+	void TraverseConnex(const A3DTopoConnex* A3DConnex, TSharedRef<CADKernel::FBody>& Body);
+	void TraverseShell(const A3DTopoShell* A3DShell, TSharedRef<CADKernel::FBody>& Body);
+	void AddFace(const A3DTopoFace* A3DFace, CADKernel::EOrientation Orientation, TSharedRef<CADKernel::FShell>& Body);
+
+	TSharedPtr<CADKernel::FTopologicalLoop> AddLoop(const A3DTopoLoop* A3DLoop, const TSharedRef<CADKernel::FSurface>& Surface, const TechSoftUtils::FUVReparameterization& UVReparameterization);
+	TSharedPtr<CADKernel::FTopologicalEdge> AddEdge(const A3DTopoCoEdge* A3DCoedge, const TSharedRef<CADKernel::FSurface>& Surface, const TechSoftUtils::FUVReparameterization& UVReparameterization, CADKernel::EOrientation& OutOrientation);
+
+	TSharedPtr<CADKernel::FSurface> AddSurface(const A3DSurfBase* A3DSurface, TechSoftUtils::FUVReparameterization& OutUVReparameterization);
+	TSharedPtr<CADKernel::FSurface> AddConeSurface(const A3DSurfCone* A3DSurface, TechSoftUtils::FUVReparameterization& OutUVReparameterization);
+	TSharedPtr<CADKernel::FSurface> AddCylinderSurface(const A3DSurfCylinder* A3DSurface, TechSoftUtils::FUVReparameterization& OutUVReparameterization);
+	TSharedPtr<CADKernel::FSurface> AddLinearTransfoSurface(const A3DSurfBase* A3DSurface, TechSoftUtils::FUVReparameterization& OutUVReparameterization);
+	TSharedPtr<CADKernel::FSurface> AddNurbsSurface(const A3DSurfNurbs* A3DSurface, TechSoftUtils::FUVReparameterization& OutUVReparameterization);
+	TSharedPtr<CADKernel::FSurface> AddOffsetSurface(const A3DSurfOffset* A3DSurface, TechSoftUtils::FUVReparameterization& OutUVReparameterization);
+	TSharedPtr<CADKernel::FSurface> AddPlaneSurface(const A3DSurfPlane* A3DSurface, TechSoftUtils::FUVReparameterization& OutUVReparameterization);
+	TSharedPtr<CADKernel::FSurface> AddRevolutionSurface(const A3DSurfRevolution* A3DSurface, TechSoftUtils::FUVReparameterization& OutUVReparameterization);
+	TSharedPtr<CADKernel::FSurface> AddRuledSurface(const A3DSurfRuled* A3DSurface, TechSoftUtils::FUVReparameterization& OutUVReparameterization);
+	TSharedPtr<CADKernel::FSurface> AddSphereSurface(const A3DSurfSphere* A3DSurface, TechSoftUtils::FUVReparameterization& OutUVReparameterization);
+	TSharedPtr<CADKernel::FSurface> AddTorusSurface(const A3DSurfTorus* A3DSurface, TechSoftUtils::FUVReparameterization& OutUVReparameterization);
+
+	TSharedPtr<CADKernel::FSurface> AddBlend01Surface(const A3DSurfBlend01* A3DSurface, TechSoftUtils::FUVReparameterization& OutUVReparameterization);
+	TSharedPtr<CADKernel::FSurface> AddBlend02Surface(const A3DSurfBlend02* A3DSurface, TechSoftUtils::FUVReparameterization& OutUVReparameterization);
+	TSharedPtr<CADKernel::FSurface> AddBlend03Surface(const A3DSurfBlend03* A3DSurface, TechSoftUtils::FUVReparameterization& OutUVReparameterization);
+	TSharedPtr<CADKernel::FSurface> AddCylindricalSurface(const A3DSurfCylindrical* A3DSurface, TechSoftUtils::FUVReparameterization& OutUVReparameterization);
+	TSharedPtr<CADKernel::FSurface> AddPipeSurface(const A3DSurfPipe* A3DSurface, TechSoftUtils::FUVReparameterization& OutUVReparameterization);
+	TSharedPtr<CADKernel::FSurface> AddExtrusionSurface(const A3DSurfExtrusion* A3DSurface, TechSoftUtils::FUVReparameterization& OutUVReparameterization);
+	TSharedPtr<CADKernel::FSurface> AddSurfaceFromCurves(const A3DSurfFromCurves* A3DSurface, TechSoftUtils::FUVReparameterization& OutUVReparameterization);
+	TSharedPtr<CADKernel::FSurface> AddTransformSurface(const A3DSurfTransform* A3DSurface, TechSoftUtils::FUVReparameterization& OutUVReparameterization);
+
+	TSharedPtr<CADKernel::FSurface> AddSurfaceAsNurbs(const A3DSurfBase* A3DSurface, TechSoftUtils::FUVReparameterization& OutUVReparameterization);
+	TSharedPtr<CADKernel::FSurface> AddSurfaceNurbs(const A3DSurfNurbsData& A3DNurbsData, TechSoftUtils::FUVReparameterization& OutUVReparameterization);
+
+	TSharedPtr<CADKernel::FRestrictionCurve> AddRestrictionCurve(const A3DCrvBase* A3DCurve, const TSharedRef<CADKernel::FSurface>& Surface);
+
+	TSharedPtr<CADKernel::FCurve> AddCurve(const A3DCrvBase* A3DCurve, const TechSoftUtils::FUVReparameterization& UVReparameterization = TechSoftUtils::FUVReparameterization::Identity);
+
+	TSharedPtr<CADKernel::FCurve> AddCurveCircle(const A3DCrvCircle* A3DCurve, const TechSoftUtils::FUVReparameterization& UVReparameterization);
+	TSharedPtr<CADKernel::FCurve> AddCurveComposite(const A3DCrvComposite* A3DCurve, const TechSoftUtils::FUVReparameterization& UVReparameterization);
+	TSharedPtr<CADKernel::FCurve> AddCurveEllipse(const A3DCrvEllipse* A3DCurve, const TechSoftUtils::FUVReparameterization& UVReparameterization);
+	TSharedPtr<CADKernel::FCurve> AddCurveHelix(const A3DCrvHelix* A3DCurve, const TechSoftUtils::FUVReparameterization& UVReparameterization);
+	TSharedPtr<CADKernel::FCurve> AddCurveHyperbola(const A3DCrvHyperbola* A3DCurve, const TechSoftUtils::FUVReparameterization& UVReparameterization);
+	TSharedPtr<CADKernel::FCurve> AddCurveLine(const A3DCrvLine* A3DCurve, const TechSoftUtils::FUVReparameterization& UVReparameterization);
+	TSharedPtr<CADKernel::FCurve> AddCurveNurbs(const A3DCrvNurbs* A3DCurve, const TechSoftUtils::FUVReparameterization& UVReparameterization);
+	TSharedPtr<CADKernel::FCurve> AddCurveParabola(const A3DCrvParabola* A3DCurve, const TechSoftUtils::FUVReparameterization& UVReparameterization);
+	TSharedPtr<CADKernel::FCurve> AddCurvePolyLine(const A3DCrvPolyLine* A3DCurve, const TechSoftUtils::FUVReparameterization& UVReparameterization);
+
+	TSharedPtr<CADKernel::FCurve> AddCurveBlend02Boundary(const A3DCrvBlend02Boundary* A3DCurve, const TechSoftUtils::FUVReparameterization& UVReparameterization);
+	TSharedPtr<CADKernel::FCurve> AddCurveEquation(const A3DCrvEquation* A3DCurve, const TechSoftUtils::FUVReparameterization& UVReparameterization);
+	TSharedPtr<CADKernel::FCurve> AddCurveIntersection(const A3DCrvIntersection* A3DCurve, const TechSoftUtils::FUVReparameterization& UVReparameterization);
+	TSharedPtr<CADKernel::FCurve> AddCurveOffset(const A3DCrvOffset* A3DCurve, const TechSoftUtils::FUVReparameterization& UVReparameterization);
+	TSharedPtr<CADKernel::FCurve> AddCurveOnSurf(const A3DCrvOnSurf* A3DCurve, const TechSoftUtils::FUVReparameterization& UVReparameterization);
+	TSharedPtr<CADKernel::FCurve> AddCurveTransform(const A3DCrvTransform* A3DCurve, const TechSoftUtils::FUVReparameterization& UVReparameterization);
+
+	TSharedPtr<CADKernel::FCurve> AddCurveAsNurbs(const A3DCrvBase* A3DCurve, const TechSoftUtils::FUVReparameterization& UVReparameterization);
+
+	void AddMetadata(FEntityMetaData& EntityData, CADKernel::FTopologicalShapeEntity& Entity);
 };
 }
 
diff --git a/Engine/Plugins/Enterprise/DatasmithCADImporter/Source/CADInterfaces/Private/TechSoftFileParser.cpp b/Engine/Plugins/Enterprise/DatasmithCADImporter/Source/CADInterfaces/Private/TechSoftFileParser.cpp
--- a/Engine/Plugins/Enterprise/DatasmithCADImporter/Source/CADInterfaces/Private/TechSoftFileParser.cpp
+++ b/Engine/Plugins/Enterprise/DatasmithCADImporter/Source/CADInterfaces/Private/TechSoftFileParser.cpp
@@ -2,14 +2,11 @@
 
 #include "TechSoftFileParser.h"
 
-#include "CADOptions.h"
 #include "CADFileData.h"
-
-#ifdef USE_TECHSOFT_SDK
-
-#include "TechSoftInterface.h"
+#include "CADOptions.h"
 #include "TechSoftUtils.h"
 #include "TechSoftUtilsPrivate.h"
+#include "TUniqueTechSoftObj.h"
 
 #include "Serialization/JsonSerializer.h"
 #include "Serialization/JsonWriter.h"
@@ -18,6 +15,8 @@
 namespace CADLibrary
 {
 
+#ifdef USE_TECHSOFT_SDK
+
 namespace TechSoftFileParserImpl
 {
 // Functions to clean metadata
@@ -137,7 +136,6 @@ void SetIOOption(A3DImport& Importer)
 	Importer.m_sLoadData.m_sGeneral.m_eReadingMode2D3D = kA3DRead_3D;
 
 	Importer.m_sLoadData.m_sGeneral.m_eReadGeomTessMode = kA3DReadGeomAndTess/*kA3DReadGeomOnly*/;
-	Importer.m_sLoadData.m_sGeneral.m_eDefaultUnit;
 	Importer.m_sLoadData.m_sGeneral.m_bReadFeature = A3D_FALSE;
 
 	Importer.m_sLoadData.m_sGeneral.m_bReadConstraints = A3D_FALSE;
@@ -192,6 +190,8 @@ void UpdateIOOptionAccordingToFormat(const CADLibrary::ECADFormat Format, A3DImp
 
 } // ns TechSoftFileParserImpl
 
+#endif
+
 FTechSoftFileParser::FTechSoftFileParser(FCADFileData& InCADData, const FString& EnginePluginsPath)
 	: CADFileData(InCADData)
 	, TechSoftInterface(FTechSoftInterface::Get())
@@ -199,6 +199,7 @@ FTechSoftFileParser::FTechSoftFileParser(FCADFileData& InCADData, const FString&
 	TechSoftInterface.InitializeKernel(*EnginePluginsPath);
 }
 
+#ifdef USE_TECHSOFT_SDK
 ECADParsingResult FTechSoftFileParser::Process()
 {
 	const FFileDescriptor& File = CADFileData.GetCADFileDescription();
@@ -208,7 +209,7 @@ ECADParsingResult FTechSoftFileParser::Process()
 		return ECADParsingResult::FileNotFound;
 	}
 
-	A3DImport Import(*File.GetPathOfFileToLoad());
+	A3DImport Import(TCHAR_TO_UTF8(*File.GetPathOfFileToLoad())); 
 
 	TechSoftFileParserImpl::SetIOOption(Import);
 
@@ -236,7 +237,7 @@ ECADParsingResult FTechSoftFileParser::Process()
 			break;
 		}
 
-		case A3DStatus::A3D_LOAD_CANNOT_ACCESS_CADFILE: 
+		case A3DStatus::A3D_LOAD_CANNOT_ACCESS_CADFILE:
 		{
 			CADFileData.AddWarningMessages(FString::Printf(TEXT("File %s hasn't been loaded because the input path cannot be opened by the running process for reading."), *File.GetFileName()));
 			break;
@@ -268,6 +269,12 @@ ECADParsingResult FTechSoftFileParser::Process()
 		FileUnit = ModelFileData->m_dUnit;
 	}
 
+	// Adapt BRep to CADKernel
+	if (AdaptBRepModel() != A3D_SUCCESS)
+	{
+		return ECADParsingResult::ProcessFailed;
+	}
+
 	// save the file for the next load
 	if (CADFileData.IsCacheDefined())
 	{
@@ -278,17 +285,9 @@ ECADParsingResult FTechSoftFileParser::Process()
 		}
 	}
 
-	if (bForceSew || CADFileData.GetImportParameters().GetStitchingTechnique() == StitchingSew && FImportParameters::bGDisableCADKernelTessellation)
+	if (bForceSew || CADFileData.GetImportParameters().GetStitchingTechnique() == StitchingSew)
 	{
-		CADLibrary::TUniqueTSObj<A3DSewOptionsData> SewData;
-		SewData->m_bComputePreferredOpenShellOrientation = false;
-		double ToleranceMM = 0.01 / FileUnit;
-
-		A3DStatus Status = TechSoftInterface::SewModel(ModelFile.Get(), ToleranceMM, SewData.GetPtr());
-		if(Status != A3DStatus::A3D_SUCCESS)
-		{
-			// To do but what ?
-		}
+		SewModel();
 	}
 
 	ReserveCADFileData();
@@ -307,109 +306,111 @@ ECADParsingResult FTechSoftFileParser::Process()
 	{
 		CADFileData.GetComponentAt(0).MetaData.Add(TEXT("TechsoftVersion"), TechSoftVersion);
 	}
-	
+
 	ModelFile.Reset();
 
 	return Result;
 }
 
-void FTechSoftFileParser::GenerateBodyMeshes()
+void FTechSoftFileParser::SewModel()
 {
-	if (FImportParameters::bGDisableCADKernelTessellation)
-	{
-		for (TPair<A3DRiRepresentationItem*, int32>& Entry : RepresentationItemsCache)
-		{
-			A3DRiRepresentationItem* RepresentationItemPtr = Entry.Key;
-
-			FArchiveBody& Body = CADFileData.GetBodyAt(Entry.Value);
-
-			FBodyMesh& BodyMesh = CADFileData.AddBodyMesh(Body.ObjectId, Body);
-
-			// todo
-			//if (CADFileData.GetImportParameters().GetStitchingTechnique() == StitchingHeal)
-			//{
-			//	TUniqueTSObj<A3DSewOptionsData> SewData;
-			//	SewData->m_bComputePreferredOpenShellOrientation = false;
-			//	double ToleranceMM = 0.01 / FileUnit;
-			//	A3DRiBrepModel** NewBReps;
-			//	uint32 NewBRepCount = 0;
-			//	A3DStatus Status = TechSoftUtils::HealBRep(&BRepModelPtr, ToleranceMM, &*SewData, &NewBReps, NewBRepCount);
-			//	if (Status == A3DStatus::A3D_SUCCESS)
-			//	{
-			//		// Aggregate all the bodies into one?
-			//		MeshRepresentationsWithTechSoft(NewBRepCount, NewBReps, Body);
-			//		TechSoftInterface.FillBodyMeshes(NewBRepCount, NewBReps, CADFileData.GetImportParameters(), FileUnit, BodyMesh);???
-			//	}
-			//}
-			//else
-			{
-				TechSoftUtils::FillBodyMesh(RepresentationItemPtr, CADFileData.GetImportParameters(), FileUnit, BodyMesh);
-			}
+	CADLibrary::TUniqueTSObj<A3DSewOptionsData> SewData;
+	SewData->m_bComputePreferredOpenShellOrientation = false;
+	double ToleranceMM = 0.01 / FileUnit;
 
-			// Convert material
-			FCADUUID DefaultColorName = Body.ColorFaceSet.Num() > 0 ? *Body.ColorFaceSet.begin()  : 0;
-			FCADUUID DefaultMaterialName = Body.MaterialFaceSet.Num() > 0 ? *Body.MaterialFaceSet.begin() : 0;
-
-			for (FTessellationData& Tessellation : BodyMesh.Faces)
-			{
-				FCADUUID ColorName = DefaultColorName;
-				FCADUUID MaterialName = DefaultMaterialName;
+	TechSoftInterface::SewModel(ModelFile.Get(), ToleranceMM, SewData.GetPtr());
+}
 
-				// Extract proper color or material based on style index
-				uint32 CachedStyleIndex = Tessellation.MaterialName;
-				Tessellation.MaterialName = 0;
 
-				if (CachedStyleIndex != FTechSoftInterface::InvalidScriptIndex)
-				{
-					ExtractGraphStyleProperties(CachedStyleIndex, ColorName, MaterialName);
-				}
+void FTechSoftFileParser::GenerateBodyMeshes()
+{
+	for (TPair<A3DRiRepresentationItem*, int32>& Entry : RepresentationItemsCache)
+	{
+		A3DRiRepresentationItem* RepresentationItemPtr = Entry.Key;
+		FArchiveBody& Body = CADFileData.GetBodyAt(Entry.Value);
+		GenerateBodyMesh(RepresentationItemPtr, Body);
+	}
+}
 
-				if (ColorName)
-				{
-					Tessellation.ColorName = ColorName;
-					BodyMesh.ColorSet.Add(ColorName);
-				}
+void FTechSoftFileParser::GenerateBodyMesh(A3DRiRepresentationItem* Representation, FArchiveBody& Body)
+{
+	FBodyMesh& BodyMesh = CADFileData.AddBodyMesh(Body.ObjectId, Body);
 
-				if (MaterialName)
-				{
-					Tessellation.MaterialName = MaterialName;
-					BodyMesh.MaterialSet.Add(MaterialName);
-				}
-			}
+	if (CADFileData.GetImportParameters().GetStitchingTechnique() == StitchingHeal)
+	{
+		TUniqueTSObj<A3DSewOptionsData> SewData;
+		SewData->m_bComputePreferredOpenShellOrientation = false;
+		double ToleranceMM = 0.01 / FileUnit;
+		A3DRiBrepModel** NewBReps;
+		uint32 NewBRepCount = 0;
+		A3DStatus Status = TechSoftInterface::SewBReps(&Representation, 1, ToleranceMM, SewData.GetPtr(), &NewBReps, NewBRepCount);
+		ensure(NewBRepCount == 1);
+		if (Status == A3DStatus::A3D_SUCCESS)
+		{
+			Representation = NewBReps[0];
+		}
+	}
 
-			Body.ColorFaceSet = BodyMesh.ColorSet;
-			Body.MaterialFaceSet = BodyMesh.MaterialSet;
+	TechSoftUtils::FillBodyMesh(Representation, CADFileData.GetImportParameters(), FileUnit, BodyMesh);
 
-			// Write part's representation as Prc file if it is a BRep
-			A3DEEntityType Type;
-			A3DEntityGetType(RepresentationItemPtr, &Type);
+	// Convert material
+	FCADUUID DefaultColorName = Body.ColorFaceSet.Num() > 0 ? *Body.ColorFaceSet.begin() : 0;
+	FCADUUID DefaultMaterialName = Body.MaterialFaceSet.Num() > 0 ? *Body.MaterialFaceSet.begin() : 0;
 
-			if (Type == kA3DTypeRiBrepModel)
-			{
-				FString FilePath = CADFileData.GetBodyCachePath(Body.MeshActorName);
-				if (!FilePath.IsEmpty())
-				{
-					TSharedPtr<FJsonObject> JsonObject = MakeShared<FJsonObject>();
+	for (FTessellationData& Tessellation : BodyMesh.Faces)
+	{
+		FCADUUID ColorName = DefaultColorName;
+		FCADUUID MaterialName = DefaultMaterialName;
 
-					// Save file unit and default color and material attributes in a json string
-					// This will be used when the file is reloaded
-					JsonObject->SetNumberField(JSON_ENTRY_FILE_UNIT, FileUnit);
-					JsonObject->SetNumberField(JSON_ENTRY_COLOR_NAME, DefaultColorName);
-					JsonObject->SetNumberField(JSON_ENTRY_MATERIAL_NAME, DefaultMaterialName);
+		// Extract proper color or material based on style index
+		uint32 CachedStyleIndex = Tessellation.MaterialName;
+		Tessellation.MaterialName = 0;
 
-					FString JsonString;
-					TSharedRef< TJsonWriter< TCHAR, TPrettyJsonPrintPolicy<TCHAR> > > JsonWriter = TJsonWriterFactory< TCHAR, TPrettyJsonPrintPolicy<TCHAR> >::Create(&JsonString);
+		if (CachedStyleIndex != FTechSoftInterface::InvalidScriptIndex)
+		{
+			ExtractGraphStyleProperties(CachedStyleIndex, ColorName, MaterialName);
+		}
 
-					FJsonSerializer::Serialize(JsonObject.ToSharedRef(), JsonWriter);
+		if (ColorName)
+		{
+			Tessellation.ColorName = ColorName;
+			BodyMesh.ColorSet.Add(ColorName);
+		}
 
-					TechSoftUtils::SaveBodiesToPrcFile(&RepresentationItemPtr, 1, FilePath, JsonString);
-				}
-			}
+		if (MaterialName)
+		{
+			Tessellation.MaterialName = MaterialName;
+			BodyMesh.MaterialSet.Add(MaterialName);
 		}
 	}
-	else
+
+	Body.ColorFaceSet = BodyMesh.ColorSet;
+	Body.MaterialFaceSet = BodyMesh.MaterialSet;
+
+	// Write part's representation as Prc file if it is a BRep
+	A3DEEntityType Type;
+	A3DEntityGetType(Representation, &Type);
+
+	if (Type == kA3DTypeRiBrepModel)
 	{
-		// Mesh with CADKernel
+		FString FilePath = CADFileData.GetBodyCachePath(Body.MeshActorName);
+		if (!FilePath.IsEmpty())
+		{
+			TSharedPtr<FJsonObject> JsonObject = MakeShared<FJsonObject>();
+
+			// Save file unit and default color and material attributes in a json string
+			// This will be used when the file is reloaded
+			JsonObject->SetNumberField(JSON_ENTRY_FILE_UNIT, FileUnit);
+			JsonObject->SetNumberField(JSON_ENTRY_COLOR_NAME, DefaultColorName);
+			JsonObject->SetNumberField(JSON_ENTRY_MATERIAL_NAME, DefaultMaterialName);
+
+			FString JsonString;
+			TSharedRef< TJsonWriter< TCHAR, TPrettyJsonPrintPolicy<TCHAR> > > JsonWriter = TJsonWriterFactory< TCHAR, TPrettyJsonPrintPolicy<TCHAR> >::Create(&JsonString);
+
+			FJsonSerializer::Serialize(JsonObject.ToSharedRef(), JsonWriter);
+
+			TechSoftUtils::SaveBodiesToPrcFile(&Representation, 1, FilePath, JsonString);
+		}
 	}
 }
 
@@ -495,7 +496,7 @@ void FTechSoftFileParser::TraverseConfigurationSet(const A3DAsmProductOccurrence
 	A3DMiscTransformation* Location = ConfigurationSetData->m_pLocation;
 	if (Location)
 	{
-		TransformMatrix = TraverseTransformation(Location);
+		TransformMatrix = ExtractTransformation(Location);
 	}
 
 	TUniqueTSObj<A3DAsmProductOccurrenceData> ConfigurationData;
@@ -594,14 +595,10 @@ void FTechSoftFileParser::TraverseReference(const A3DAsmProductOccurrence* Refer
 	}
 
 	ExtractSpecificMetaData(ReferencePtr, MetaData);
-	BuildReferenceName(MetaData.MetaData);
-
-	ExtractMaterialProperties(ReferencePtr);
+	BuildReferenceName(MetaData);
 
-	int32 ComponentId = LastEntityId++;
-	int32 Index = CADFileData.AddComponent(ComponentId);
-	FArchiveComponent& Component = CADFileData.GetComponentAt(Index);
-	Component.MetaData = MoveTemp(MetaData.MetaData);
+	FArchiveInstance EmptyInstance;
+	FArchiveComponent& Component = FTechSoftFileParser::AddComponent(MetaData, EmptyInstance);
 
 	TUniqueTSObj<A3DAsmProductOccurrenceData> ReferenceData(ReferencePtr);
 	if (!ReferenceData.IsValid())
@@ -613,9 +610,9 @@ void FTechSoftFileParser::TraverseReference(const A3DAsmProductOccurrence* Refer
 	A3DMiscTransformation* Location = ReferenceData->m_pLocation;
 	if (Location)
 	{
-		ReferenceMatrix = TraverseTransformation(Location);
+		ReferenceMatrix = ExtractTransformation(Location);
 	}
-	
+
 	Component.TransformMatrix = ParentMatrix * ReferenceMatrix;
 
 	for (uint32 OccurenceIndex = 0; OccurenceIndex < ReferenceData->m_uiPOccurrencesSize; ++OccurenceIndex)
@@ -700,7 +697,7 @@ int32 FTechSoftFileParser::AddBody(FEntityMetaData& BodyMetaData, const FMatrix&
 	FArchiveBody& Body = CADFileData.GetBodyAt(BodyIndex);
 	Body.TransformMatrix = Matrix;
 	Body.MetaData = MoveTemp(BodyMetaData.MetaData);
-	if(BodyMetaData.ColorName != 0)
+	if (BodyMetaData.ColorName != 0)
 	{
 		Body.ColorFaceSet.Add(BodyMetaData.ColorName);
 	}
@@ -739,9 +736,7 @@ FCadId FTechSoftFileParser::TraverseOccurrence(const A3DAsmProductOccurrence* Oc
 	}
 
 	ExtractSpecificMetaData(OccurrencePtr, InstanceMetaData);
-	BuildInstanceName(InstanceMetaData.MetaData);
-
-	ExtractMaterialProperties(OccurrencePtr);
+	BuildInstanceName(InstanceMetaData);
 
 	FArchiveInstance& Instance = AddInstance(InstanceMetaData);
 
@@ -755,7 +750,7 @@ FCadId FTechSoftFileParser::TraverseOccurrence(const A3DAsmProductOccurrence* Oc
 
 	if (Location)
 	{
-		Instance.TransformMatrix = TraverseTransformation(Location);
+		Instance.TransformMatrix = ExtractTransformation(Location);
 	}
 
 	if (PrototypeMetaData.bUnloaded)
@@ -940,7 +935,6 @@ void FTechSoftFileParser::ProcessPrototype(const A3DAsmProductOccurrence* InProt
 		{
 			ExtractMetaData(PrototypePtr, OutPrototypeMetaData);
 			ExtractSpecificMetaData(PrototypePtr, OutPrototypeMetaData);
-			ExtractMaterialProperties(PrototypePtr);
 
 			TUniqueTSObj<A3DUTF8Char*> FilePathUTF8Ptr;
 			FilePathUTF8Ptr.FillWith(&TechSoftInterface::GetFilePathName, PrototypePtr);
@@ -982,7 +976,7 @@ void FTechSoftFileParser::ProcessPrototype(const A3DAsmProductOccurrence* InProt
 		{
 			PrototypeData.FillFrom(PrototypeData->m_pPrototype);
 		}
-		if(PrototypeData.IsValid())
+		if (PrototypeData.IsValid())
 		{
 			*OutLocation = PrototypeData->m_pLocation;
 		}
@@ -997,7 +991,7 @@ void FTechSoftFileParser::ProcessPrototype(const A3DAsmProductOccurrence* InProt
 		OutPrototypeMetaData.ExternalFile.Empty();
 	}
 
-	BuildReferenceName(OutPrototypeMetaData.MetaData);
+	BuildReferenceName(OutPrototypeMetaData);
 }
 
 void FTechSoftFileParser::CountUnderPrototype(const A3DAsmProductOccurrence* Prototype)
@@ -1022,9 +1016,7 @@ void FTechSoftFileParser::TraversePartDefinition(const A3DAsmPartDefinition* Par
 	}
 
 	ExtractSpecificMetaData(PartDefinitionPtr, PartMetaData);
-	BuildPartName(PartMetaData.MetaData);
-
-	ExtractMaterialProperties(PartDefinitionPtr);
+	BuildPartName(PartMetaData);
 
 	TUniqueTSObj<A3DAsmPartDefinitionData> PartData(PartDefinitionPtr);
 	if (PartData.IsValid())
@@ -1052,7 +1044,7 @@ void FTechSoftFileParser::CountUnderPartDefinition(const A3DAsmPartDefinition* P
 	}
 }
 
-FCadId FTechSoftFileParser::TraverseRepresentationItem(A3DRiRepresentationItem* RepresentationItem, FEntityMetaData& PartMetaData)
+FCadId FTechSoftFileParser::TraverseRepresentationItem(A3DRiRepresentationItem* RepresentationItem, const FEntityMetaData& PartMetaData)
 {
 	A3DEEntityType Type;
 	A3DEntityGetType(RepresentationItem, &Type);
@@ -1091,7 +1083,7 @@ void FTechSoftFileParser::CountUnderRepresentationItem(const A3DRiRepresentation
 	}
 }
 
-FCadId FTechSoftFileParser::TraverseRepresentationSet(const A3DRiSet* RepresentationSetPtr, FEntityMetaData& PartMetaData)
+FCadId FTechSoftFileParser::TraverseRepresentationSet(const A3DRiSet* RepresentationSetPtr, const FEntityMetaData& PartMetaData)
 {
 	TUniqueTSObj<A3DRiSetData> RepresentationSetData(RepresentationSetPtr);
 	if (!RepresentationSetData.IsValid())
@@ -1107,8 +1099,6 @@ FCadId FTechSoftFileParser::TraverseRepresentationSet(const A3DRiSet* Representa
 		return 0;
 	}
 
-	ExtractMaterialProperties(RepresentationSetPtr);
-
 	FCadId RepresentationSetId = 0;
 	FArchiveComponent& RepresentationSet = AddOccurence(RepresentationSetMetaData, RepresentationSetId);
 
@@ -1135,7 +1125,7 @@ void FTechSoftFileParser::CountUnderRepresentationSet(const A3DRiSet* Representa
 	}
 }
 
-FCadId FTechSoftFileParser::TraverseBRepModel(A3DRiBrepModel* BRepModelPtr, FEntityMetaData& PartMetaData)
+FCadId FTechSoftFileParser::TraverseBRepModel(A3DRiBrepModel* BRepModelPtr, const FEntityMetaData& PartMetaData)
 {
 	if (!BRepModelPtr)
 	{
@@ -1144,11 +1134,14 @@ FCadId FTechSoftFileParser::TraverseBRepModel(A3DRiBrepModel* BRepModelPtr, FEnt
 
 	FEntityMetaData BRepMetaData;
 	ExtractMetaData(BRepModelPtr, BRepMetaData);
+
 	if (!BRepMetaData.bShow || BRepMetaData.bRemoved)
 	{
 		return 0;
 	}
 
+	TUniqueTSObj<A3DRiBrepModelData> BRepModelData(BRepModelPtr);
+	BuildBodyName(BRepMetaData, PartMetaData, (bool) BRepModelData->m_bSolid);
 
 	if (int32* BodyIndexPtr = RepresentationItemsCache.Find(BRepModelPtr))
 	{
@@ -1156,13 +1149,12 @@ FCadId FTechSoftFileParser::TraverseBRepModel(A3DRiBrepModel* BRepModelPtr, FEnt
 	}
 
 	ExtractSpecificMetaData(BRepModelPtr, BRepMetaData);
-	ExtractMaterialProperties(BRepModelPtr);
 
 	FMatrix Matrix = FMatrix::Identity;
 	TUniqueTSObj<A3DRiRepresentationItemData> RepresentationData(BRepModelPtr);
 	if (RepresentationData->m_pCoordinateSystem)
 	{
-		Matrix = FTechSoftFileParser::TraverseCoordinateSystem(RepresentationData->m_pCoordinateSystem);
+		Matrix = ExtractCoordinateSystem(RepresentationData->m_pCoordinateSystem);
 	}
 
 	int32 BodyIndex = AddBody(BRepMetaData, Matrix);
@@ -1173,7 +1165,7 @@ FCadId FTechSoftFileParser::TraverseBRepModel(A3DRiBrepModel* BRepModelPtr, FEnt
 	return Body.ObjectId;
 }
 
-FCadId FTechSoftFileParser::TraversePolyBRepModel(A3DRiPolyBrepModel* PolygonalPtr, FEntityMetaData& PartMetaData)
+FCadId FTechSoftFileParser::TraversePolyBRepModel(A3DRiPolyBrepModel* PolygonalPtr, const FEntityMetaData& PartMetaData)
 {
 	if (!PolygonalPtr)
 	{
@@ -1182,24 +1174,27 @@ FCadId FTechSoftFileParser::TraversePolyBRepModel(A3DRiPolyBrepModel* PolygonalP
 
 	FEntityMetaData BRepMetaData;
 	ExtractMetaData(PolygonalPtr, BRepMetaData);
+
 	if (!BRepMetaData.bShow || BRepMetaData.bRemoved)
 	{
 		return 0;
 	}
 
+	TUniqueTSObj<A3DRiPolyBrepModelData> BRepModelData(PolygonalPtr);
+	BuildBodyName(BRepMetaData, PartMetaData, (bool) BRepModelData->m_bIsClosed);
+
 	if (int32* BodyIndexPtr = RepresentationItemsCache.Find(PolygonalPtr))
 	{
 		return CADFileData.GetBodyAt(*BodyIndexPtr).ObjectId;
 	}
 
 	ExtractSpecificMetaData(PolygonalPtr, BRepMetaData);
-	ExtractMaterialProperties(PolygonalPtr);
 
 	FMatrix Matrix = FMatrix::Identity;
 	TUniqueTSObj<A3DRiRepresentationItemData> RepresentationData(PolygonalPtr);
 	if (RepresentationData->m_pCoordinateSystem)
 	{
-		Matrix = FTechSoftFileParser::TraverseCoordinateSystem(RepresentationData->m_pCoordinateSystem);
+		Matrix = ExtractCoordinateSystem(RepresentationData->m_pCoordinateSystem);
 	}
 
 	int32 BodyIndex = AddBody(BRepMetaData, Matrix);
@@ -1246,15 +1241,16 @@ void FTechSoftFileParser::ExtractMetaData(const A3DEntity* Entity, FEntityMetaDa
 	}
 }
 
-void FTechSoftFileParser::BuildReferenceName(TMap<FString, FString>& MetaData)
+void FTechSoftFileParser::BuildReferenceName(FEntityMetaData& ReferenceData)
 {
+	TMap<FString, FString>& MetaData = ReferenceData.MetaData;
 	if (MetaData.IsEmpty())
 	{
 		return;
 	}
 
 	FString* NamePtr = MetaData.Find(TEXT("InstanceName"));
-	if (NamePtr != nullptr)
+	if (NamePtr != nullptr && !NamePtr->IsEmpty())
 	{
 		FString& Name = MetaData.FindOrAdd(TEXT("Name"));
 		Name = *NamePtr;
@@ -1267,24 +1263,24 @@ void FTechSoftFileParser::BuildReferenceName(TMap<FString, FString>& MetaData)
 
 	if (Format == ECADFormat::JT)
 	{
-		if (TechSoftFileParserImpl::ReplaceOrAddNameValue(MetaData, TEXT("SDKName")))
+		if (TechSoftUtils::ReplaceOrAddNameValue(MetaData, TEXT("SDKName")))
 		{
 			return;
 		}
 	}
 
-	if (TechSoftFileParserImpl::CheckIfNameExists(MetaData))
+	if (TechSoftUtils::CheckIfNameExists(MetaData))
 	{
 		return;
 	}
 
-	if (TechSoftFileParserImpl::ReplaceOrAddNameValue(MetaData, TEXT("PartNumber")))
+	if (TechSoftUtils::ReplaceOrAddNameValue(MetaData, TEXT("PartNumber")))
 	{
 		return;
 	}
 
 	NamePtr = MetaData.Find(TEXT("SDKName"));
-	if (NamePtr != nullptr)
+	if (NamePtr != nullptr && !NamePtr->IsEmpty())
 	{
 		FString SdkName = *NamePtr;
 
@@ -1308,12 +1304,9 @@ void FTechSoftFileParser::BuildReferenceName(TMap<FString, FString>& MetaData)
 	}
 }
 
-void FTechSoftFileParser::BuildInstanceName(TMap<FString, FString>& MetaData)
+void FTechSoftFileParser::BuildInstanceName(FEntityMetaData& InstanceData)
 {
-	if (MetaData.IsEmpty())
-	{
-		return;
-	}
+	TMap<FString, FString>& MetaData = InstanceData.MetaData;
 
 	if (TechSoftFileParserImpl::ReplaceOrAddNameValue(MetaData, TEXT("InstanceName")))
 	{
@@ -1326,7 +1319,7 @@ void FTechSoftFileParser::BuildInstanceName(TMap<FString, FString>& MetaData)
 	}
 
 	FString* NamePtr = MetaData.Find(TEXT("SDKName"));
-	if (NamePtr != nullptr)
+	if (NamePtr != nullptr && !NamePtr->IsEmpty())
 	{
 		FString SdkName = *NamePtr;
 
@@ -1355,10 +1348,16 @@ void FTechSoftFileParser::BuildInstanceName(TMap<FString, FString>& MetaData)
 		return;
 	}
 
+	FString& Name = MetaData.FindOrAdd(TEXT("Name"));
+	if (Name.IsEmpty())
+	{
+		Name = TEXT("Product");
+	}
 }
 
-void FTechSoftFileParser::BuildPartName(TMap<FString, FString>& MetaData)
+void FTechSoftFileParser::BuildPartName(FEntityMetaData& PartData)
 {
+	TMap<FString, FString>& MetaData = PartData.MetaData;
 	if (MetaData.IsEmpty())
 	{
 		return;
@@ -1380,14 +1379,11 @@ void FTechSoftFileParser::BuildPartName(TMap<FString, FString>& MetaData)
 	}
 }
 
-void FTechSoftFileParser::BuildBodyName(TMap<FString, FString>& MetaData)
+void FTechSoftFileParser::BuildBodyName(FEntityMetaData& BodyData, const FEntityMetaData& PartMetaData, bool bIsSolid)
 {
-	if (MetaData.IsEmpty())
-	{
-		return;
-	}
+	TMap<FString, FString>& MetaData = BodyData.MetaData;
 
-	if (TechSoftFileParserImpl::CheckIfNameExists(MetaData))
+	if (TechSoftUtils::CheckIfNameExists(MetaData))
 	{
 		return;
 	}
@@ -1406,7 +1402,38 @@ void FTechSoftFileParser::BuildBodyName(TMap<FString, FString>& MetaData)
 		return;
 	}
 
-	MetaData.Add(TEXT("Name"), TEXT("NoName"));
+	if (Format == ECADFormat::CATIA)
+	{
+		NamePtr = MetaData.Find(TEXT("BodyID"));
+		if (NamePtr != nullptr && !NamePtr->IsEmpty())
+		{
+			FString& Name = MetaData.FindOrAdd(TEXT("Name"));
+			Name = *NamePtr;
+			return;
+		}
+	}
+
+	const FString* PartName = PartMetaData.MetaData.Find(TEXT("Name"));
+
+	FString& Name = MetaData.FindOrAdd(TEXT("Name"));
+	if (Name.IsEmpty())
+	{
+		if (PartName != nullptr && !PartName->IsEmpty())
+		{
+			Name = *PartName;
+		}
+		else
+		{
+			if (bIsSolid)
+			{
+				Name = TEXT("Solid");
+			}
+			else
+			{
+				Name = TEXT("Shell");
+			}
+		}
+	}
 }
 
 void FTechSoftFileParser::ExtractSpecificMetaData(const A3DAsmProductOccurrence* Occurrence, FEntityMetaData& OutMetaData)
@@ -1631,33 +1658,7 @@ void FTechSoftFileParser::ExtractGraphStyleProperties(uint32 StyleIndex, FCADUUI
 	}
 }
 
-void FTechSoftFileParser::ExtractMaterialProperties(const A3DEntity* Entity)
-{
-	TUniqueTSObj<A3DMiscMaterialPropertiesData> MaterialPropertiesData(Entity);
-	if (!MaterialPropertiesData.IsValid())
-	{
-		return;
-	}
-
-	//xml->SetDoubleAttribute("m_dDensity", MaterialPropertiesData.m_dDensity);
-	//xml->SetAttribute("m_pcMaterialName", MaterialPropertiesData.m_pcMaterialName ? MaterialPropertiesData.m_pcMaterialName : "NULL");
-	//xml->SetAttribute("m_ePhysicType", MaterialPropertiesData.m_ePhysicType);
-	switch (MaterialPropertiesData->m_ePhysicType)
-	{
-	case A3DPhysicType_None:
-	case A3DPhysicType_Fiber:
-	case A3DPhysicType_HoneyComb:
-	case A3DPhysicType_Isotropic:
-	case A3DPhysicType_Orthotropic2D:
-	case A3DPhysicType_Orthotropic3D:
-	case A3DPhysicType_Anisotropic:
-
-	default:
-		break;
-	}
-}
-
-FMatrix FTechSoftFileParser::TraverseTransformation3D(const A3DMiscTransformation* CartesianTransformation)
+FMatrix FTechSoftFileParser::ExtractTransformation3D(const A3DMiscTransformation* CartesianTransformation)
 {
 	TUniqueTSObj<A3DMiscCartesianTransformationData> CartesianTransformationData(CartesianTransformation);
 
@@ -1688,7 +1689,7 @@ FMatrix FTechSoftFileParser::TraverseTransformation3D(const A3DMiscTransformatio
 	return FMatrix::Identity;
 }
 
-FMatrix FTechSoftFileParser::TraverseGeneralTransformation(const A3DMiscTransformation* GeneralTransformation)
+FMatrix FTechSoftFileParser::ExtractGeneralTransformation(const A3DMiscTransformation* GeneralTransformation)
 {
 	TUniqueTSObj<A3DMiscGeneralTransformationData> GeneralTransformationData(GeneralTransformation);
 	if (GeneralTransformationData.IsValid())
@@ -1713,7 +1714,7 @@ FMatrix FTechSoftFileParser::TraverseGeneralTransformation(const A3DMiscTransfor
 	return FMatrix::Identity;
 }
 
-FMatrix FTechSoftFileParser::TraverseTransformation(const A3DMiscTransformation* Transformation3D)
+FMatrix FTechSoftFileParser::ExtractTransformation(const A3DMiscTransformation* Transformation3D)
 {
 	if (Transformation3D == NULL)
 	{
@@ -1725,21 +1726,21 @@ FMatrix FTechSoftFileParser::TraverseTransformation(const A3DMiscTransformation*
 
 	if (Type == kA3DTypeMiscCartesianTransformation)
 	{
-		return TraverseTransformation3D(Transformation3D);
+		return ExtractTransformation3D(Transformation3D);
 	}
 	else if (Type == kA3DTypeMiscGeneralTransformation)
 	{
-		return TraverseGeneralTransformation(Transformation3D);
+		return ExtractGeneralTransformation(Transformation3D);
 	}
 	return FMatrix::Identity;
 }
 
-FMatrix FTechSoftFileParser::TraverseCoordinateSystem(const A3DRiCoordinateSystem* CoordinateSystem)
+FMatrix FTechSoftFileParser::ExtractCoordinateSystem(const A3DRiCoordinateSystem* CoordinateSystem)
 {
 	TUniqueTSObj<A3DRiCoordinateSystemData> CoordinateSystemData(CoordinateSystem);
 	if (CoordinateSystemData.IsValid())
 	{
-		return TraverseTransformation3D(CoordinateSystemData->m_pTransformation);
+		return ExtractTransformation3D(CoordinateSystemData->m_pTransformation);
 	}
 	return FMatrix::Identity;
 }
@@ -1832,7 +1833,7 @@ void ExtractTextureDefinition(const A3DGraphTextureDefinitionData& TextureDefini
 void FTechSoftFileParser::ReadMaterialsAndColors()
 {
 	A3DGlobal* GlobalPtr = TechSoftInterface::GetGlobalPointer();
-	if(GlobalPtr == nullptr)
+	if (GlobalPtr == nullptr)
 	{
 		return;
 	}
@@ -1856,7 +1857,7 @@ void FTechSoftFileParser::ReadMaterialsAndColors()
 			}
 		}
 	}
-	
+
 	{
 		uint32 PictureCount = GlobalData->m_uiPicturesSize;
 		if (PictureCount)
@@ -1878,6 +1879,6 @@ void FTechSoftFileParser::ReadMaterialsAndColors()
 	}
 }
 
-} // ns CADLibrary
+#endif  
 
-#endif  
\ No newline at end of file
+} // ns CADLibrary
diff --git a/Engine/Plugins/Enterprise/DatasmithCADImporter/Source/CADInterfaces/Private/TechSoftFileParser.h b/Engine/Plugins/Enterprise/DatasmithCADImporter/Source/CADInterfaces/Private/TechSoftFileParser.h
--- a/Engine/Plugins/Enterprise/DatasmithCADImporter/Source/CADInterfaces/Private/TechSoftFileParser.h
+++ b/Engine/Plugins/Enterprise/DatasmithCADImporter/Source/CADInterfaces/Private/TechSoftFileParser.h
@@ -6,25 +6,6 @@
 #include "CADData.h"
 #include "CADFileParser.h"
 #include "TechSoftInterface.h"
-#include "TUniqueTechSoftObj.h"
-
-typedef void A3DAsmModelFile;
-typedef void A3DAsmPartDefinition;
-typedef void A3DAsmProductOccurrence;
-typedef void A3DEntity;
-typedef void A3DGraphics;
-typedef void A3DMiscAttribute;
-typedef void A3DMiscCartesianTransformation;
-typedef void A3DMiscCartesianTransformation;
-typedef void A3DMiscGeneralTransformation;
-typedef void A3DRiBrepModel;
-typedef void A3DRiCoordinateSystem;
-typedef void A3DRiPolyBrepModel;
-typedef void A3DRiRepresentationItem;
-typedef void A3DRiSet;
-typedef void A3DTess3D;
-typedef void A3DTessBase;
-typedef void A3DTopoBrepData;
 
 namespace CADLibrary
 {
@@ -36,7 +17,6 @@ class FArchiveInstance;
 class FArchiveMaterial;
 class FArchiveUnloadedComponent;
 class FCADFileData;
-class FTechSoftInterface;
 
 struct FEntityMetaData;
 
@@ -68,8 +48,28 @@ public:
 #else
 	virtual ECADParsingResult Process() override;
 
-private:
+	double GetFileUnit()
+	{
+		return FileUnit;
+	}
+
+	void ExtractMetaData(const A3DEntity* Entity, FEntityMetaData& OutMetaData);
+
+protected:
+
+	virtual A3DStatus AdaptBRepModel()
+	{
+		return A3DStatus::A3D_SUCCESS;
+	}
+
+	virtual void SewModel();
 
+	// Tessellation methods
+	virtual void GenerateBodyMeshes();
+	virtual void GenerateBodyMesh(A3DRiRepresentationItem* Representation, FArchiveBody& Body);
+
+
+private:
 	// Methods to parse a first time the file to count the number of components
 	// Needed to reserve CADFileData
 	// Start with CountUnderModel
@@ -96,22 +96,19 @@ private:
 	FCadId TraverseOccurrence(const A3DAsmProductOccurrence* Occurrence);
 	void ProcessPrototype(const A3DAsmProductOccurrence* InPrototype, FEntityMetaData& OutMetaData, A3DMiscTransformation** OutLocation);
 	void TraversePartDefinition(const A3DAsmPartDefinition* PartDefinition, FArchiveComponent& Component);
-	FCadId TraverseRepresentationSet(const A3DRiSet* pSet, FEntityMetaData& PartMetaData);
-	FCadId TraverseRepresentationItem(A3DRiRepresentationItem* RepresentationItem, FEntityMetaData& PartMetaData);
-	FCadId TraverseBRepModel(A3DRiBrepModel* BrepModel, FEntityMetaData& PartMetaData);
-	FCadId TraversePolyBRepModel(A3DRiPolyBrepModel* PolygonalBrepModel, FEntityMetaData& PartMetaData);
+	FCadId TraverseRepresentationSet(const A3DRiSet* pSet, const FEntityMetaData& PartMetaData);
+	FCadId TraverseRepresentationItem(A3DRiRepresentationItem* RepresentationItem, const FEntityMetaData& PartMetaData);
+	FCadId TraverseBRepModel(A3DRiBrepModel* BrepModel, const FEntityMetaData& PartMetaData);
+	FCadId TraversePolyBRepModel(A3DRiPolyBrepModel* PolygonalBrepModel, const FEntityMetaData& PartMetaData);
 
-	// Tessellation methods
-	void GenerateBodyMeshes();
 
 	// MetaData
-	void ExtractMetaData(const A3DEntity* Entity, FEntityMetaData& OutMetaData);
 	void ExtractSpecificMetaData(const A3DAsmProductOccurrence* Occurrence, FEntityMetaData& OutMetaData);
 
-	void BuildInstanceName(TMap<FString, FString>& MetaData);
-	void BuildReferenceName(TMap<FString, FString>& MetaData);
-	void BuildPartName(TMap<FString, FString>& MetaData);
-	void BuildBodyName(TMap<FString, FString>& MetaData);
+	void BuildInstanceName(FEntityMetaData& MetaData);
+	void BuildReferenceName(FEntityMetaData& MetaData);
+	void BuildPartName(FEntityMetaData& MetaData);
+	void BuildBodyName(FEntityMetaData& MetaData, const FEntityMetaData& PartMetaData, bool bIsSolid);
 
 	// Graphic properties
 	void ExtractGraphicProperties(const A3DGraphics* Graphics, FEntityMetaData& OutMetaData);
@@ -121,7 +118,6 @@ private:
 	 * This method update the value ColorName or MaterialName accordingly of the GraphStyleData type (material or color)
 	 */
 	void ExtractGraphStyleProperties(uint32 StyleIndex, FCADUUID& ColorName, FCADUUID& MaterialName);
-	void ExtractMaterialProperties(const A3DEntity* Entity);
 	FArchiveColor& FindOrAddColor(uint32 ColorIndex, uint8 Alpha);
 	FArchiveMaterial& FindOrAddMaterial(uint32 MaterialId, const A3DGraphStyleData& GraphStyleData);
 
@@ -136,10 +132,10 @@ private:
 	}
 
 	// Transform
-	FMatrix TraverseCoordinateSystem(const A3DRiCoordinateSystem* CoordinateSystem);
-	FMatrix TraverseTransformation(const A3DMiscTransformation* Transformation3d);
-	FMatrix TraverseGeneralTransformation(const A3DMiscTransformation* GeneralTransformation);
-	FMatrix TraverseTransformation3D(const A3DMiscTransformation* CartesianTransformation);
+	FMatrix ExtractCoordinateSystem(const A3DRiCoordinateSystem* CoordinateSystem);
+	FMatrix ExtractTransformation(const A3DMiscTransformation* Transformation3d);
+	FMatrix ExtractGeneralTransformation(const A3DMiscTransformation* GeneralTransformation);
+	FMatrix ExtractTransformation3D(const A3DMiscTransformation* CartesianTransformation);
 
 	// Archive methods
 	FArchiveInstance& AddInstance(FEntityMetaData& InstanceMetaData);
@@ -150,7 +146,7 @@ private:
 	int32 AddBody(FEntityMetaData& BodyMetaData, const FMatrix& Matrix);
 #endif
 
-private:
+protected:
 
 	FUniqueTechSoftModelFile ModelFile;
 
@@ -159,14 +155,13 @@ private:
 
 	FCADFileData& CADFileData;
 	FTechSoftInterface& TechSoftInterface;
-
 	ECADFormat Format;
 
 	EModellerType ModellerType;
 	double FileUnit = 1;
+
 	FCadId LastEntityId = 1;
 
-	// 
 	TMap<A3DRiRepresentationItem*, int32> RepresentationItemsCache;
 };
 
diff --git a/Engine/Plugins/Enterprise/DatasmithCADImporter/Source/CADInterfaces/Private/TechSoftFileParserCADKernelTessellator.cpp b/Engine/Plugins/Enterprise/DatasmithCADImporter/Source/CADInterfaces/Private/TechSoftFileParserCADKernelTessellator.cpp
new file mode 100644
--- /dev/null
+++ b/Engine/Plugins/Enterprise/DatasmithCADImporter/Source/CADInterfaces/Private/TechSoftFileParserCADKernelTessellator.cpp
@@ -0,0 +1,168 @@
+// Copyright Epic Games, Inc. All Rights Reserved.
+
+#include "TechSoftFileParserCADKernelTessellator.h"
+
+
+#include "CADFileData.h"
+#include "CADKernelTools.h"
+#include "TechSoftUtils.h"
+#include "TechSoftBridge.h"
+#include "TUniqueTechSoftObj.h"
+
+#include "CADKernel/Core/Entity.h"
+#include "CADKernel/Core/Session.h"
+#include "CADKernel/Core/Types.h"
+
+#include "CADKernel/Mesh/Meshers/ParametricMesher.h"
+#include "CADKernel/Mesh/Structure/ModelMesh.h"
+
+#include "CADKernel/Topo/Body.h"
+#include "CADKernel/Topo/Model.h"
+#include "CADKernel/Topo/TopologicalShapeEntity.h"
+#include "CADKernel/Topo/Topomaker.h"
+
+namespace CADLibrary
+{
+
+FTechSoftFileParserCADKernelTessellator::FTechSoftFileParserCADKernelTessellator(FCADFileData& InCADData, const FString& EnginePluginsPath)
+	: FTechSoftFileParser(InCADData, EnginePluginsPath)
+	, LastHostIdUsed(1 << 30)
+{
+}
+
+#ifdef USE_TECHSOFT_SDK
+
+void FTechSoftFileParserCADKernelTessellator::GenerateBodyMesh(A3DRiRepresentationItem* Representation, FArchiveBody& ArchiveBody)
+{
+	double GeometricTolerance = CADFileData.GetImportParameters().ConvertMMToImportUnit(0.01); // mm
+
+	CADKernel::FSession CADKernelSession(GeometricTolerance);
+	CADKernelSession.SetFirstNewHostId(LastHostIdUsed);
+	CADKernel::FModel& CADKernelModel = CADKernelSession.GetModel();
+
+	CADKernel::FCADFileReport Report;
+	FTechSoftBridge TechSoftBridge(*this, CADKernelSession, Report);
+
+	CADKernel::FBody* CADKernelBody = TechSoftBridge.AddBody(Representation, ArchiveBody.MetaData);
+
+	if (CADFileData.GetImportParameters().GetStitchingTechnique() == StitchingHeal)
+	{
+		double Tolerance = CADFileData.GetImportParameters().ConvertMMToImportUnit(0.1); // mm
+		CADKernel::FTopomaker Topomaker(CADKernelSession, Tolerance);
+		Topomaker.Sew();
+		Topomaker.OrientShells();
+	}
+
+	MeshAndGetTessellation(CADKernelSession, ArchiveBody, *CADKernelBody);
+}
+
+void FTechSoftFileParserCADKernelTessellator::MeshAndGetTessellation(CADKernel::FSession& CADKernelSession, FArchiveBody& ArchiveBody, CADKernel::FBody& CADKernelBody)
+{
+	FBodyMesh& BodyMesh = CADFileData.AddBodyMesh(ArchiveBody.ObjectId, ArchiveBody);
+	ArchiveBody.ColorFaceSet = BodyMesh.ColorSet;
+	ArchiveBody.MaterialFaceSet = BodyMesh.MaterialSet;
+
+	// Save Body in CADKernelArchive file for re-tessellation
+	if (CADFileData.IsCacheDefined())
+	{
+		FString BodyFilePath = CADFileData.GetBodyCachePath(ArchiveBody.MeshActorName);
+		CADKernelSession.SaveDatabase(*BodyFilePath);
+	}
+
+	// Tessellate the body
+	TSharedRef<CADKernel::FModelMesh> CADKernelModelMesh = CADKernel::FEntity::MakeShared<CADKernel::FModelMesh>();
+
+	FCADKernelTools::DefineMeshCriteria(CADKernelModelMesh.Get(), CADFileData.GetImportParameters(), CADKernelSession.GetGeometricTolerance());
+
+	CADKernel::FParametricMesher Mesher(*CADKernelModelMesh);
+	Mesher.MeshEntity(CADKernelBody);
+
+	FCADKernelTools::GetBodyTessellation(*CADKernelModelMesh, CADKernelBody, BodyMesh);
+
+	ArchiveBody.ColorFaceSet = BodyMesh.ColorSet;
+	ArchiveBody.MaterialFaceSet = BodyMesh.MaterialSet;
+}
+
+A3DStatus FTechSoftFileParserCADKernelTessellator::AdaptBRepModel()
+{
+	const A3DUns32 ValidSurfaceCount = 9;
+	A3DUns32 AcceptedSurfaces[ValidSurfaceCount] = {
+		//kA3DTypeSurfBlend01,
+		//kA3DTypeSurfBlend02,
+		//kA3DTypeSurfBlend03,
+		kA3DTypeSurfNurbs,
+		kA3DTypeSurfCone,
+		kA3DTypeSurfCylinder,
+		kA3DTypeSurfCylindrical,
+		//kA3DTypeSurfOffset,
+		//kA3DTypeSurfPipe,
+		kA3DTypeSurfPlane,
+		kA3DTypeSurfRuled,
+		kA3DTypeSurfSphere,
+		kA3DTypeSurfRevolution,
+		//kA3DTypeSurfExtrusion,
+		//kA3DTypeSurfFromCurves,
+		kA3DTypeSurfTorus,
+		//kA3DTypeSurfTransform,
+	};
+
+	const A3DUns32 ValidCurveCount = 7;
+	A3DUns32 AcceptedCurves[ValidCurveCount] = {
+		//kA3DTypeCrvBase,
+		//kA3DTypeCrvBlend02Boundary,
+		kA3DTypeCrvNurbs,
+		kA3DTypeCrvCircle,
+		//kA3DTypeCrvComposite,
+		//kA3DTypeCrvOnSurf,
+		kA3DTypeCrvEllipse,
+		//kA3DTypeCrvEquation,
+		//kA3DTypeCrvHelix,
+		kA3DTypeCrvHyperbola,
+		//kA3DTypeCrvIntersection,
+		kA3DTypeCrvLine,
+		//kA3DTypeCrvOffset,
+		kA3DTypeCrvParabola,
+		kA3DTypeCrvPolyLine,
+		//kA3DTypeCrvTransform,
+	};
+
+	TUniqueTSObj<A3DCopyAndAdaptBrepModelData> CopyAndAdaptBrepModelData;
+	CopyAndAdaptBrepModelData->m_bUseSameParam = false;                        // If `A3D_TRUE`, surfaces will keep their parametrization when converted to NURBS.       
+	CopyAndAdaptBrepModelData->m_dTol = 1e-3;                                  // Tolerance value of resulting B-rep. The value is relative to the scale of the model.
+	CopyAndAdaptBrepModelData->m_bDeleteCrossingUV = false;                    // If `A3D_TRUE`, UV curves that cross seams of periodic surfaces are replaced by 3D curves 
+	CopyAndAdaptBrepModelData->m_bSplitFaces = false;                          // If `A3D_TRUE`, the faces with a periodic basis surface are split on parametric seams
+	CopyAndAdaptBrepModelData->m_bSplitClosedFaces = false;                    // If `A3D_TRUE`, the faces with a closed basis surface are split into faces at the parametric seam and mid-parameter
+	CopyAndAdaptBrepModelData->m_bForceComputeUV = true;                       // If `A3D_TRUE`, UV curves are computed from the B-rep data
+	CopyAndAdaptBrepModelData->m_bAllowUVCrossingSeams = true;                 // If `A3D_TRUE` and m_bForceComputeUV is set to `A3D_TRUE`, computed UV curves can cross seams.
+	CopyAndAdaptBrepModelData->m_bForceCompute3D = false;                      // If `A3D_TRUE`, 3D curves are computed from the B-rep data
+	CopyAndAdaptBrepModelData->m_bContinueOnError = true;                      // Continue processing even if an error occurs. Use \ref A3DCopyAndAdaptBrepModelAdvanced to get the error status.
+	CopyAndAdaptBrepModelData->m_bClampTolerantUVCurvesInsideUVDomain = false; // If `A3D_FALSE`, UV curves may stray outside the UV domain as long as the 3D edge tolerance is respected. If set to `A3D_TRUE`, the UV curves will be clamped to the UV domain (if the clamp still leaves them within the edge tolerance). */
+	CopyAndAdaptBrepModelData->m_bForceDuplicateGeometries = false;            // If `A3D_TRUE`, break the sharing of surfaces and curves into topologies.*/
+
+	CopyAndAdaptBrepModelData->m_uiAcceptableSurfacesSize = ValidSurfaceCount;
+	CopyAndAdaptBrepModelData->m_puiAcceptableSurfaces = &AcceptedSurfaces[0];
+	CopyAndAdaptBrepModelData->m_uiAcceptableCurvesSize = ValidCurveCount;
+	CopyAndAdaptBrepModelData->m_puiAcceptableCurves = &AcceptedCurves[0];
+
+	int32 ErrorCount = 0;
+	A3DCopyAndAdaptBrepModelErrorData* Errors = nullptr;
+	A3DStatus Ret = TechSoftInterface::AdaptBRepInModelFile(ModelFile.Get(), *CopyAndAdaptBrepModelData, ErrorCount, &Errors);
+	if ((Ret == A3D_SUCCESS || Ret == A3D_TOOLS_CONTINUE_ON_ERROR) && ErrorCount > 0)
+	{
+		// Add warning about error during the adaptation
+		CADFileData.AddWarningMessages(FString::Printf(TEXT("File %s had %d error during BRep adaptation step."), *CADFileData.GetCADFileDescription().GetFileName(), ErrorCount));
+		if (Errors != nullptr)
+		{
+			delete[] Errors;
+		}
+	}
+	else if (Ret != A3D_SUCCESS)
+	{
+		CADFileData.AddWarningMessages(FString::Printf(TEXT("File %s failed during BRep adaptation step."), *CADFileData.GetCADFileDescription().GetFileName(), ErrorCount));
+		return A3D_ERROR;
+	}
+	return A3D_SUCCESS;
+}
+#endif  
+
+} // ns CADLibrary
diff --git a/Engine/Plugins/Enterprise/DatasmithCADImporter/Source/CADInterfaces/Private/TechSoftFileParserCADKernelTessellator.h b/Engine/Plugins/Enterprise/DatasmithCADImporter/Source/CADInterfaces/Private/TechSoftFileParserCADKernelTessellator.h
new file mode 100644
--- /dev/null
+++ b/Engine/Plugins/Enterprise/DatasmithCADImporter/Source/CADInterfaces/Private/TechSoftFileParserCADKernelTessellator.h
@@ -0,0 +1,51 @@
+// Copyright Epic Games, Inc. All Rights Reserved.
+#pragma once
+
+#include "CoreMinimal.h"
+
+#include "TechSoftFileParser.h"
+
+namespace CADKernel
+{
+class FBody;
+class FSession;
+}
+
+namespace CADLibrary
+{
+
+class CADINTERFACES_API FTechSoftFileParserCADKernelTessellator : public FTechSoftFileParser
+{
+public:
+
+	/**
+	 * @param InCADData TODO
+	 * @param EnginePluginsPath Full Path of EnginePlugins. Mandatory to set KernelIO to import DWG, or DGN files
+	 */
+	FTechSoftFileParserCADKernelTessellator(FCADFileData& InCADData, const FString& EnginePluginsPath = TEXT(""));
+
+#ifdef USE_TECHSOFT_SDK
+
+private:
+
+	virtual A3DStatus AdaptBRepModel() override;
+
+	virtual void SewModel() override
+	{
+		// Do in GenerateBodyMeshes
+	}
+
+	// Tessellation methods
+	virtual void GenerateBodyMesh(A3DRiRepresentationItem* Representation, FArchiveBody& Body) override;
+	void MeshAndGetTessellation(CADKernel::FSession& CADKernelSession, FArchiveBody& ArchiveBody, CADKernel::FBody& CADKernelBody);
+
+#endif
+
+private:
+
+	FCadId LastEntityId = 1;
+	int32 LastHostIdUsed = 0;
+
+};
+
+} // ns CADLibrary
\ No newline at end of file
diff --git a/Engine/Plugins/Enterprise/DatasmithCADImporter/Source/CADInterfaces/Private/TechSoftInterface.cpp b/Engine/Plugins/Enterprise/DatasmithCADImporter/Source/CADInterfaces/Private/TechSoftInterface.cpp
--- a/Engine/Plugins/Enterprise/DatasmithCADImporter/Source/CADInterfaces/Private/TechSoftInterface.cpp
+++ b/Engine/Plugins/Enterprise/DatasmithCADImporter/Source/CADInterfaces/Private/TechSoftInterface.cpp
@@ -40,7 +40,7 @@ bool FTechSoftInterface::InitializeKernel(const TCHAR* InEnginePluginsPath)
 #ifdef CADKERNEL_DEV
 	FString TechSoftDllPath = EnginePluginsPath;
 #else
-	FString TechSoftDllPath = FPaths::Combine(EnginePluginsPath, TEXT("Enterprise/DatasmithCADImporter"), TEXT("Binaries"), FPlatformProcess::GetBinariesSubdirectory(), "TechSoft");
+	FString TechSoftDllPath = FPaths::Combine(EnginePluginsPath, TEXT("Enterprise/DatasmithCADImporter"), TEXT("Binaries"), FPlatformProcess::GetBinariesSubdirectory(), TEXT("TechSoft"));
 #endif
 
 	TechSoftDllPath = FPaths::ConvertRelativePathToFull(TechSoftDllPath);
