Fix Crash when importing "190204 NEWCASTLE EAST END - Kevin Lee.ifc"

diff --git a/Engine/Plugins/Enterprise/DatasmithCADImporter/Source/CADInterfaces/Private/TechSoftInterfaceUtils.inl b/Engine/Plugins/Enterprise/DatasmithCADImporter/Source/CADInterfaces/Private/TechSoftInterfaceUtils.inl
--- a/Engine/Plugins/Enterprise/DatasmithCADImporter/Source/CADInterfaces/Private/TechSoftInterfaceUtils.inl
+++ b/Engine/Plugins/Enterprise/DatasmithCADImporter/Source/CADInterfaces/Private/TechSoftInterfaceUtils.inl
@@ -239,12 +239,16 @@ namespace TechSoftInterfaceUtils
 					A3DUns32 StripeSize = FaceTessData.m_puiSizesTriangulated[FaceSetIndex++] & kA3DTessFaceDataNormalMask;
 					for (A3DUns32 StripeIndex = 0; StripeIndex < StripeSize; ++StripeIndex)
 					{
-						// #ueent_techsoft: There are some specificities with the case of oneNormal
-						//					See TessConnector.cpp around line 237
-						//	bool bIsOneNormal = (pFaceTessData->m_puiSizesTriangulated[uiCurrentSize] & kA3DTessFaceDataNormalSingle) != 0;????
-						//	A3DUns32* pStripeNormalIndice = puiTriangulatedIndexes;
-						//	A3DUns32 uiNbPoint = pFaceTessData->m_puiSizesTriangulated[uiCurrentSize++] & kA3DTessFaceDataNormalMask;
+						bool bIsOneNormal = (FaceTessData.m_puiSizesTriangulated[FaceSetIndex] & kA3DTessFaceDataNormalSingle) != 0;
 						A3DUns32 PointCount = FaceTessData.m_puiSizesTriangulated[FaceSetIndex++] & kA3DTessFaceDataNormalMask;
+
+						// Is there only one normal for the entire stripe?
+						if (bIsOneNormal == false)
+						{
+							AddFaceTriangleStripe(Tessellation, PointCount, LastTrianguleIndex, LastVertexIndex);
+							continue;
+						}
+
 						AddFaceTriangleStripeWithUniqueNormal(Tessellation, PointCount, LastTrianguleIndex, LastVertexIndex);
 					}
 					bMustProcess = FaceTessData.m_uiSizesTriangulatedSize > FaceSetIndex;
@@ -606,7 +610,7 @@ namespace TechSoftInterfaceUtils
 			InOutLastTriangleIndex += TextureCount;
 			FaceIndex[1] = TriangulatedIndexes[InOutLastTriangleIndex++] / 3;
 
-			for (unsigned long TriangleIndex = 0; TriangleIndex < InTriangleCount; TriangleIndex++)
+			for (unsigned long TriangleIndex = 2; TriangleIndex < InTriangleCount; TriangleIndex++)
 			{
 				NormalIndex[2] = TriangulatedIndexes[InOutLastTriangleIndex++];
 				TextureIndex[2] = TriangulatedIndexes[InOutLastTriangleIndex];
@@ -655,7 +659,7 @@ namespace TechSoftInterfaceUtils
 			FaceIndex[0] = TriangulatedIndexes[InOutLastTriangleIndex++] / 3;
 			FaceIndex[1] = TriangulatedIndexes[InOutLastTriangleIndex++] / 3;
 
-			for (unsigned long TriangleIndex = 0; TriangleIndex < InTriangleCount; TriangleIndex++)
+			for (unsigned long TriangleIndex = 2; TriangleIndex < InTriangleCount; TriangleIndex++)
 			{
 				FaceIndex[2] = TriangulatedIndexes[InOutLastTriangleIndex++] / 3;
 
@@ -702,7 +706,7 @@ namespace TechSoftInterfaceUtils
 			InOutLastTriangleIndex += TextureCount;
 			FaceIndex[1] = TriangulatedIndexes[InOutLastTriangleIndex++] / 3;
 
-			for (unsigned long TriangleIndex = 0; TriangleIndex < InTriangleCount; TriangleIndex++)
+			for (unsigned long TriangleIndex = 2; TriangleIndex < InTriangleCount; TriangleIndex++)
 			{
 				TextureIndex[2] = TriangulatedIndexes[InOutLastTriangleIndex];
 				InOutLastTriangleIndex += TextureCount;
