Fix Compile UnrealEditor Win64 - 'A3DRiRepresentationItem': undeclared identifier

diff --git a/Engine/Plugins/Enterprise/DatasmithCADImporter/Source/CADInterfaces/Private/TechSoftFileParser.h b/Engine/Plugins/Enterprise/DatasmithCADImporter/Source/CADInterfaces/Private/TechSoftFileParser.h
--- a/Engine/Plugins/Enterprise/DatasmithCADImporter/Source/CADInterfaces/Private/TechSoftFileParser.h
+++ b/Engine/Plugins/Enterprise/DatasmithCADImporter/Source/CADInterfaces/Private/TechSoftFileParser.h
@@ -7,6 +7,9 @@
 #include "CADFileParser.h"
 #include "TechSoftInterface.h"
 
+typedef void A3DAsmProductOccurrence;
+typedef void A3DRiRepresentationItem;
+
 namespace CADLibrary
 {
 
