Implement Stitching Technique = Sew when using CADKernel tessellator

diff --git a/Engine/Plugins/Enterprise/DatasmithCADImporter/Source/CADInterfaces/Private/CADFileReader.cpp b/Engine/Plugins/Enterprise/DatasmithCADImporter/Source/CADInterfaces/Private/CADFileReader.cpp
--- a/Engine/Plugins/Enterprise/DatasmithCADImporter/Source/CADInterfaces/Private/CADFileReader.cpp
+++ b/Engine/Plugins/Enterprise/DatasmithCADImporter/Source/CADInterfaces/Private/CADFileReader.cpp
@@ -13,6 +13,7 @@
 
 #ifdef USE_TECHSOFT_SDK
 #include "TechSoftFileParser.h"
+#include "TechSoftFileParserCADKernelTessellator.h"
 #endif
 
 namespace CADLibrary
@@ -32,7 +33,14 @@ namespace CADLibrary
 #ifdef USE_TECHSOFT_SDK
 		if (FImportParameters::GCADLibrary == TEXT("TechSoft"))
 		{
-			CADParser = MakeUnique<FTechSoftFileParser>(CADFileData, EnginePluginsPath);
+			if (FImportParameters::bGDisableCADKernelTessellation)
+			{
+				CADParser = MakeUnique<FTechSoftFileParser>(CADFileData, EnginePluginsPath);
+			}
+			else
+			{
+				CADParser = MakeUnique<FTechSoftFileParserCADKernelTessellator>(CADFileData, EnginePluginsPath);
+			}
 		}
 #endif
 	}
diff --git a/Engine/Plugins/Enterprise/DatasmithCADImporter/Source/CADInterfaces/Private/CADSceneGraph.cpp b/Engine/Plugins/Enterprise/DatasmithCADImporter/Source/CADInterfaces/Private/CADSceneGraph.cpp
--- a/Engine/Plugins/Enterprise/DatasmithCADImporter/Source/CADInterfaces/Private/CADSceneGraph.cpp
+++ b/Engine/Plugins/Enterprise/DatasmithCADImporter/Source/CADInterfaces/Private/CADSceneGraph.cpp
@@ -44,6 +44,7 @@ FArchive& operator<<(FArchive& Ar, FArchiveBody& Body)
 	Ar << (FCADArchiveObject&) Body;
 	Ar << Body.MaterialFaceSet;
 	Ar << Body.ColorFaceSet;
+	Ar << Body.ParentId;
 	Ar << Body.MeshActorName;
 	return Ar;
 }
diff --git a/Engine/Plugins/Enterprise/DatasmithCADImporter/Source/CADInterfaces/Private/TechSoftFileParser.cpp b/Engine/Plugins/Enterprise/DatasmithCADImporter/Source/CADInterfaces/Private/TechSoftFileParser.cpp
--- a/Engine/Plugins/Enterprise/DatasmithCADImporter/Source/CADInterfaces/Private/TechSoftFileParser.cpp
+++ b/Engine/Plugins/Enterprise/DatasmithCADImporter/Source/CADInterfaces/Private/TechSoftFileParser.cpp
@@ -215,7 +215,6 @@ ECADParsingResult FTechSoftFileParser::Process()
 
 	// Add specific options according to format
 	Format = File.GetFileFormat();
-	bool bForceSew = false;
 	TechSoftFileParserImpl::UpdateIOOptionAccordingToFormat(Format, Import, bForceSew);
 
 	A3DStatus LoadStatus = A3DStatus::A3D_SUCCESS;
@@ -690,11 +689,12 @@ FArchiveComponent& FTechSoftFileParser::AddOccurence(FEntityMetaData& InstanceMe
 	return Prototype;
 }
 
-int32 FTechSoftFileParser::AddBody(FEntityMetaData& BodyMetaData, const FMatrix& Matrix)
+int32 FTechSoftFileParser::AddBody(FEntityMetaData& BodyMetaData, const FMatrix& Matrix, const FCadId ParentId)
 {
 	FCadId BodyId = LastEntityId++;
 	int32 BodyIndex = CADFileData.AddBody(BodyId);
 	FArchiveBody& Body = CADFileData.GetBodyAt(BodyIndex);
+	Body.ParentId = ParentId;
 	Body.TransformMatrix = Matrix;
 	Body.MetaData = MoveTemp(BodyMetaData.MetaData);
 	if (BodyMetaData.ColorName != 0)
@@ -1023,7 +1023,7 @@ void FTechSoftFileParser::TraversePartDefinition(const A3DAsmPartDefinition* Par
 	{
 		for (unsigned int Index = 0; Index < PartData->m_uiRepItemsSize; ++Index)
 		{
-			int32 ChildId = TraverseRepresentationItem(PartData->m_ppRepItems[Index], PartMetaData);
+			int32 ChildId = TraverseRepresentationItem(PartData->m_ppRepItems[Index], PartMetaData, Part.ObjectId);
 			Part.Children.Add(ChildId);
 		}
 	}
@@ -1044,7 +1044,7 @@ void FTechSoftFileParser::CountUnderPartDefinition(const A3DAsmPartDefinition* P
 	}
 }
 
-FCadId FTechSoftFileParser::TraverseRepresentationItem(A3DRiRepresentationItem* RepresentationItem, const FEntityMetaData& PartMetaData)
+FCadId FTechSoftFileParser::TraverseRepresentationItem(A3DRiRepresentationItem* RepresentationItem, const FEntityMetaData& PartMetaData, const FCadId ParentId)
 {
 	A3DEEntityType Type;
 	A3DEntityGetType(RepresentationItem, &Type);
@@ -1054,9 +1054,9 @@ FCadId FTechSoftFileParser::TraverseRepresentationItem(A3DRiRepresentationItem*
 	case kA3DTypeRiSet:
 		return TraverseRepresentationSet(RepresentationItem, PartMetaData);
 	case kA3DTypeRiBrepModel:
-		return TraverseBRepModel(RepresentationItem, PartMetaData);
+		return TraverseBRepModel(RepresentationItem, PartMetaData, ParentId);
 	case kA3DTypeRiPolyBrepModel:
-		return TraversePolyBRepModel(RepresentationItem, PartMetaData);
+		return TraversePolyBRepModel(RepresentationItem, PartMetaData, ParentId);
 	default:
 		break;
 	}
@@ -1104,7 +1104,7 @@ FCadId FTechSoftFileParser::TraverseRepresentationSet(const A3DRiSet* Representa
 
 	for (A3DUns32 Index = 0; Index < RepresentationSetData->m_uiRepItemsSize; ++Index)
 	{
-		int32 ChildId = TraverseRepresentationItem(RepresentationSetData->m_ppRepItems[Index], RepresentationSetMetaData);
+		int32 ChildId = TraverseRepresentationItem(RepresentationSetData->m_ppRepItems[Index], RepresentationSetMetaData, RepresentationSet.ObjectId);
 		RepresentationSet.Children.Add(ChildId);
 	}
 	return RepresentationSetId;
@@ -1125,7 +1125,7 @@ void FTechSoftFileParser::CountUnderRepresentationSet(const A3DRiSet* Representa
 	}
 }
 
-FCadId FTechSoftFileParser::TraverseBRepModel(A3DRiBrepModel* BRepModelPtr, const FEntityMetaData& PartMetaData)
+FCadId FTechSoftFileParser::TraverseBRepModel(A3DRiBrepModel* BRepModelPtr, const FEntityMetaData& PartMetaData, const FCadId ParentId)
 {
 	if (!BRepModelPtr)
 	{
@@ -1157,7 +1157,7 @@ FCadId FTechSoftFileParser::TraverseBRepModel(A3DRiBrepModel* BRepModelPtr, cons
 		Matrix = ExtractCoordinateSystem(RepresentationData->m_pCoordinateSystem);
 	}
 
-	int32 BodyIndex = AddBody(BRepMetaData, Matrix);
+	int32 BodyIndex = AddBody(BRepMetaData, Matrix, ParentId);
 	FArchiveBody& Body = CADFileData.GetBodyAt(BodyIndex);
 
 	RepresentationItemsCache.Add(BRepModelPtr, BodyIndex);
@@ -1165,7 +1165,7 @@ FCadId FTechSoftFileParser::TraverseBRepModel(A3DRiBrepModel* BRepModelPtr, cons
 	return Body.ObjectId;
 }
 
-FCadId FTechSoftFileParser::TraversePolyBRepModel(A3DRiPolyBrepModel* PolygonalPtr, const FEntityMetaData& PartMetaData)
+FCadId FTechSoftFileParser::TraversePolyBRepModel(A3DRiPolyBrepModel* PolygonalPtr, const FEntityMetaData& PartMetaData, const FCadId ParentId)
 {
 	if (!PolygonalPtr)
 	{
@@ -1197,7 +1197,7 @@ FCadId FTechSoftFileParser::TraversePolyBRepModel(A3DRiPolyBrepModel* PolygonalP
 		Matrix = ExtractCoordinateSystem(RepresentationData->m_pCoordinateSystem);
 	}
 
-	int32 BodyIndex = AddBody(BRepMetaData, Matrix);
+	int32 BodyIndex = AddBody(BRepMetaData, Matrix, ParentId);
 	FArchiveBody& Body = CADFileData.GetBodyAt(BodyIndex);
 
 	RepresentationItemsCache.Add(PolygonalPtr, BodyIndex);
diff --git a/Engine/Plugins/Enterprise/DatasmithCADImporter/Source/CADInterfaces/Private/TechSoftFileParser.h b/Engine/Plugins/Enterprise/DatasmithCADImporter/Source/CADInterfaces/Private/TechSoftFileParser.h
--- a/Engine/Plugins/Enterprise/DatasmithCADImporter/Source/CADInterfaces/Private/TechSoftFileParser.h
+++ b/Engine/Plugins/Enterprise/DatasmithCADImporter/Source/CADInterfaces/Private/TechSoftFileParser.h
@@ -65,6 +65,10 @@ protected:
 		return A3DStatus::A3D_SUCCESS;
 	}
 
+	/**
+	 * If the tessellator is TechSoft, SewModel call TechSoftInterface::SewModel
+	 * If the tessellator is CADKernel, SewModel do nothing as the file is not yet parsed. In this case, the sew is done in GenerateBodyMeshes.
+	 */
 	virtual void SewModel();
 
 	// Tessellation methods
@@ -100,10 +104,9 @@ private:
 	void ProcessPrototype(const A3DAsmProductOccurrence* InPrototype, FEntityMetaData& OutMetaData, A3DMiscTransformation** OutLocation);
 	void TraversePartDefinition(const A3DAsmPartDefinition* PartDefinition, FArchiveComponent& Component);
 	FCadId TraverseRepresentationSet(const A3DRiSet* pSet, const FEntityMetaData& PartMetaData);
-	FCadId TraverseRepresentationItem(A3DRiRepresentationItem* RepresentationItem, const FEntityMetaData& PartMetaData);
-	FCadId TraverseBRepModel(A3DRiBrepModel* BrepModel, const FEntityMetaData& PartMetaData);
-	FCadId TraversePolyBRepModel(A3DRiPolyBrepModel* PolygonalBrepModel, const FEntityMetaData& PartMetaData);
-
+	FCadId TraverseRepresentationItem(A3DRiRepresentationItem* RepresentationItem, const FEntityMetaData& PartMetaData, const FCadId ParentId);
+	FCadId TraverseBRepModel(A3DRiBrepModel* BrepModel, const FEntityMetaData& PartMetaData, const FCadId ParentId);
+	FCadId TraversePolyBRepModel(A3DRiPolyBrepModel* PolygonalBrepModel, const FEntityMetaData& PartMetaData, const FCadId ParentId);
 
 	// MetaData
 	void ExtractSpecificMetaData(const A3DAsmProductOccurrence* Occurrence, FEntityMetaData& OutMetaData);
@@ -146,7 +149,7 @@ private:
 	FArchiveUnloadedComponent& AddUnloadedComponent(FEntityMetaData& ComponentMetaData, FArchiveInstance& Instance);
 	FArchiveComponent& AddOccurence(FEntityMetaData& InstanceMetaData, FEntityMetaData& ReferenceMetaData, FCadId& OutComponentId);
 	FArchiveComponent& AddOccurence(FEntityMetaData& InstanceMetaData, FCadId& OutComponentId);
-	int32 AddBody(FEntityMetaData& BodyMetaData, const FMatrix& Matrix);
+	int32 AddBody(FEntityMetaData& BodyMetaData, const FMatrix& Matrix, const FCadId ParentId);
 #endif
 
 protected:
@@ -159,6 +162,7 @@ protected:
 	FCADFileData& CADFileData;
 	FTechSoftInterface& TechSoftInterface;
 	ECADFormat Format;
+	bool bForceSew = false;
 
 	EModellerType ModellerType;
 	double FileUnit = 1;
diff --git a/Engine/Plugins/Enterprise/DatasmithCADImporter/Source/CADInterfaces/Private/TechSoftFileParserCADKernelTessellator.cpp b/Engine/Plugins/Enterprise/DatasmithCADImporter/Source/CADInterfaces/Private/TechSoftFileParserCADKernelTessellator.cpp
--- a/Engine/Plugins/Enterprise/DatasmithCADImporter/Source/CADInterfaces/Private/TechSoftFileParserCADKernelTessellator.cpp
+++ b/Engine/Plugins/Enterprise/DatasmithCADImporter/Source/CADInterfaces/Private/TechSoftFileParserCADKernelTessellator.cpp
@@ -32,6 +32,159 @@ FTechSoftFileParserCADKernelTessellator::FTechSoftFileParserCADKernelTessellator
 
 #ifdef USE_TECHSOFT_SDK
 
+void FTechSoftFileParserCADKernelTessellator::GenerateBodyMeshes()
+{
+	if (bForceSew || CADFileData.GetImportParameters().GetStitchingTechnique() == EStitchingTechnique::StitchingSew)
+	{
+		SewAndGenerateBodyMeshes();
+	}
+	else
+	{
+		// If no sew is required, FTechSoftFileParser::GenerateBodyMeshes is called to process all the bodies one by one with the override GenerateBodyMesh
+		FTechSoftFileParser::GenerateBodyMeshes();
+	}
+}
+
+void FTechSoftFileParserCADKernelTessellator::SewAndGenerateBodyMeshes()
+{
+	int32 RepresentationCount = RepresentationItemsCache.Num();
+	TMap<FCadId, TArray<A3DRiRepresentationItem*>> OccurenceIdToRepresentations;
+	for (TPair<A3DRiRepresentationItem*, int32>& Entry : RepresentationItemsCache)
+	{
+		A3DRiRepresentationItem* RepresentationItemPtr = Entry.Key;
+		FArchiveBody& Body = CADFileData.GetBodyAt(Entry.Value);
+
+		TArray<A3DRiRepresentationItem*>& Representations = OccurenceIdToRepresentations.FindOrAdd(Body.ParentId);
+		if (Representations.Max() == 0)
+		{
+			Representations.Reserve(RepresentationCount);
+		}
+
+		Representations.Add(RepresentationItemPtr);
+	}
+
+	for (TPair<FCadId, TArray<A3DRiRepresentationItem*>>& Representations : OccurenceIdToRepresentations)
+	{
+		SewAndMesh(Representations.Value);
+	}
+
+}
+
+void FTechSoftFileParserCADKernelTessellator::SewAndMesh(TArray<A3DRiRepresentationItem*>& Representations)
+{
+	double GeometricTolerance = CADFileData.GetImportParameters().ConvertMMToImportUnit(0.01); // mm
+
+	CADKernel::FSession CADKernelSession(GeometricTolerance);
+	CADKernelSession.SetFirstNewHostId(LastHostIdUsed);
+	CADKernel::FModel& CADKernelModel = CADKernelSession.GetModel();
+
+	CADKernel::FCADFileReport Report;
+	FTechSoftBridge TechSoftBridge(*this, CADKernelSession, Report);
+
+	for (A3DRiRepresentationItem* Representation : Representations)
+	{
+		int32* BodyIndex = RepresentationItemsCache.Find(Representation);
+		if (BodyIndex == nullptr)
+		{
+			continue;
+		}
+		FArchiveBody& Body = CADFileData.GetBodyAt(*BodyIndex);
+
+		CADKernel::FBody* CADKernelBody = TechSoftBridge.AddBody(Representation, Body.MetaData);
+	}
+
+	// Sew if needed
+	CADKernel::FTopomaker Topomaker(CADKernelSession, GeometricTolerance);
+	Topomaker.Sew();
+	Topomaker.SplitIntoConnectedShells();
+	Topomaker.OrientShells();
+
+	// The Sew + SplitIntoConnectedShells change the bodies: some are deleted some are create
+	// but at the end the count of body is always <= than the initial count
+	// We need to found the unchanged bodies to link them to their FArchiveBody
+	// The new bodies will be linked to FArchiveBody of deleted bodies but the metadata of these FArchiveBody havbe to be cleaned
+
+	int32 BodyCount = CADKernelModel.GetBodies().Num();
+	int32 DeletedBodyCount = Representations.Num() - BodyCount;
+
+	TArray<CADKernel::FBody*> NewBodies;
+	TArray<CADKernel::FBody*> ExistingBodies;
+	NewBodies.Reserve(BodyCount);
+	ExistingBodies.Reserve(BodyCount);
+
+	// find new and existing bodies
+	for (const TSharedPtr<CADKernel::FBody>& CADKernelBody : CADKernelModel.GetBodies())
+	{
+		if (!CADKernelBody.IsValid())
+		{
+			continue;
+		}
+
+		const A3DRiRepresentationItem* Representation = TechSoftBridge.GetA3DBody(CADKernelBody.Get());
+		if (Representation == nullptr)
+		{
+			NewBodies.Add(CADKernelBody.Get());
+		}
+		else
+		{
+			ExistingBodies.Add(CADKernelBody.Get());
+		}
+	}
+
+	// find Representation of deleted bodies to find unused FArchiveBody
+	//int32 DeletedBodyCount = Representations.Num() - BodyCount;
+	TArray<int32> ArchiveBodyIndexOfDeletedRepresentation;
+	ArchiveBodyIndexOfDeletedRepresentation.Reserve(DeletedBodyCount);
+
+	for (A3DRiRepresentationItem* Representation : Representations)
+	{
+		CADKernel::FBody* Body = TechSoftBridge.GetBody(Representation);
+		if (Body == nullptr)
+		{
+			int32* ArchiveBodyIndex = RepresentationItemsCache.Find(Representation);
+			if (ArchiveBodyIndex == nullptr)
+			{
+				continue; // should not append
+			}
+			ArchiveBodyIndexOfDeletedRepresentation.Add(*ArchiveBodyIndex);
+		}
+	}
+
+	// Process existing bodies
+	for (CADKernel::FBody* Body : ExistingBodies)
+	{
+		const A3DRiRepresentationItem* Representation = TechSoftBridge.GetA3DBody(Body);
+
+		int32* ArchiveBodyIndex = RepresentationItemsCache.Find(Representation);
+		if (ArchiveBodyIndex == nullptr)
+		{
+			continue; // should not append
+		}
+
+		FArchiveBody& ArchiveBody = CADFileData.GetBodyAt(*ArchiveBodyIndex);
+		MeshAndGetTessellation(CADKernelSession, ArchiveBody, *Body);
+	}
+
+	// Process new bodies
+	int32 Index = 0;
+	for (CADKernel::FBody* Body : NewBodies)
+	{
+		int32 ArchiveBodyIndex = ArchiveBodyIndexOfDeletedRepresentation[Index++];
+		FArchiveBody& ArchiveBody = CADFileData.GetBodyAt(ArchiveBodyIndex);
+		MeshAndGetTessellation(CADKernelSession, ArchiveBody, *Body);
+	}
+
+	// Delete unused FArchiveBody
+	for (; Index < ArchiveBodyIndexOfDeletedRepresentation.Num(); ++Index)
+	{
+		int32 ArchiveBodyIndex = ArchiveBodyIndexOfDeletedRepresentation[Index++];
+		FArchiveBody& ArchiveBody = CADFileData.GetBodyAt(ArchiveBodyIndex);
+		ArchiveBody.MetaData.Empty();
+		ArchiveBody.ParentId = 0;
+		ArchiveBody.MeshActorName = 0;
+	}
+}
+
 void FTechSoftFileParserCADKernelTessellator::GenerateBodyMesh(A3DRiRepresentationItem* Representation, FArchiveBody& ArchiveBody)
 {
 	double GeometricTolerance = CADFileData.GetImportParameters().ConvertMMToImportUnit(0.01); // mm
diff --git a/Engine/Plugins/Enterprise/DatasmithCADImporter/Source/CADInterfaces/Private/TechSoftFileParserCADKernelTessellator.h b/Engine/Plugins/Enterprise/DatasmithCADImporter/Source/CADInterfaces/Private/TechSoftFileParserCADKernelTessellator.h
--- a/Engine/Plugins/Enterprise/DatasmithCADImporter/Source/CADInterfaces/Private/TechSoftFileParserCADKernelTessellator.h
+++ b/Engine/Plugins/Enterprise/DatasmithCADImporter/Source/CADInterfaces/Private/TechSoftFileParserCADKernelTessellator.h
@@ -36,9 +36,13 @@ private:
 	}
 
 	// Tessellation methods
+	virtual void GenerateBodyMeshes() override;
 	virtual void GenerateBodyMesh(A3DRiRepresentationItem* Representation, FArchiveBody& Body) override;
 	void MeshAndGetTessellation(CADKernel::FSession& CADKernelSession, FArchiveBody& ArchiveBody, CADKernel::FBody& CADKernelBody);
 
+	void SewAndGenerateBodyMeshes();
+	void SewAndMesh(TArray<A3DRiRepresentationItem*>& Representations);
+
 #endif
 
 private:
diff --git a/Engine/Plugins/Enterprise/DatasmithCADImporter/Source/CADInterfaces/Public/CADSceneGraph.h b/Engine/Plugins/Enterprise/DatasmithCADImporter/Source/CADInterfaces/Public/CADSceneGraph.h
--- a/Engine/Plugins/Enterprise/DatasmithCADImporter/Source/CADInterfaces/Public/CADSceneGraph.h
+++ b/Engine/Plugins/Enterprise/DatasmithCADImporter/Source/CADInterfaces/Public/CADSceneGraph.h
@@ -80,7 +80,8 @@ public:
 	friend FArchive& operator<<(FArchive& Ar, FArchiveBody& C);
 
 public:
-	FCADUUID MeshActorName;
+	FCadId ParentId = 0;
+	FCADUUID MeshActorName = 0;
 
 	TSet<FMaterialId> MaterialFaceSet;
 	TSet<FColorId> ColorFaceSet;
diff --git a/Engine/Plugins/Enterprise/DatasmithCADImporter/Source/DatasmithCADTranslator/Private/DatasmithSceneGraphBuilder.cpp b/Engine/Plugins/Enterprise/DatasmithCADImporter/Source/DatasmithCADTranslator/Private/DatasmithSceneGraphBuilder.cpp
--- a/Engine/Plugins/Enterprise/DatasmithCADImporter/Source/DatasmithCADTranslator/Private/DatasmithSceneGraphBuilder.cpp
+++ b/Engine/Plugins/Enterprise/DatasmithCADImporter/Source/DatasmithCADTranslator/Private/DatasmithSceneGraphBuilder.cpp
@@ -478,6 +478,11 @@ TSharedPtr<IDatasmithActorElement> FDatasmithSceneBaseGraphBuilder::BuildBody(in
 
 	CADLibrary::FArchiveBody& Body = SceneGraph->Bodies[BodyIndex];
 
+	if (Body.ParentId == 0 || Body.MeshActorName == 0)
+	{
+		return TSharedPtr<IDatasmithActorElement>();
+	}
+
 	FString BodyUUID;
 	FString BodyLabel;
 	GetNodeUUIDAndName(InstanceNodeMetaDataMap, Body.MetaData, BodyIndex, ParentData.Uuid, BodyUUID, BodyLabel);
diff --git a/Engine/Source/Runtime/Datasmith/CADKernel/Private/CADKernel/Math/BSpline.cpp b/Engine/Source/Runtime/Datasmith/CADKernel/Private/CADKernel/Math/BSpline.cpp
--- a/Engine/Source/Runtime/Datasmith/CADKernel/Private/CADKernel/Math/BSpline.cpp
+++ b/Engine/Source/Runtime/Datasmith/CADKernel/Private/CADKernel/Math/BSpline.cpp
@@ -815,6 +815,10 @@ namespace CADKernel
 
 		TSharedPtr<FNURBSCurve> DuplicateNurbsCurveWithHigherDegree(int32 Degre, const FNURBSCurve& InCurve)
 		{
+			// To avoid crashes while waiting for the fix (jira UETOOL-5046)
+			return TSharedPtr<FNURBSCurve>();
+
+#ifdef UETOOL_5046_WIP
 			TArray<FPoint> NewPoles;
 			TArray<double> NewNodalVector;
 			TArray<double> NewWeights;
@@ -825,6 +829,7 @@ namespace CADKernel
 				return TSharedPtr<FNURBSCurve>();
 			}
 			return FEntity::MakeShared<FNURBSCurve>(Degre, NewNodalVector, NewPoles, NewWeights, InCurve.GetDimension());
+#endif
 		}
 
 
@@ -1537,6 +1542,7 @@ namespace CADKernel
 
 			if (newPoles.Num() + newN - 1 != newNodalVector.Num())
 			{
+				// To avoid crashes while waiting for the fix (jira UETOOL-5046)
 				newPoles.Empty();
 				return;
 			}
diff --git a/Engine/Source/Runtime/Datasmith/CADKernel/Private/CADKernel/Topo/Shell.cpp b/Engine/Source/Runtime/Datasmith/CADKernel/Private/CADKernel/Topo/Shell.cpp
--- a/Engine/Source/Runtime/Datasmith/CADKernel/Private/CADKernel/Topo/Shell.cpp
+++ b/Engine/Source/Runtime/Datasmith/CADKernel/Private/CADKernel/Topo/Shell.cpp
@@ -420,7 +420,9 @@ int32 FShell::Orient()
 	{
 		if (SubshellSwappedFaceCount > (SubshellFaceCount / 2))
 		{
+#ifdef CADERNEL_DEV
 			FMessage::Printf(Log, TEXT("the open shell %s (Id %d) is badly oriented. Its orientation is swapped"), *GetName(), GetId());
+#endif
 			SwapSubshellOrientation();
 			ShellSwappedFaceCount += (SubshellFaceCount - SubshellSwappedFaceCount);
 		}
@@ -460,7 +462,9 @@ int32 FShell::Orient()
 
 		if (bHasWrongOrientation)
 		{
+#ifdef CADERNEL_DEV
 			FMessage::Printf(Log, TEXT("The closed shell %s (Id %d) is badly oriented. Its orientation is swapped"), *GetName(), GetId());
+#endif
 			SwapSubshellOrientation();
 			ShellSwappedFaceCount += (SubshellFaceCount - SubshellSwappedFaceCount);
 		}
