Crash when importing CAD file via Datasmith with Stitching Sew: a delete was done on an array that its life cycle is managed by TechSoft

diff --git a/Engine/Plugins/Enterprise/DatasmithCADImporter/Source/ParametricSurface/Private/TechSoft/CADModelToTechSoftConverterBase.cpp b/Engine/Plugins/Enterprise/DatasmithCADImporter/Source/ParametricSurface/Private/TechSoft/CADModelToTechSoftConverterBase.cpp
--- a/Engine/Plugins/Enterprise/DatasmithCADImporter/Source/ParametricSurface/Private/TechSoft/CADModelToTechSoftConverterBase.cpp
+++ b/Engine/Plugins/Enterprise/DatasmithCADImporter/Source/ParametricSurface/Private/TechSoft/CADModelToTechSoftConverterBase.cpp
@@ -30,11 +30,6 @@ bool FCADModelToTechSoftConverterBase::RepairTopology()
 		{
 			RiRepresentationItems.Add(OutNewBReps[Index]);
 		}
-
-		if (OutNewBRepCount != 0)
-		{
-			delete OutNewBReps;
-		}
 	}
 #endif
 	return true;
