Hoops import with stitching Sew crashes on wire file:
In this case, Alias shells are sewn with TechSoft and the result is a set of shells (A3DRiRepresentationItem). The crash occur during the conversion of this set of A3DRiRepresentationItem into a FBodyMesh.
Retessellate with techSoft did not support this kind of StaticMesh yet. This is also fixed.

diff --git a/Engine/Plugins/Enterprise/DatasmithCADImporter/Source/CADInterfaces/Private/TechSoftInterfaceUtils.inl b/Engine/Plugins/Enterprise/DatasmithCADImporter/Source/CADInterfaces/Private/TechSoftInterfaceUtils.inl
--- a/Engine/Plugins/Enterprise/DatasmithCADImporter/Source/CADInterfaces/Private/TechSoftInterfaceUtils.inl
+++ b/Engine/Plugins/Enterprise/DatasmithCADImporter/Source/CADInterfaces/Private/TechSoftInterfaceUtils.inl
@@ -278,9 +278,9 @@ namespace TechSoftInterfaceUtils
 				
 				if (VertexOffset > 0)
 				{
-					for (int32& VertexIndex : Tessellation.VertexIndices)
+					for (int32& PositionIndex : Tessellation.PositionIndices)
 					{
-						VertexIndex += VertexOffset;
+						PositionIndex += VertexOffset;
 					}
 				}
 
diff --git a/Engine/Plugins/Enterprise/DatasmithCADImporter/Source/CADInterfaces/Private/TechSoftUtils.cpp b/Engine/Plugins/Enterprise/DatasmithCADImporter/Source/CADInterfaces/Private/TechSoftUtils.cpp
--- a/Engine/Plugins/Enterprise/DatasmithCADImporter/Source/CADInterfaces/Private/TechSoftUtils.cpp
+++ b/Engine/Plugins/Enterprise/DatasmithCADImporter/Source/CADInterfaces/Private/TechSoftUtils.cpp
@@ -34,8 +34,6 @@ void SetRootOccurenceAttributes(A3DEntity* Entity);
 
 bool GetBodyFromPcrFile(const FString& Filename, const FImportParameters& ImportParameters, FBodyMesh& BodyMesh)
 {
-	bool bExtractionSuccessful = false;
-
 #if defined USE_TECHSOFT_SDK && !defined CADKERNEL_DEV
 	A3DRWParamsPrcReadHelper* ReadHelper = nullptr;
 
@@ -68,18 +66,20 @@ bool GetBodyFromPcrFile(const FString& Filename, const FImportParameters& Import
 	if (JsonObject.IsValid())
 	{
 		double FileUnit = 1.0;
-
 		JsonObject->TryGetNumberField(JSON_ENTRY_FILE_UNIT, FileUnit);
 
-		if (FillBodyMesh(PartDefinitionData->m_ppRepItems[0], ImportParameters, FileUnit, BodyMesh))
+		for (A3DUns32 Index = 0; Index < PartDefinitionData->m_uiRepItemsSize; ++Index)
 		{
-			bExtractionSuccessful = true;
-			RestoreMaterials(JsonObject, BodyMesh);
+			if (!FillBodyMesh(PartDefinitionData->m_ppRepItems[Index], ImportParameters, FileUnit, BodyMesh))
+			{
+				return false;
+			}
 		}
+		RestoreMaterials(JsonObject, BodyMesh);
+		return true;
 	}
-
 #endif
-	return bExtractionSuccessful;
+	return false;
 }
 
 FUniqueTechSoftModelFile SaveBodiesToPrcFile(void** Bodies, uint32 BodyCount, const FString& Filename, const FString& JsonString)
@@ -180,7 +180,10 @@ bool FillBodyMesh(void* BodyPtr, const FImportParameters& ImportParameters, doub
 
 	{
 		A3DEntityGetType(RepresentationItemData->m_pTessBase, &Type);
-		ensure(Type == kA3DTypeTess3D);
+		if (Type != kA3DTypeTess3D)
+		{
+			return false;
+		}
 	}
 
 	TechSoftInterfaceUtils::FTechSoftTessellationExtractor Extractor(RepresentationItemData->m_pTessBase);
