Wrong orientation of some sub-part of StaticMesh:
The imported mesh is not well connected. The orientation algorithm cannot work properly.
Coincident nodes are merged before to build the StaticMesh

diff --git a/Engine/Plugins/Enterprise/DatasmithCADImporter/Source/CADLibrary/Private/MeshDescriptionHelper.cpp b/Engine/Plugins/Enterprise/DatasmithCADImporter/Source/CADLibrary/Private/MeshDescriptionHelper.cpp
--- a/Engine/Plugins/Enterprise/DatasmithCADImporter/Source/CADLibrary/Private/MeshDescriptionHelper.cpp
+++ b/Engine/Plugins/Enterprise/DatasmithCADImporter/Source/CADLibrary/Private/MeshDescriptionHelper.cpp
@@ -206,6 +206,55 @@ namespace CADLibrary
 		}
 	}
 
+	void MergeCoincidentVertices(TArray<FVector>& VertexArray, TArray<int32>& VertexIdSet)
+	{
+		const double CoincidenceTolerance = 0.001;
+
+		// Create a list of vertex Z/index pairs
+		TArray<FVertexData> VertexDataSet;
+		VertexDataSet.Reserve(VertexArray.Num());
+
+		const FVector* Position = VertexArray.GetData();
+		const int32* VertexId = VertexIdSet.GetData();
+		for (int32 Index = 0; Index < VertexArray.Num(); ++Index, ++Position, ++VertexId)
+		{
+			VertexDataSet.Emplace(*VertexId, *Position);
+		}
+
+		// Sort the vertices by z value
+		VertexDataSet.Sort(FCompareVertexZ());
+
+		// Search for duplicates
+		for (int32 Index = 0; Index < VertexDataSet.Num(); Index++)
+		{
+			if (VertexDataSet[Index].bIsMerged)
+			{
+				continue;
+			}
+
+			VertexDataSet[Index].bIsMerged = true;
+			int32 NewIndex = VertexDataSet[Index].Index;
+
+			const FVector& PositionA = VertexDataSet[Index].Coordinates;
+
+			// only need to search forward, since we add pairs both ways
+			for (int32 Andex = Index + 1; Andex < VertexDataSet.Num(); Andex++)
+			{
+				if (FMath::Abs(VertexDataSet[Andex].Z - VertexDataSet[Index].Z) > 3.*CoincidenceTolerance)
+				{
+					break; // can't be any more duplicated
+				}
+
+				const FVector& PositionB = VertexDataSet[Andex].Coordinates;
+				if (PositionA.Equals(PositionB, CoincidenceTolerance))
+				{
+					VertexDataSet[Andex].bIsMerged = true;
+					VertexIdSet[VertexDataSet[Andex].Index] = NewIndex;
+				}
+			}
+		}
+	}
+
 	void FillVertexPosition(const FImportParameters& ImportParams, const FMeshParameters& MeshParameters, FBodyMesh& Body, FMeshDescription& MeshDescription)
 	{
 		int32 TriangleCount = Body.TriangleCount;
@@ -243,6 +292,8 @@ namespace CADLibrary
 			VertexIdSet[VertexIndex] = VertexID;
 		}
 
+		MergeCoincidentVertices(VertexArray, VertexIdSet);
+
 		// if Symmetric mesh, the symmetric side of the mesh have to be generated
 		if (MeshParameters.bIsSymmetric)
 		{
@@ -261,8 +312,9 @@ namespace CADLibrary
 				}
 
 				FVertexID VertexID = MeshDescription.CreateVertex();
-				VertexPositions[VertexID] = (FVector3f)FDatasmithUtils::ConvertVector((FDatasmithUtils::EModelCoordSystem) ImportParams.GetModelCoordSys(), Vertex);
-				VertexPositions[VertexID] = FVector4f(SymmetricMatrix.TransformPosition((FVector)VertexPositions[VertexID]));
+				const FVector VertexPosition = FDatasmithUtils::ConvertVector((FDatasmithUtils::EModelCoordSystem) ImportParams.GetModelCoordSys(), Vertex);
+				const FVector4f SymmetricPosition = FVector4f(SymmetricMatrix.TransformPosition(VertexPosition));
+				VertexPositions[VertexID] = FVector3f(SymmetricPosition);
 				SymmetricVertexIds[VertexIndex++] = VertexID;
 			}
 		}
