FTechSoftFileParser: Wrong logic in the step of count of the file components in the case of file with configurations (Solidworks).
The right count of elements is mandatory to avoid a crash in the next step.

diff --git a/Engine/Plugins/Enterprise/DatasmithCADImporter/Source/CADInterfaces/Private/TechSoftFileParser.cpp b/Engine/Plugins/Enterprise/DatasmithCADImporter/Source/CADInterfaces/Private/TechSoftFileParser.cpp
--- a/Engine/Plugins/Enterprise/DatasmithCADImporter/Source/CADInterfaces/Private/TechSoftFileParser.cpp
+++ b/Engine/Plugins/Enterprise/DatasmithCADImporter/Source/CADInterfaces/Private/TechSoftFileParser.cpp
@@ -21,6 +21,17 @@ namespace TechSoftFileParserImpl
 {
 // Functions to clean metadata
 
+inline void CheckScaleComponentAndWarning(FMatrix TransformMatrix, FCADFileData& CADFileData)
+{
+	FTransform3d Tranform(TransformMatrix);
+	const FVector3d ScaleVector = Tranform.GetScale3D();
+	if (!ScaleVector.Equals(FVector3d::OneVector, 0.001))
+	{
+		double Scale = (ScaleVector.X + ScaleVector.Y + ScaleVector.Z) / 3.;
+		CADFileData.AddWarningMessages(FString::Printf(TEXT("The scale component of the root transform matrix is %lf. The tessellation parameters do not take this scale into account. The resulting mesh may not be the expected one."), Scale));
+	}
+}
+
 inline void RemoveUnwantedChar(FString& StringToClean, TCHAR UnwantedChar)
 {
 	FString NewString;
@@ -498,6 +509,8 @@ void FTechSoftFileParser::TraverseConfigurationSet(const A3DAsmProductOccurrence
 		TransformMatrix = ExtractTransformation(Location);
 	}
 
+	TechSoftFileParserImpl::CheckScaleComponentAndWarning(TransformMatrix, CADFileData);
+
 	TUniqueTSObj<A3DAsmProductOccurrenceData> ConfigurationData;
 	for (unsigned int Index = 0; Index < ConfigurationSetData->m_uiPOccurrencesSize; ++Index)
 	{
@@ -560,26 +573,59 @@ void FTechSoftFileParser::CountUnderConfigurationSet(const A3DAsmProductOccurren
 		return;
 	}
 
+	const FString& ConfigurationToLoad = CADFileData.GetCADFileDescription().GetConfiguration();
+
 	TUniqueTSObj<A3DAsmProductOccurrenceData> ConfigurationData;
 	for (unsigned int Index = 0; Index < ConfigurationSetData->m_uiPOccurrencesSize; ++Index)
 	{
 		ConfigurationData.FillFrom(ConfigurationSetData->m_ppPOccurrences[Index]);
 		if (!ConfigurationData.IsValid())
 		{
-			return;
+			continue;
 		}
 
-		if (ConfigurationData->m_uiProductFlags & (A3D_PRODUCT_FLAG_DEFAULT | A3D_PRODUCT_FLAG_CONFIG))
+		if (ConfigurationData->m_uiProductFlags & A3D_PRODUCT_FLAG_CONFIG)
 		{
-			CountUnderOccurrence(ConfigurationSetData->m_ppPOccurrences[Index]);
-			return;
+			bool bIsConfigurationToLoad = false;
+			if (!ConfigurationToLoad.IsEmpty())
+			{
+				FEntityMetaData ConfigurationMetaData;
+				ExtractMetaData(ConfigurationSetData->m_ppPOccurrences[Index], ConfigurationMetaData);
+				const FString* ConfigurationName = ConfigurationMetaData.MetaData.Find(TEXT("SDKName"));
+				if (ConfigurationName)
+				{
+					bIsConfigurationToLoad = (ConfigurationName->Equals(ConfigurationToLoad));
+				}
+			}
+			else
+			{
+				bIsConfigurationToLoad = ConfigurationData->m_uiProductFlags & A3D_PRODUCT_FLAG_DEFAULT;
+			}
+
+			if (bIsConfigurationToLoad)
+			{
+				CountUnderOccurrence(ConfigurationSetData->m_ppPOccurrences[Index]);
+				return;
+			}
 		}
 	}
 
-	// no default configuration, traverse the first
-	if (ConfigurationSetData->m_uiPOccurrencesSize)
+	if (ConfigurationToLoad.IsEmpty())
 	{
-		CountUnderOccurrence(ConfigurationSetData->m_ppPOccurrences[0]);
+		// no default configuration, traverse the first configuration
+		for (unsigned int Index = 0; Index < ConfigurationSetData->m_uiPOccurrencesSize; ++Index)
+		{
+			ConfigurationData.FillFrom(ConfigurationSetData->m_ppPOccurrences[Index]);
+			if (!ConfigurationData.IsValid())
+			{
+				return;
+			}
+
+			if (ConfigurationData->m_uiProductFlags & A3D_PRODUCT_FLAG_CONFIG)
+			{
+				CountUnderOccurrence(ConfigurationSetData->m_ppPOccurrences[Index]);
+			}
+		}
 	}
 }
 
@@ -612,6 +658,8 @@ void FTechSoftFileParser::TraverseReference(const A3DAsmProductOccurrence* Refer
 		ReferenceMatrix = ExtractTransformation(Location);
 	}
 
+	TechSoftFileParserImpl::CheckScaleComponentAndWarning(ReferenceMatrix, CADFileData);
+
 	Component.TransformMatrix = ParentMatrix * ReferenceMatrix;
 
 	for (uint32 OccurenceIndex = 0; OccurenceIndex < ReferenceData->m_uiPOccurrencesSize; ++OccurenceIndex)
