Test_File.IGS imported with missing faces

The problem come from automatic sew operation (because it's a IGES file) done with a too small value of the stitching tolerance due to the file unit (inch).
For IGES file, sew tolerance is now 0.01 regardless of the file unit.
Add ds.CADTranslator.StitchingTolerance CVar to allow user to modify the stitching tolerance (#jira UETOOL-5054)

diff --git a/Engine/Plugins/Enterprise/DatasmithCADImporter/Source/CADInterfaces/Private/TechSoftFileParser.cpp b/Engine/Plugins/Enterprise/DatasmithCADImporter/Source/CADInterfaces/Private/TechSoftFileParser.cpp
--- a/Engine/Plugins/Enterprise/DatasmithCADImporter/Source/CADInterfaces/Private/TechSoftFileParser.cpp
+++ b/Engine/Plugins/Enterprise/DatasmithCADImporter/Source/CADInterfaces/Private/TechSoftFileParser.cpp
@@ -12,6 +12,14 @@
 #include "Serialization/JsonWriter.h"
 #include "Templates/UnrealTemplate.h"
 
+float GStitchingTolerance = 0;
+FAutoConsoleVariableRef GCADTranslatorStitchingTolerance(
+	TEXT("ds.CADTranslator.StitchingTolerance"),
+	GStitchingTolerance,
+	TEXT("Welding threshold for Heal/Sew stitching methods in mm (except for IGES format where the threshold is in the file unit)\n\
+Default value of StitchingTolerance is 0.01 mm\n"),
+	ECVF_Default);
+
 namespace CADLibrary
 {
 
@@ -295,6 +303,7 @@ ECADParsingResult FTechSoftFileParser::Process()
 		}
 	}
 
+	// Some formats (like IGES) require a sew all the time. In this case, bForceSew = true
 	if (bForceSew || CADFileData.GetImportParameters().GetStitchingTechnique() == StitchingSew)
 	{
 		SewModel();
@@ -326,9 +335,11 @@ void FTechSoftFileParser::SewModel()
 {
 	CADLibrary::TUniqueTSObj<A3DSewOptionsData> SewData;
 	SewData->m_bComputePreferredOpenShellOrientation = false;
-	double ToleranceMM = 0.01 / FileUnit;
+	
+	// Except for IGES files, the value in GStitchingTolerance is assumed to be millimeters therefore must be converted to the model's unit
+	const double StitchingTolerance = Format == ECADFormat::IGES ? GStitchingTolerance : GStitchingTolerance / FileUnit;
 
-	TechSoftInterface::SewModel(ModelFile.Get(), ToleranceMM, SewData.GetPtr());
+	TechSoftInterface::SewModel(ModelFile.Get(), StitchingTolerance, SewData.GetPtr());
 }
 
 
