Static Mesh is not visible in Geometry folder or in viewport
The problem is not that StaticMesh are not visible in view port, it's that the file unit is not repected and Mesh are imported with mm unit instead of meter unit.
So the meshes are 1000 time smaller than expected.
FileUnit is not defined at a unique place in TechSoft SDK

diff --git a/Engine/Plugins/Enterprise/DatasmithCADImporter/Source/CADInterfaces/Private/TechSoftFileParser.cpp b/Engine/Plugins/Enterprise/DatasmithCADImporter/Source/CADInterfaces/Private/TechSoftFileParser.cpp
--- a/Engine/Plugins/Enterprise/DatasmithCADImporter/Source/CADInterfaces/Private/TechSoftFileParser.cpp
+++ b/Engine/Plugins/Enterprise/DatasmithCADImporter/Source/CADInterfaces/Private/TechSoftFileParser.cpp
@@ -662,6 +662,13 @@ void FTechSoftFileParser::TraverseReference(const A3DAsmProductOccurrence* Refer
 		return;
 	}
 
+	// Override model's FileUnit only if model's FileUnit is equal to 1. In this case, this means that the file unit is defined in the Reference otherwise we do not override the file unit.
+	// We have notice that we can have :
+	//    - ModelFileData->m_dUnit = 1; ReferenceData->m_dUnit = 25.4 => FileUnit = 25.4
+	//    or ModelFileData->m_dUnit = 25.4; ReferenceData->m_dUnit = 1 => FileUnit = 25.4
+	// ModelFileData has many references, the use of TGuardValue protects the case where each reference has a different "dUnit"
+	TGuardValue<double> RootUnit(FileUnit, FMath::IsNearlyEqual(FileUnit, 1.) ? ReferenceData->m_dUnit : FileUnit);
+
 	FMatrix ReferenceMatrix = FMatrix::Identity;;
 	A3DMiscTransformation* Location = ReferenceData->m_pLocation;
 	if (Location)
