CAD import: Unexpected variation of actor creation

The problem come from that some actors are created with the same Name. Actor Name has to be unique otherwise the first actor with the duplicated name is deleted.
In CADTranslator, Actor name is based on its parent name, its label and an indice.
The choice of the indice was not perfect and in this case two Actors (an Actor and a MeshActor) had the same name... The Actor (that was created first) was deleted and its StaticMeshes was linked under the Editor Node.

diff --git a/Engine/Plugins/Enterprise/DatasmithCADImporter/Source/DatasmithCADTranslator/Private/DatasmithSceneGraphBuilder.cpp b/Engine/Plugins/Enterprise/DatasmithCADImporter/Source/DatasmithCADTranslator/Private/DatasmithSceneGraphBuilder.cpp
--- a/Engine/Plugins/Enterprise/DatasmithCADImporter/Source/DatasmithCADTranslator/Private/DatasmithSceneGraphBuilder.cpp
+++ b/Engine/Plugins/Enterprise/DatasmithCADImporter/Source/DatasmithCADTranslator/Private/DatasmithSceneGraphBuilder.cpp
@@ -150,7 +150,7 @@ void FDatasmithSceneGraphBuilder::FillAnchorActor(const TSharedRef<IDatasmithAct
 	TMap<FString, FString> InstanceNodeMetaDataMap;
 	FString ActorUUID;
 	FString ActorLabel;
-	GetNodeUUIDAndName(InstanceNodeMetaDataMap, Component.MetaData, 1, ParentData.Uuid, ActorUUID, ActorLabel);
+	GetNodeUUIDAndName(InstanceNodeMetaDataMap, Component.MetaData, Component.ObjectId, ParentData.Uuid, ActorUUID, ActorLabel);
 
 	AddMetaData(ActorElement, InstanceNodeMetaDataMap, Component.MetaData);
 
@@ -337,7 +337,7 @@ TSharedPtr< IDatasmithActorElement >  FDatasmithSceneBaseGraphBuilder::BuildInst
 
 	FString ActorUUID;
 	FString ActorLabel;
-	GetNodeUUIDAndName(Instance.MetaData, Reference->MetaData, InstanceIndex, ParentData.Uuid, ActorUUID, ActorLabel);
+	GetNodeUUIDAndName(Instance.MetaData, Reference->MetaData, Instance.ObjectId, ParentData.Uuid, ActorUUID, ActorLabel);
 
 	TSharedPtr<IDatasmithActorElement> Actor = CreateActor(*ActorUUID, *ActorLabel);
 	if (!Actor.IsValid())
@@ -449,10 +449,9 @@ TSharedPtr<IDatasmithActorElement> FDatasmithSceneBaseGraphBuilder::BuildCompone
 {
 	TMap<FString, FString> InstanceNodeMetaDataMap;
 
-	FString ActorUUID = TEXT("");
-	FString ActorLabel = TEXT("");
-	
-	GetNodeUUIDAndName(InstanceNodeMetaDataMap, Component.MetaData, 1, ParentData.Uuid, ActorUUID, ActorLabel);
+	FString ActorUUID;
+	FString ActorLabel;
+	GetNodeUUIDAndName(InstanceNodeMetaDataMap, Component.MetaData, Component.ObjectId, ParentData.Uuid, ActorUUID, ActorLabel);
 
 	TSharedPtr< IDatasmithActorElement > Actor = CreateActor(*ActorUUID, *ActorLabel);
 	if (!Actor.IsValid())
@@ -485,7 +484,7 @@ TSharedPtr<IDatasmithActorElement> FDatasmithSceneBaseGraphBuilder::BuildBody(in
 
 	FString BodyUUID;
 	FString BodyLabel;
-	GetNodeUUIDAndName(InstanceNodeMetaDataMap, Body.MetaData, BodyIndex, ParentData.Uuid, BodyUUID, BodyLabel);
+	GetNodeUUIDAndName(InstanceNodeMetaDataMap, Body.MetaData, Body.ObjectId, ParentData.Uuid, BodyUUID, BodyLabel);
 
 	// Apply materials on the current part
 	uint32 MaterialUuid = 0;
