Cad import - Wire - some part of model doesn't appears with Hoops import

The problem come from the orientation of each patch getting from Alias is not right.
Ignoring this information fix the problem on all cases.

diff --git a/Engine/Plugins/Enterprise/DatasmithCADImporter/Source/DatasmithWireTranslator/Private/AliasModelToTechSoftConverter.cpp b/Engine/Plugins/Enterprise/DatasmithCADImporter/Source/DatasmithWireTranslator/Private/AliasModelToTechSoftConverter.cpp
--- a/Engine/Plugins/Enterprise/DatasmithCADImporter/Source/DatasmithWireTranslator/Private/AliasModelToTechSoftConverter.cpp
+++ b/Engine/Plugins/Enterprise/DatasmithCADImporter/Source/DatasmithWireTranslator/Private/AliasModelToTechSoftConverter.cpp
@@ -348,10 +348,6 @@ bool FAliasModelToTechSoftConverter::AddBRep(AlDagNode& DagNode, EAliasObjectRef
 #ifdef USE_TECHSOFT_SDK
 	AlEdgeToTSCoEdge.Empty();
 
-	boolean bAlOrientation;
-	DagNode.getSurfaceOrientation(bAlOrientation);
-	bool bOrientation = (bool)bAlOrientation;
-
 	AlMatrix4x4 AlMatrix;
 	if (InObjectReference == EAliasObjectReference::ParentReference)
 	{
@@ -359,9 +355,7 @@ bool FAliasModelToTechSoftConverter::AddBRep(AlDagNode& DagNode, EAliasObjectRef
 	}
 
 	TArray<A3DTopoFace*> TSFaces;
-	TArray<A3DUns8> FaceOrientations;
 	TSFaces.Reserve(100);
-	FaceOrientations.Reserve(100);
 
 	AlObjectType objectType = DagNode.type();
 	switch (objectType)
@@ -380,7 +374,6 @@ bool FAliasModelToTechSoftConverter::AddBRep(AlDagNode& DagNode, EAliasObjectRef
 					if (TSFace != nullptr)
 					{
 						TSFaces.Add(TSFace);
-						FaceOrientations.Add(bOrientation);
 					}
 				}
 			}
@@ -404,7 +397,6 @@ bool FAliasModelToTechSoftConverter::AddBRep(AlDagNode& DagNode, EAliasObjectRef
 						if (TSFace != nullptr)
 						{
 							TSFaces.Add(TSFace);
-							FaceOrientations.Add(bOrientation);
 						}
 					}
 					break;
@@ -417,7 +409,6 @@ bool FAliasModelToTechSoftConverter::AddBRep(AlDagNode& DagNode, EAliasObjectRef
 					if (TSFace != nullptr)
 					{
 						TSFaces.Add(TSFace);
-						FaceOrientations.Add(bOrientation);
 					}
 				}
 			}
@@ -435,6 +426,9 @@ bool FAliasModelToTechSoftConverter::AddBRep(AlDagNode& DagNode, EAliasObjectRef
 
 	A3DTopoShell* TopoShellPtr = nullptr;
 	{
+		TArray<A3DUns8> FaceOrientations;
+		FaceOrientations.Init(true, TSFaces.Num());
+
 		CADLibrary::TUniqueTSObj<A3DTopoShellData> TopoShellData;
 		TopoShellData->m_bClosed = false;
 		TopoShellData->m_ppFaces = TSFaces.GetData();
