Fixed warning during static analysis on TechSoftFileParser.cpp

diff --git a/Engine/Plugins/Enterprise/DatasmithCADImporter/Source/CADInterfaces/Private/TechSoftFileParser.cpp b/Engine/Plugins/Enterprise/DatasmithCADImporter/Source/CADInterfaces/Private/TechSoftFileParser.cpp
--- a/Engine/Plugins/Enterprise/DatasmithCADImporter/Source/CADInterfaces/Private/TechSoftFileParser.cpp
+++ b/Engine/Plugins/Enterprise/DatasmithCADImporter/Source/CADInterfaces/Private/TechSoftFileParser.cpp
@@ -1628,7 +1628,6 @@ FArchiveMaterial& FTechSoftFileParser::FindOrAddMaterial(uint32 MaterialIndex, c
 		TUniqueTSObjFromIndex<A3DGraphTextureApplicationData> TextureData(MaterialIndex);
 		if (TextureData.IsValid())
 		{
-			TextureData->m_uiMaterialIndex;
 			return AddMaterialAt(TextureData->m_uiMaterialIndex, MaterialIndex, GraphStyleData);
 			
 #ifdef NOTYETDEFINE
