Worker lost when Importing STP file with TechSoft + Stitching "Heal"

In some case, Techsoft return 0 body after a sew. In this case, the initial body is used.
If the sew return multiple bodies, they are merged in the static mesh

diff --git a/Engine/Plugins/Enterprise/DatasmithCADImporter/Source/CADInterfaces/Private/TechSoftFileParser.cpp b/Engine/Plugins/Enterprise/DatasmithCADImporter/Source/CADInterfaces/Private/TechSoftFileParser.cpp
--- a/Engine/Plugins/Enterprise/DatasmithCADImporter/Source/CADInterfaces/Private/TechSoftFileParser.cpp
+++ b/Engine/Plugins/Enterprise/DatasmithCADImporter/Source/CADInterfaces/Private/TechSoftFileParser.cpp
@@ -12,7 +12,7 @@
 #include "Serialization/JsonWriter.h"
 #include "Templates/UnrealTemplate.h"
 
-float GStitchingTolerance = 0;
+float GStitchingTolerance = 0.01f;
 FAutoConsoleVariableRef GCADTranslatorStitchingTolerance(
 	TEXT("ds.CADTranslator.StitchingTolerance"),
 	GStitchingTolerance,
@@ -357,22 +357,32 @@ void FTechSoftFileParser::GenerateBodyMesh(A3DRiRepresentationItem* Representati
 {
 	FBodyMesh& BodyMesh = CADFileData.AddBodyMesh(Body.ObjectId, Body);
 
+	uint32 NewBRepCount = 0;
+	A3DRiBrepModel** NewBReps = nullptr;
+
 	if (CADFileData.GetImportParameters().GetStitchingTechnique() == StitchingHeal)
 	{
 		TUniqueTSObj<A3DSewOptionsData> SewData;
 		SewData->m_bComputePreferredOpenShellOrientation = false;
 		double ToleranceMM = 0.01 / FileUnit;
-		A3DRiBrepModel** NewBReps;
-		uint32 NewBRepCount = 0;
 		A3DStatus Status = TechSoftInterface::SewBReps(&Representation, 1, ToleranceMM, SewData.GetPtr(), &NewBReps, NewBRepCount);
-		ensure(NewBRepCount == 1);
-		if (Status == A3DStatus::A3D_SUCCESS)
+		if (Status != A3DStatus::A3D_SUCCESS)
 		{
-			Representation = NewBReps[0];
+			CADFileData.AddWarningMessages(TEXT("A body healing failed. A body could missed."));
 		}
 	}
 
-	TechSoftUtils::FillBodyMesh(Representation, CADFileData.GetImportParameters(), FileUnit, BodyMesh);
+	if (NewBRepCount > 0)
+	{
+		for (uint32 Index = 0; Index < NewBRepCount; ++Index)
+		{
+			TechSoftUtils::FillBodyMesh(NewBReps[Index], CADFileData.GetImportParameters(), FileUnit, BodyMesh);
+		}
+	} 
+	else
+	{
+		TechSoftUtils::FillBodyMesh(Representation, CADFileData.GetImportParameters(), FileUnit, BodyMesh);
+	}
 
 	// Convert material
 	FCADUUID DefaultColorName = Body.ColorFaceSet.Num() > 0 ? *Body.ColorFaceSet.begin() : 0;
