CAD import - import failed with specific SLDPRT file

The unit of a mesh is defined in the unit of the file and in the scale value of each transform matrix. This information has to be propagate along the scene tree to have a StaticMesh at scale 1

diff --git a/Engine/Plugins/Enterprise/DatasmithCADImporter/Source/CADInterfaces/Private/CADSceneGraph.cpp b/Engine/Plugins/Enterprise/DatasmithCADImporter/Source/CADInterfaces/Private/CADSceneGraph.cpp
--- a/Engine/Plugins/Enterprise/DatasmithCADImporter/Source/CADInterfaces/Private/CADSceneGraph.cpp
+++ b/Engine/Plugins/Enterprise/DatasmithCADImporter/Source/CADInterfaces/Private/CADSceneGraph.cpp
@@ -46,6 +46,8 @@ FArchive& operator<<(FArchive& Ar, FArchiveBody& Body)
 	Ar << Body.ColorFaceSet;
 	Ar << Body.ParentId;
 	Ar << Body.MeshActorName;
+	Ar << Body.BodyUnit;
+
 	return Ar;
 }
 
diff --git a/Engine/Plugins/Enterprise/DatasmithCADImporter/Source/CADInterfaces/Private/TechSoftFileParser.cpp b/Engine/Plugins/Enterprise/DatasmithCADImporter/Source/CADInterfaces/Private/TechSoftFileParser.cpp
--- a/Engine/Plugins/Enterprise/DatasmithCADImporter/Source/CADInterfaces/Private/TechSoftFileParser.cpp
+++ b/Engine/Plugins/Enterprise/DatasmithCADImporter/Source/CADInterfaces/Private/TechSoftFileParser.cpp
@@ -12,12 +12,12 @@
 #include "Serialization/JsonWriter.h"
 #include "Templates/UnrealTemplate.h"
 
-float GStitchingTolerance = 0.01f;
+float GStitchingTolerance = 0.001f;
 FAutoConsoleVariableRef GCADTranslatorStitchingTolerance(
 	TEXT("ds.CADTranslator.StitchingTolerance"),
 	GStitchingTolerance,
-	TEXT("Welding threshold for Heal/Sew stitching methods in mm (except for IGES format where the threshold is in the file unit)\n\
-Default value of StitchingTolerance is 0.01 mm\n"),
+	TEXT("Welding threshold for Heal/Sew stitching methods in cm\n\
+Default value of StitchingTolerance is 0.001 cm\n"),
 	ECVF_Default);
 
 namespace CADLibrary
@@ -29,17 +29,6 @@ namespace TechSoftFileParserImpl
 {
 // Functions to clean metadata
 
-inline void CheckScaleComponentAndWarning(FMatrix TransformMatrix, FCADFileData& CADFileData)
-{
-	FTransform3d Tranform(TransformMatrix);
-	const FVector3d ScaleVector = Tranform.GetScale3D();
-	if (!ScaleVector.Equals(FVector3d::OneVector, 0.001))
-	{
-		double Scale = (ScaleVector.X + ScaleVector.Y + ScaleVector.Z) / 3.;
-		CADFileData.AddWarningMessages(FString::Printf(TEXT("The scale component of the root transform matrix is %lf. The tessellation parameters do not take this scale into account. The resulting mesh may not be the expected one."), Scale));
-	}
-}
-
 inline void RemoveUnwantedChar(FString& StringToClean, TCHAR UnwantedChar)
 {
 	FString NewString;
@@ -207,6 +196,33 @@ void UpdateIOOptionAccordingToFormat(const CADLibrary::ECADFormat Format, A3DImp
 	};
 }
 
+double ExtractUniformScale(FVector3d& Scale)
+{
+	double UniformScale = (Scale.X + Scale.Y + Scale.Z) / 3.;
+	double Tolerance = UniformScale * KINDA_SMALL_NUMBER;
+
+	if (!FMath::IsNearlyEqual(UniformScale, Scale.X, Tolerance) && !FMath::IsNearlyEqual(UniformScale, Scale.Y, Tolerance))
+	{
+		// non uniform scale 
+		// Used in format like IFC or DGN to define pipe by their diameter and their length
+		// we remove the diameter component of the scale to have a scale like (Length/diameter, 1, 1) to have a mesh tessellated according the meshing parameters
+		if (FMath::IsNearlyEqual((double)Scale.X, (double)Scale.Y, Tolerance) || (FMath::IsNearlyEqual((double)Scale.X, (double)Scale.Z, Tolerance)))
+		{
+			UniformScale = Scale.X;
+		}
+		else if (FMath::IsNearlyEqual((double)Scale.Y, (double)Scale.Z, Tolerance))
+		{
+			UniformScale = Scale.Y;
+		}
+	}
+
+	Scale.X /= UniformScale;
+	Scale.Y /= UniformScale;
+	Scale.Z /= UniformScale;
+
+	return UniformScale;
+}
+
 } // ns TechSoftFileParserImpl
 
 #endif
@@ -284,7 +300,7 @@ ECADParsingResult FTechSoftFileParser::Process()
 		}
 
 		ModellerType = (EModellerType)ModelFileData->m_eModellerType;
-		FileUnit = ModelFileData->m_dUnit;
+		FileUnit = TechSoftInterface::GetModelFileUnit(ModelFile.Get());
 	}
 
 	// Adapt BRep to CADKernel
@@ -336,10 +352,7 @@ void FTechSoftFileParser::SewModel()
 	CADLibrary::TUniqueTSObj<A3DSewOptionsData> SewData;
 	SewData->m_bComputePreferredOpenShellOrientation = false;
 	
-	// Except for IGES files, the value in GStitchingTolerance is assumed to be millimeters therefore must be converted to the model's unit
-	const double StitchingTolerance = Format == ECADFormat::IGES ? GStitchingTolerance : GStitchingTolerance / FileUnit;
-
-	TechSoftInterface::SewModel(ModelFile.Get(), StitchingTolerance, SewData.GetPtr());
+	TechSoftInterface::SewModel(ModelFile.Get(), GStitchingTolerance, SewData.GetPtr());
 }
 
 
@@ -364,11 +377,10 @@ void FTechSoftFileParser::GenerateBodyMesh(A3DRiRepresentationItem* Representati
 	{
 		TUniqueTSObj<A3DSewOptionsData> SewData;
 		SewData->m_bComputePreferredOpenShellOrientation = false;
-		double ToleranceMM = 0.01 / FileUnit;
-		A3DStatus Status = TechSoftInterface::SewBReps(&Representation, 1, ToleranceMM, SewData.GetPtr(), &NewBReps, NewBRepCount);
+		A3DStatus Status = TechSoftInterface::SewBReps(&Representation, 1, GStitchingTolerance, FileUnit, SewData.GetPtr(), &NewBReps, NewBRepCount);
 		if (Status != A3DStatus::A3D_SUCCESS)
 		{
-			CADFileData.AddWarningMessages(TEXT("A body healing failed. A body could missed."));
+			CADFileData.AddWarningMessages(TEXT("A body healing failed. A body could be missing."));
 		}
 	}
 
@@ -376,12 +388,12 @@ void FTechSoftFileParser::GenerateBodyMesh(A3DRiRepresentationItem* Representati
 	{
 		for (uint32 Index = 0; Index < NewBRepCount; ++Index)
 		{
-			TechSoftUtils::FillBodyMesh(NewBReps[Index], CADFileData.GetImportParameters(), FileUnit, BodyMesh);
+			TechSoftUtils::FillBodyMesh(NewBReps[Index], CADFileData.GetImportParameters(), Body.BodyUnit, BodyMesh);
 		}
 	} 
 	else
 	{
-		TechSoftUtils::FillBodyMesh(Representation, CADFileData.GetImportParameters(), FileUnit, BodyMesh);
+		TechSoftUtils::FillBodyMesh(Representation, CADFileData.GetImportParameters(), Body.BodyUnit, BodyMesh);
 	}
 
 	// Convert material
@@ -431,7 +443,7 @@ void FTechSoftFileParser::GenerateBodyMesh(A3DRiRepresentationItem* Representati
 
 			// Save file unit and default color and material attributes in a json string
 			// This will be used when the file is reloaded
-			JsonObject->SetNumberField(JSON_ENTRY_FILE_UNIT, FileUnit);
+			JsonObject->SetNumberField(JSON_ENTRY_FILE_UNIT, Body.BodyUnit);
 			JsonObject->SetNumberField(JSON_ENTRY_COLOR_NAME, DefaultColorName);
 			JsonObject->SetNumberField(JSON_ENTRY_MATERIAL_NAME, DefaultMaterialName);
 
@@ -498,18 +510,18 @@ ECADParsingResult FTechSoftFileParser::TraverseModel()
 	{
 		if (IsConfigurationSet(ModelFileData->m_ppPOccurrences[Index]))
 		{
-			TraverseConfigurationSet(ModelFileData->m_ppPOccurrences[Index]);
+			TraverseConfigurationSet(ModelFileData->m_ppPOccurrences[Index], FileUnit);
 		}
 		else
 		{
-			TraverseReference(ModelFileData->m_ppPOccurrences[Index], FMatrix::Identity);
+			TraverseReference(ModelFileData->m_ppPOccurrences[Index], FMatrix::Identity, FileUnit);
 		}
 	}
 
 	return ECADParsingResult::ProcessOk;
 }
 
-void FTechSoftFileParser::TraverseConfigurationSet(const A3DAsmProductOccurrence* ConfigurationSetPtr)
+void FTechSoftFileParser::TraverseConfigurationSet(const A3DAsmProductOccurrence* ConfigurationSetPtr, double ParentUnit)
 {
 	TUniqueTSObj<A3DAsmProductOccurrenceData> ConfigurationSetData(ConfigurationSetPtr);
 	if (!ConfigurationSetData.IsValid())
@@ -524,14 +536,13 @@ void FTechSoftFileParser::TraverseConfigurationSet(const A3DAsmProductOccurrence
 	const FString& ConfigurationToLoad = CADFileData.GetCADFileDescription().GetConfiguration();
 
 	FMatrix TransformMatrix = FMatrix::Identity;;
+	double OccurenceUnit = ParentUnit;
 	A3DMiscTransformation* Location = ConfigurationSetData->m_pLocation;
 	if (Location)
 	{
-		TransformMatrix = ExtractTransformation(Location);
+		TransformMatrix = ExtractTransformation(Location, OccurenceUnit);
 	}
 
-	TechSoftFileParserImpl::CheckScaleComponentAndWarning(TransformMatrix, CADFileData);
-
 	TUniqueTSObj<A3DAsmProductOccurrenceData> ConfigurationData;
 	for (unsigned int Index = 0; Index < ConfigurationSetData->m_uiPOccurrencesSize; ++Index)
 	{
@@ -561,7 +572,7 @@ void FTechSoftFileParser::TraverseConfigurationSet(const A3DAsmProductOccurrence
 
 			if (bIsConfigurationToLoad)
 			{
-				TraverseReference(ConfigurationSetData->m_ppPOccurrences[Index], TransformMatrix);
+				TraverseReference(ConfigurationSetData->m_ppPOccurrences[Index], TransformMatrix, OccurenceUnit);
 				return;
 			}
 		}
@@ -580,7 +591,7 @@ void FTechSoftFileParser::TraverseConfigurationSet(const A3DAsmProductOccurrence
 
 			if (ConfigurationData->m_uiProductFlags & A3D_PRODUCT_FLAG_CONFIG)
 			{
-				TraverseReference(ConfigurationSetData->m_ppPOccurrences[Index], TransformMatrix);
+				TraverseReference(ConfigurationSetData->m_ppPOccurrences[Index], TransformMatrix, OccurenceUnit);
 			}
 		}
 	}
@@ -650,7 +661,7 @@ void FTechSoftFileParser::CountUnderConfigurationSet(const A3DAsmProductOccurren
 	}
 }
 
-void FTechSoftFileParser::TraverseReference(const A3DAsmProductOccurrence* ReferencePtr, const FMatrix& ParentMatrix)
+void FTechSoftFileParser::TraverseReference(const A3DAsmProductOccurrence* ReferencePtr, const FMatrix& ParentMatrix, double ParentUnit)
 {
 	FEntityMetaData MetaData;
 	ExtractMetaData(ReferencePtr, MetaData);
@@ -672,33 +683,26 @@ void FTechSoftFileParser::TraverseReference(const A3DAsmProductOccurrence* Refer
 		return;
 	}
 
-	// Override model's FileUnit only if model's FileUnit is equal to 1. In this case, this means that the file unit is defined in the Reference otherwise we do not override the file unit.
-	// We have notice that we can have :
-	//    - ModelFileData->m_dUnit = 1; ReferenceData->m_dUnit = 25.4 => FileUnit = 25.4
-	//    or ModelFileData->m_dUnit = 25.4; ReferenceData->m_dUnit = 1 => FileUnit = 25.4
-	// ModelFileData has many references, the use of TGuardValue protects the case where each reference has a different "dUnit"
-	TGuardValue<double> RootUnit(FileUnit, FMath::IsNearlyEqual(FileUnit, 1.) ? ReferenceData->m_dUnit : FileUnit);
+	FMatrix ReferenceMatrix = FMatrix::Identity;
+	double ReferenceUnit = ParentUnit;
 
-	FMatrix ReferenceMatrix = FMatrix::Identity;;
 	A3DMiscTransformation* Location = ReferenceData->m_pLocation;
 	if (Location)
 	{
-		ReferenceMatrix = ExtractTransformation(Location);
+		ReferenceMatrix = ExtractTransformation(Location, ReferenceUnit);
 	}
 
-	TechSoftFileParserImpl::CheckScaleComponentAndWarning(ReferenceMatrix, CADFileData);
-
 	Component.TransformMatrix = ParentMatrix * ReferenceMatrix;
 
 	for (uint32 OccurenceIndex = 0; OccurenceIndex < ReferenceData->m_uiPOccurrencesSize; ++OccurenceIndex)
 	{
-		int32 ChildrenId = TraverseOccurrence(ReferenceData->m_ppPOccurrences[OccurenceIndex]);
+		int32 ChildrenId = TraverseOccurrence(ReferenceData->m_ppPOccurrences[OccurenceIndex], ReferenceUnit);
 		Component.Children.Add(ChildrenId);
 	}
 
 	if (ReferenceData->m_pPart)
 	{
-		TraversePartDefinition(ReferenceData->m_pPart, Component);
+		TraversePartDefinition(ReferenceData->m_pPart, Component, ReferenceUnit);
 	}
 }
 
@@ -765,13 +769,14 @@ FArchiveComponent& FTechSoftFileParser::AddOccurence(FEntityMetaData& InstanceMe
 	return Prototype;
 }
 
-int32 FTechSoftFileParser::AddBody(FEntityMetaData& BodyMetaData, const FMatrix& Matrix, const FCadId ParentId)
+int32 FTechSoftFileParser::AddBody(FEntityMetaData& BodyMetaData, const FMatrix& Matrix, const FCadId ParentId, double BodyUnit)
 {
 	FCadId BodyId = LastEntityId++;
 	int32 BodyIndex = CADFileData.AddBody(BodyId);
 	FArchiveBody& Body = CADFileData.GetBodyAt(BodyIndex);
 	Body.ParentId = ParentId;
 	Body.TransformMatrix = Matrix;
+	Body.BodyUnit = BodyUnit;
 	Body.MetaData = MoveTemp(BodyMetaData.MetaData);
 	if (BodyMetaData.ColorName != 0)
 	{
@@ -784,7 +789,7 @@ int32 FTechSoftFileParser::AddBody(FEntityMetaData& BodyMetaData, const FMatrix&
 	return BodyIndex;
 }
 
-FCadId FTechSoftFileParser::TraverseOccurrence(const A3DAsmProductOccurrence* OccurrencePtr)
+FCadId FTechSoftFileParser::TraverseOccurrence(const A3DAsmProductOccurrence* OccurrencePtr, double ParentUnit)
 {
 	// first product occurence with m_pPart != nullptr || m_uiPOccurrencesSize > 0
 	const A3DAsmProductOccurrence* CachedOccurrencePtr = OccurrencePtr;
@@ -824,9 +829,10 @@ FCadId FTechSoftFileParser::TraverseOccurrence(const A3DAsmProductOccurrence* Oc
 		ProcessPrototype(OccurrenceData->m_pPrototype, PrototypeMetaData, &Location);
 	}
 
+	double OccurrenceUnit = ParentUnit;
 	if (Location)
 	{
-		Instance.TransformMatrix = ExtractTransformation(Location);
+		Instance.TransformMatrix = ExtractTransformation(Location, OccurrenceUnit);
 	}
 
 	if (PrototypeMetaData.bUnloaded)
@@ -856,7 +862,7 @@ FCadId FTechSoftFileParser::TraverseOccurrence(const A3DAsmProductOccurrence* Oc
 	if (OccurrenceData->m_pPart != nullptr)
 	{
 		A3DAsmPartDefinition* PartDefinition = OccurrenceData->m_pPart;
-		TraversePartDefinition(PartDefinition, Component);
+		TraversePartDefinition(PartDefinition, Component, OccurrenceUnit);
 	}
 
 	// Add Occurrence's Children
@@ -870,7 +876,7 @@ FCadId FTechSoftFileParser::TraverseOccurrence(const A3DAsmProductOccurrence* Oc
 	A3DAsmProductOccurrence** Children = OccurrenceData->m_ppPOccurrences;
 	for (uint32 Index = 0; Index < ChildrenCount; ++Index)
 	{
-		int32 ChildrenId = TraverseOccurrence(Children[Index]);
+		int32 ChildrenId = TraverseOccurrence(Children[Index], OccurrenceUnit);
 		Component.Children.Add(ChildrenId);
 	}
 
@@ -886,14 +892,14 @@ FCadId FTechSoftFileParser::TraverseOccurrence(const A3DAsmProductOccurrence* Oc
 		if (ExternalData->m_pPart != nullptr)
 		{
 			A3DAsmPartDefinition* PartDefinition = ExternalData->m_pPart;
-			TraversePartDefinition(PartDefinition, Component);
+			TraversePartDefinition(PartDefinition, Component, OccurrenceUnit);
 		}
 
 		uint32 ExternalChildrenCount = ExternalData->m_uiPOccurrencesSize;
 		A3DAsmProductOccurrence** ExternalChildren = ExternalData->m_ppPOccurrences;
 		for (uint32 Index = 0; Index < ExternalChildrenCount; ++Index)
 		{
-			int32 ChildrenId = TraverseOccurrence(ExternalChildren[Index]);
+			int32 ChildrenId = TraverseOccurrence(ExternalChildren[Index], OccurrenceUnit);
 			Component.Children.Add(ChildrenId);
 		}
 	}
@@ -1081,7 +1087,7 @@ void FTechSoftFileParser::CountUnderPrototype(const A3DAsmProductOccurrence* Pro
 	ComponentCount[EComponentType::Reference] ++;
 }
 
-void FTechSoftFileParser::TraversePartDefinition(const A3DAsmPartDefinition* PartDefinitionPtr, FArchiveComponent& Part)
+void FTechSoftFileParser::TraversePartDefinition(const A3DAsmPartDefinition* PartDefinitionPtr, FArchiveComponent& Part, double ParentUnit)
 {
 	FEntityMetaData PartMetaData;
 	ExtractMetaData(PartDefinitionPtr, PartMetaData);
@@ -1099,7 +1105,7 @@ void FTechSoftFileParser::TraversePartDefinition(const A3DAsmPartDefinition* Par
 	{
 		for (unsigned int Index = 0; Index < PartData->m_uiRepItemsSize; ++Index)
 		{
-			int32 ChildId = TraverseRepresentationItem(PartData->m_ppRepItems[Index], PartMetaData, Part.ObjectId);
+			int32 ChildId = TraverseRepresentationItem(PartData->m_ppRepItems[Index], PartMetaData, Part.ObjectId, ParentUnit);
 			Part.Children.Add(ChildId);
 		}
 	}
@@ -1120,7 +1126,7 @@ void FTechSoftFileParser::CountUnderPartDefinition(const A3DAsmPartDefinition* P
 	}
 }
 
-FCadId FTechSoftFileParser::TraverseRepresentationItem(A3DRiRepresentationItem* RepresentationItem, const FEntityMetaData& PartMetaData, const FCadId ParentId)
+FCadId FTechSoftFileParser::TraverseRepresentationItem(A3DRiRepresentationItem* RepresentationItem, const FEntityMetaData& PartMetaData, const FCadId ParentId, double ParentUnit)
 {
 	A3DEEntityType Type;
 	A3DEntityGetType(RepresentationItem, &Type);
@@ -1128,11 +1134,11 @@ FCadId FTechSoftFileParser::TraverseRepresentationItem(A3DRiRepresentationItem*
 	switch (Type)
 	{
 	case kA3DTypeRiSet:
-		return TraverseRepresentationSet(RepresentationItem, PartMetaData);
+		return TraverseRepresentationSet(RepresentationItem, PartMetaData, ParentUnit);
 	case kA3DTypeRiBrepModel:
-		return TraverseBRepModel(RepresentationItem, PartMetaData, ParentId);
+		return TraverseBRepModel(RepresentationItem, PartMetaData, ParentId, ParentUnit);
 	case kA3DTypeRiPolyBrepModel:
-		return TraversePolyBRepModel(RepresentationItem, PartMetaData, ParentId);
+		return TraversePolyBRepModel(RepresentationItem, PartMetaData, ParentId, ParentUnit);
 	default:
 		break;
 	}
@@ -1159,7 +1165,7 @@ void FTechSoftFileParser::CountUnderRepresentationItem(const A3DRiRepresentation
 	}
 }
 
-FCadId FTechSoftFileParser::TraverseRepresentationSet(const A3DRiSet* RepresentationSetPtr, const FEntityMetaData& PartMetaData)
+FCadId FTechSoftFileParser::TraverseRepresentationSet(const A3DRiSet* RepresentationSetPtr, const FEntityMetaData& PartMetaData, double ParentUnit)
 {
 	TUniqueTSObj<A3DRiSetData> RepresentationSetData(RepresentationSetPtr);
 	if (!RepresentationSetData.IsValid())
@@ -1180,7 +1186,7 @@ FCadId FTechSoftFileParser::TraverseRepresentationSet(const A3DRiSet* Representa
 
 	for (A3DUns32 Index = 0; Index < RepresentationSetData->m_uiRepItemsSize; ++Index)
 	{
-		int32 ChildId = TraverseRepresentationItem(RepresentationSetData->m_ppRepItems[Index], RepresentationSetMetaData, RepresentationSet.ObjectId);
+		int32 ChildId = TraverseRepresentationItem(RepresentationSetData->m_ppRepItems[Index], RepresentationSetMetaData, RepresentationSet.ObjectId, ParentUnit);
 		RepresentationSet.Children.Add(ChildId);
 	}
 	return RepresentationSetId;
@@ -1201,7 +1207,7 @@ void FTechSoftFileParser::CountUnderRepresentationSet(const A3DRiSet* Representa
 	}
 }
 
-FCadId FTechSoftFileParser::TraverseBRepModel(A3DRiBrepModel* BRepModelPtr, const FEntityMetaData& PartMetaData, const FCadId ParentId)
+FCadId FTechSoftFileParser::TraverseBRepModel(A3DRiBrepModel* BRepModelPtr, const FEntityMetaData& PartMetaData, const FCadId ParentId, double ParentUnit)
 {
 	if (!BRepModelPtr)
 	{
@@ -1227,13 +1233,14 @@ FCadId FTechSoftFileParser::TraverseBRepModel(A3DRiBrepModel* BRepModelPtr, cons
 	ExtractSpecificMetaData(BRepModelPtr, BRepMetaData);
 
 	FMatrix Matrix = FMatrix::Identity;
+	double BRepUnit = ParentUnit;
 	TUniqueTSObj<A3DRiRepresentationItemData> RepresentationData(BRepModelPtr);
 	if (RepresentationData->m_pCoordinateSystem)
 	{
-		Matrix = ExtractCoordinateSystem(RepresentationData->m_pCoordinateSystem);
+		Matrix = ExtractCoordinateSystem(RepresentationData->m_pCoordinateSystem, BRepUnit);
 	}
 
-	int32 BodyIndex = AddBody(BRepMetaData, Matrix, ParentId);
+	int32 BodyIndex = AddBody(BRepMetaData, Matrix, ParentId, BRepUnit);
 	FArchiveBody& Body = CADFileData.GetBodyAt(BodyIndex);
 
 	RepresentationItemsCache.Add(BRepModelPtr, BodyIndex);
@@ -1241,7 +1248,7 @@ FCadId FTechSoftFileParser::TraverseBRepModel(A3DRiBrepModel* BRepModelPtr, cons
 	return Body.ObjectId;
 }
 
-FCadId FTechSoftFileParser::TraversePolyBRepModel(A3DRiPolyBrepModel* PolygonalPtr, const FEntityMetaData& PartMetaData, const FCadId ParentId)
+FCadId FTechSoftFileParser::TraversePolyBRepModel(A3DRiPolyBrepModel* PolygonalPtr, const FEntityMetaData& PartMetaData, const FCadId ParentId, double ParentUnit)
 {
 	if (!PolygonalPtr)
 	{
@@ -1267,13 +1274,14 @@ FCadId FTechSoftFileParser::TraversePolyBRepModel(A3DRiPolyBrepModel* PolygonalP
 	ExtractSpecificMetaData(PolygonalPtr, BRepMetaData);
 
 	FMatrix Matrix = FMatrix::Identity;
+	double BRepUnit = ParentUnit;
 	TUniqueTSObj<A3DRiRepresentationItemData> RepresentationData(PolygonalPtr);
 	if (RepresentationData->m_pCoordinateSystem)
 	{
-		Matrix = ExtractCoordinateSystem(RepresentationData->m_pCoordinateSystem);
+		Matrix = ExtractCoordinateSystem(RepresentationData->m_pCoordinateSystem, BRepUnit);
 	}
 
-	int32 BodyIndex = AddBody(BRepMetaData, Matrix, ParentId);
+	int32 BodyIndex = AddBody(BRepMetaData, Matrix, ParentId, BRepUnit);
 	FArchiveBody& Body = CADFileData.GetBodyAt(BodyIndex);
 
 	RepresentationItemsCache.Add(PolygonalPtr, BodyIndex);
@@ -1733,7 +1741,7 @@ void FTechSoftFileParser::ExtractGraphStyleProperties(uint32 StyleIndex, FCADUUI
 	}
 }
 
-FMatrix FTechSoftFileParser::ExtractTransformation3D(const A3DMiscTransformation* CartesianTransformation)
+FMatrix FTechSoftFileParser::ExtractTransformation3D(const A3DMiscTransformation* CartesianTransformation, double& InOutUnit)
 {
 	TUniqueTSObj<A3DMiscCartesianTransformationData> CartesianTransformationData(CartesianTransformation);
 
@@ -1745,9 +1753,19 @@ FMatrix FTechSoftFileParser::ExtractTransformation3D(const A3DMiscTransformation
 
 		FVector ZVector = XVector ^ YVector;
 
-		const A3DVector3dData& Scale = CartesianTransformationData->m_sScale;
+		Origin *= InOutUnit;
 
-		FMatrix Matrix(XVector * Scale.m_dX, YVector * Scale.m_dY, ZVector * Scale.m_dZ, FVector::Zero());
+		const A3DVector3dData& A3DScale = CartesianTransformationData->m_sScale;
+		FVector3d Scale(A3DScale.m_dX, A3DScale.m_dY, A3DScale.m_dZ);
+		double UniformScale = TechSoftFileParserImpl::ExtractUniformScale(Scale);
+
+		XVector *= Scale.X;
+		YVector *= Scale.Y;
+		ZVector *= Scale.Z;
+
+		InOutUnit *= UniformScale;
+
+		FMatrix Matrix(XVector, YVector, ZVector, FVector::Zero());
 
 		if (CartesianTransformationData->m_ucBehaviour & kA3DTransformationMirror)
 		{
@@ -1756,15 +1774,14 @@ FMatrix FTechSoftFileParser::ExtractTransformation3D(const A3DMiscTransformation
 			Matrix.M[2][2] *= -1;
 		}
 
-		Matrix.SetOrigin(Origin * FileUnit);
-
+		Matrix.SetOrigin(Origin);
 		return Matrix;
 	}
 
 	return FMatrix::Identity;
 }
 
-FMatrix FTechSoftFileParser::ExtractGeneralTransformation(const A3DMiscTransformation* GeneralTransformation)
+FMatrix FTechSoftFileParser::ExtractGeneralTransformation(const A3DMiscTransformation* GeneralTransformation, double& InOutUnit)
 {
 	TUniqueTSObj<A3DMiscGeneralTransformationData> GeneralTransformationData(GeneralTransformation);
 	if (GeneralTransformationData.IsValid())
@@ -1779,17 +1796,37 @@ FMatrix FTechSoftFileParser::ExtractGeneralTransformation(const A3DMiscTransform
 			}
 		}
 
-		for (Index = 0; Index < 3; ++Index, ++Index)
+		FTransform3d Transform(Matrix);
+		FVector3d Scale = Transform.GetScale3D();
+		if (Scale.Equals(FVector3d::OneVector, KINDA_SMALL_NUMBER))
 		{
-			Matrix.M[3][Index] *= FileUnit;
+			for (Index = 0; Index < 3; ++Index, ++Index)
+			{
+				Matrix.M[3][Index] *= InOutUnit;
+			}
+			return Matrix;
 		}
 
-		return Matrix;
+		FVector3d Translation = Transform.GetTranslation();
+
+		double UniformScale = TechSoftFileParserImpl::ExtractUniformScale(Scale);
+		InOutUnit *= UniformScale;
+
+		FQuat4d Rotation = Transform.GetRotation();
+
+		FTransform3d NewTransform;
+		NewTransform.SetScale3D(Scale);
+		NewTransform.SetRotation(Rotation);
+
+		FMatrix NewMatrix = NewTransform.ToMatrixWithScale();
+		NewMatrix.SetOrigin(Translation);
+
+		return NewMatrix;
 	}
 	return FMatrix::Identity;
 }
 
-FMatrix FTechSoftFileParser::ExtractTransformation(const A3DMiscTransformation* Transformation3D)
+FMatrix FTechSoftFileParser::ExtractTransformation(const A3DMiscTransformation* Transformation3D, double& InOutUnit)
 {
 	if (Transformation3D == NULL)
 	{
@@ -1801,21 +1838,21 @@ FMatrix FTechSoftFileParser::ExtractTransformation(const A3DMiscTransformation*
 
 	if (Type == kA3DTypeMiscCartesianTransformation)
 	{
-		return ExtractTransformation3D(Transformation3D);
+		return ExtractTransformation3D(Transformation3D, InOutUnit);
 	}
 	else if (Type == kA3DTypeMiscGeneralTransformation)
 	{
-		return ExtractGeneralTransformation(Transformation3D);
+		return ExtractGeneralTransformation(Transformation3D, InOutUnit);
 	}
 	return FMatrix::Identity;
 }
 
-FMatrix FTechSoftFileParser::ExtractCoordinateSystem(const A3DRiCoordinateSystem* CoordinateSystem)
+FMatrix FTechSoftFileParser::ExtractCoordinateSystem(const A3DRiCoordinateSystem* CoordinateSystem, double& InOutUnit)
 {
 	TUniqueTSObj<A3DRiCoordinateSystemData> CoordinateSystemData(CoordinateSystem);
 	if (CoordinateSystemData.IsValid())
 	{
-		return ExtractTransformation3D(CoordinateSystemData->m_pTransformation);
+		return ExtractTransformation3D(CoordinateSystemData->m_pTransformation, InOutUnit);
 	}
 	return FMatrix::Identity;
 }
diff --git a/Engine/Plugins/Enterprise/DatasmithCADImporter/Source/CADInterfaces/Private/TechSoftFileParser.h b/Engine/Plugins/Enterprise/DatasmithCADImporter/Source/CADInterfaces/Private/TechSoftFileParser.h
--- a/Engine/Plugins/Enterprise/DatasmithCADImporter/Source/CADInterfaces/Private/TechSoftFileParser.h
+++ b/Engine/Plugins/Enterprise/DatasmithCADImporter/Source/CADInterfaces/Private/TechSoftFileParser.h
@@ -97,16 +97,16 @@ private:
 
 	// Traverse ASM tree by starting from the model
 	ECADParsingResult TraverseModel();
-	void TraverseReference(const A3DAsmProductOccurrence* Reference, const FMatrix& ParentMatrix);
+	void TraverseReference(const A3DAsmProductOccurrence* Reference, const FMatrix& ParentMatrix, double ParentUnit);
 	bool IsConfigurationSet(const A3DAsmProductOccurrence* Occurrence);
-	void TraverseConfigurationSet(const A3DAsmProductOccurrence* ConfigurationSet);
-	FCadId TraverseOccurrence(const A3DAsmProductOccurrence* Occurrence);
+	void TraverseConfigurationSet(const A3DAsmProductOccurrence* ConfigurationSet, double ParentUnit);
+	FCadId TraverseOccurrence(const A3DAsmProductOccurrence* Occurrence, double ParentUnit);
 	void ProcessPrototype(const A3DAsmProductOccurrence* InPrototype, FEntityMetaData& OutMetaData, A3DMiscTransformation** OutLocation);
-	void TraversePartDefinition(const A3DAsmPartDefinition* PartDefinition, FArchiveComponent& Component);
-	FCadId TraverseRepresentationSet(const A3DRiSet* pSet, const FEntityMetaData& PartMetaData);
-	FCadId TraverseRepresentationItem(A3DRiRepresentationItem* RepresentationItem, const FEntityMetaData& PartMetaData, const FCadId ParentId);
-	FCadId TraverseBRepModel(A3DRiBrepModel* BrepModel, const FEntityMetaData& PartMetaData, const FCadId ParentId);
-	FCadId TraversePolyBRepModel(A3DRiPolyBrepModel* PolygonalBrepModel, const FEntityMetaData& PartMetaData, const FCadId ParentId);
+	void TraversePartDefinition(const A3DAsmPartDefinition* PartDefinition, FArchiveComponent& Component, double ParentUnit);
+	FCadId TraverseRepresentationSet(const A3DRiSet* pSet, const FEntityMetaData& PartMetaData, double ParentUnit);
+	FCadId TraverseRepresentationItem(A3DRiRepresentationItem* RepresentationItem, const FEntityMetaData& PartMetaData, const FCadId ParentId, double ParentUnit);
+	FCadId TraverseBRepModel(A3DRiBrepModel* BrepModel, const FEntityMetaData& PartMetaData, const FCadId ParentId, double ParentUnit);
+	FCadId TraversePolyBRepModel(A3DRiPolyBrepModel* PolygonalBrepModel, const FEntityMetaData& PartMetaData, const FCadId ParentId, double ParentUnit);
 
 	// MetaData
 	void ExtractSpecificMetaData(const A3DAsmProductOccurrence* Occurrence, FEntityMetaData& OutMetaData);
@@ -128,7 +128,7 @@ private:
 	FArchiveMaterial& FindOrAddMaterial(uint32 MaterialId, const A3DGraphStyleData& GraphStyleData);
 
 	/**
-	 * @param GraphMaterialIndex is the techsoft index of the graphic data
+	 * @param GraphMaterialIndex is the Techsoft index of the graphic data
 	 * @param MaterialIndexToSave is the index of the material really saved (i.e. for texture, at the texture index, with saved the material used by the texture)
 	 */
 	FArchiveMaterial& AddMaterialAt(uint32 MaterialIndexToSave, uint32 GraphMaterialIndex, const A3DGraphStyleData& GraphStyleData);
@@ -138,10 +138,10 @@ private:
 	}
 
 	// Transform
-	FMatrix ExtractCoordinateSystem(const A3DRiCoordinateSystem* CoordinateSystem);
-	FMatrix ExtractTransformation(const A3DMiscTransformation* Transformation3d);
-	FMatrix ExtractGeneralTransformation(const A3DMiscTransformation* GeneralTransformation);
-	FMatrix ExtractTransformation3D(const A3DMiscTransformation* CartesianTransformation);
+	FMatrix ExtractCoordinateSystem(const A3DRiCoordinateSystem* CoordinateSystem, double& InOutUnit);
+	FMatrix ExtractTransformation(const A3DMiscTransformation* Transformation3d, double& InOutUnit);
+	FMatrix ExtractGeneralTransformation(const A3DMiscTransformation* GeneralTransformation, double& InOutUnit);
+	FMatrix ExtractTransformation3D(const A3DMiscTransformation* CartesianTransformation, double& InOutUnit);
 
 	// Archive methods
 	FArchiveInstance& AddInstance(FEntityMetaData& InstanceMetaData);
@@ -149,7 +149,7 @@ private:
 	FArchiveUnloadedComponent& AddUnloadedComponent(FEntityMetaData& ComponentMetaData, FArchiveInstance& Instance);
 	FArchiveComponent& AddOccurence(FEntityMetaData& InstanceMetaData, FEntityMetaData& ReferenceMetaData, FCadId& OutComponentId);
 	FArchiveComponent& AddOccurence(FEntityMetaData& InstanceMetaData, FCadId& OutComponentId);
-	int32 AddBody(FEntityMetaData& BodyMetaData, const FMatrix& Matrix, const FCadId ParentId);
+	int32 AddBody(FEntityMetaData& BodyMetaData, const FMatrix& Matrix, const FCadId ParentId, double BodyUnit);
 #endif
 
 protected:
diff --git a/Engine/Plugins/Enterprise/DatasmithCADImporter/Source/CADInterfaces/Private/TechSoftInterface.cpp b/Engine/Plugins/Enterprise/DatasmithCADImporter/Source/CADInterfaces/Private/TechSoftInterface.cpp
--- a/Engine/Plugins/Enterprise/DatasmithCADImporter/Source/CADInterfaces/Private/TechSoftInterface.cpp
+++ b/Engine/Plugins/Enterprise/DatasmithCADImporter/Source/CADInterfaces/Private/TechSoftInterface.cpp
@@ -422,15 +422,28 @@ A3DStatus ExportModelFileToPrcFile(const A3DAsmModelFile* ModelFile, const A3DRW
 	return A3DAsmModelFileExportToPrcFile(ModelFile, ParamsExportData, CADFileName, PrcWriteHelper);
 }
 
-A3DStatus SewModel(A3DAsmModelFile* ModelPtr, double Tolerance, A3DSewOptionsData const* SewOptions)
+double GetModelFileUnit(const A3DAsmModelFile* ModelFile)
 {
-	return A3DAsmModelFileSew(&ModelPtr, Tolerance, SewOptions);
+	double FileUnit = 0.1;
+	if (A3DAsmModelFileGetUnit(ModelFile, &FileUnit) != A3DStatus::A3D_SUCCESS)
+	{
+		return 0.1;
+	}
+	return FileUnit * 0.1;
+}
+
+A3DStatus SewModel(A3DAsmModelFile* ModelPtr, double ToleranceInCM, A3DSewOptionsData const* SewOptions)
+{
+	const double ToleranceInMM = ToleranceInCM * 10.; // cm => mm
+	return A3DAsmModelFileSew(&ModelPtr, ToleranceInMM, SewOptions);
 }
 
-A3DStatus SewBReps(A3DRiBrepModel** BRepsToSew, uint32 const BRepCount, double Tolerance, A3DSewOptionsData const* SewOptions, A3DRiBrepModel*** OutNewBReps, uint32& OutNewBRepCount)
+A3DStatus SewBReps(A3DRiBrepModel** BRepsToSew, uint32 const BRepCount, double ToleranceInCM, double FileUnit, A3DSewOptionsData const* SewOptions, A3DRiBrepModel*** OutNewBReps, uint32& OutNewBRepCount)
 {
+	const double ToleranceInFileUnit = ToleranceInCM * 10. / FileUnit; 
+
 	A3DUns32 NewBRepCount = 0;
-	A3DStatus Status = A3DSewBrep(&BRepsToSew, BRepCount, Tolerance, SewOptions, OutNewBReps, &NewBRepCount);
+	A3DStatus Status = A3DSewBrep(&BRepsToSew, BRepCount, ToleranceInFileUnit, SewOptions, OutNewBReps, &NewBRepCount);
 	OutNewBRepCount = NewBRepCount;
 	return Status;
 }
diff --git a/Engine/Plugins/Enterprise/DatasmithCADImporter/Source/CADInterfaces/Private/TechSoftUtils.cpp b/Engine/Plugins/Enterprise/DatasmithCADImporter/Source/CADInterfaces/Private/TechSoftUtils.cpp
--- a/Engine/Plugins/Enterprise/DatasmithCADImporter/Source/CADInterfaces/Private/TechSoftUtils.cpp
+++ b/Engine/Plugins/Enterprise/DatasmithCADImporter/Source/CADInterfaces/Private/TechSoftUtils.cpp
@@ -65,12 +65,12 @@ bool GetBodyFromPcrFile(const FString& Filename, const FImportParameters& Import
 	TSharedPtr<FJsonObject> JsonObject = GetJsonObject(ModelFileData->m_ppPOccurrences[0]);
 	if (JsonObject.IsValid())
 	{
-		double FileUnit = 1.0;
-		JsonObject->TryGetNumberField(JSON_ENTRY_FILE_UNIT, FileUnit);
+		double BodyUnit = 1.0;
+		JsonObject->TryGetNumberField(JSON_ENTRY_FILE_UNIT, BodyUnit);
 
 		for (A3DUns32 Index = 0; Index < PartDefinitionData->m_uiRepItemsSize; ++Index)
 		{
-			if (!FillBodyMesh(PartDefinitionData->m_ppRepItems[Index], ImportParameters, FileUnit, BodyMesh))
+			if (!FillBodyMesh(PartDefinitionData->m_ppRepItems[Index], ImportParameters, BodyUnit, BodyMesh))
 			{
 				return false;
 			}
@@ -160,7 +160,7 @@ bool FillBodyMesh(void* BodyPtr, const FImportParameters& ImportParameters, doub
 
 	TessellationParameters->m_eTessellationLevelOfDetail = kA3DTessLODUserDefined; // Enum to specify predefined values for some following members.
 	TessellationParameters->m_bUseHeightInsteadOfRatio = A3D_TRUE;
-	TessellationParameters->m_dMaxChordHeight = ImportParameters.GetChordTolerance() * 10. / FileUnit;
+	TessellationParameters->m_dMaxChordHeight = ImportParameters.GetChordTolerance() * 10.; // cm to mm
 	TessellationParameters->m_dAngleToleranceDeg = ImportParameters.GetMaxNormalAngle();
 	TessellationParameters->m_dMaximalTriangleEdgeLength = 0; //ImportParameters.MaxEdgeLength;
 
diff --git a/Engine/Plugins/Enterprise/DatasmithCADImporter/Source/CADInterfaces/Public/CADSceneGraph.h b/Engine/Plugins/Enterprise/DatasmithCADImporter/Source/CADInterfaces/Public/CADSceneGraph.h
--- a/Engine/Plugins/Enterprise/DatasmithCADImporter/Source/CADInterfaces/Public/CADSceneGraph.h
+++ b/Engine/Plugins/Enterprise/DatasmithCADImporter/Source/CADInterfaces/Public/CADSceneGraph.h
@@ -82,6 +82,7 @@ public:
 public:
 	FCadId ParentId = 0;
 	FCADUUID MeshActorName = 0;
+	double BodyUnit = 1.;
 
 	TSet<FMaterialId> MaterialFaceSet;
 	TSet<FColorId> ColorFaceSet;
diff --git a/Engine/Plugins/Enterprise/DatasmithCADImporter/Source/CADInterfaces/Public/TechSoftInterface.h b/Engine/Plugins/Enterprise/DatasmithCADImporter/Source/CADInterfaces/Public/TechSoftInterface.h
--- a/Engine/Plugins/Enterprise/DatasmithCADImporter/Source/CADInterfaces/Public/TechSoftInterface.h
+++ b/Engine/Plugins/Enterprise/DatasmithCADImporter/Source/CADInterfaces/Public/TechSoftInterface.h
@@ -134,8 +134,17 @@ CADINTERFACES_API A3DUns32 InsertGraphRgbColor(const A3DGraphRgbColorData& InRgb
 CADINTERFACES_API A3DUns32 InsertGraphMaterial(const A3DGraphMaterialData& InMaterialData);
 CADINTERFACES_API A3DUns32 InsertGraphStyle(const A3DGraphStyleData& InStyleData);
 
-CADINTERFACES_API A3DStatus SewModel(A3DAsmModelFile* ModelPtr, double Tolerance, A3DSewOptionsData const* SewOptions);
-CADINTERFACES_API A3DStatus SewBReps(A3DRiBrepModel** BRepsToSew, uint32 const BRepCount, double Tolerance, A3DSewOptionsData const* SewOptions, A3DRiBrepModel*** OutNewBReps, uint32& OutNewBRepCount);
+CADINTERFACES_API double GetModelFileUnit(const A3DAsmModelFile* pModelFile);
+
+/**
+ * @param ToleranceInCM : The maximum tolerance for the sewing (in cm).
+ */
+CADINTERFACES_API A3DStatus SewModel(A3DAsmModelFile* ModelPtr, double ToleranceInCM, A3DSewOptionsData const* SewOptions);
+
+/**
+ * @param ToleranceInCM : The maximum tolerance for the sewing (in cm).
+ */
+CADINTERFACES_API A3DStatus SewBReps(A3DRiBrepModel** BRepsToSew, uint32 const BRepCount, double ToleranceInCM, double FileUnit, A3DSewOptionsData const* SewOptions, A3DRiBrepModel*** OutNewBReps, uint32& OutNewBRepCount);
 
 #endif
 
diff --git a/Engine/Plugins/Enterprise/DatasmithCADImporter/Source/CADInterfaces/Public/TechSoftUtils.h b/Engine/Plugins/Enterprise/DatasmithCADImporter/Source/CADInterfaces/Public/TechSoftUtils.h
--- a/Engine/Plugins/Enterprise/DatasmithCADImporter/Source/CADInterfaces/Public/TechSoftUtils.h
+++ b/Engine/Plugins/Enterprise/DatasmithCADImporter/Source/CADInterfaces/Public/TechSoftUtils.h
@@ -6,7 +6,7 @@
 #include "CADData.h"
 #include "TechSoftInterface.h"
 
-#define JSON_ENTRY_FILE_UNIT		TEXT("FileUnit")
+#define JSON_ENTRY_BODY_UNIT		TEXT("BodyUnit")
 #define JSON_ENTRY_COLOR_NAME		TEXT("ColorName")
 #define JSON_ENTRY_MATERIAL_NAME	TEXT("MaterialName")
 
diff --git a/Engine/Plugins/Enterprise/DatasmithCADImporter/Source/CADTools/Private/CADToolsModule.cpp b/Engine/Plugins/Enterprise/DatasmithCADImporter/Source/CADTools/Private/CADToolsModule.cpp
--- a/Engine/Plugins/Enterprise/DatasmithCADImporter/Source/CADTools/Private/CADToolsModule.cpp
+++ b/Engine/Plugins/Enterprise/DatasmithCADImporter/Source/CADTools/Private/CADToolsModule.cpp
@@ -18,7 +18,7 @@ bool FCADToolsModule::IsAvailable()
 
 uint32 FCADToolsModule::GetCacheVersion()
 {
-	return 500;
+	return 501;
 }
 
 IMPLEMENT_MODULE(FCADToolsModule, CADTools);
diff --git a/Engine/Plugins/Enterprise/DatasmithCADImporter/Source/CADTools/Public/CADOptions.h b/Engine/Plugins/Enterprise/DatasmithCADImporter/Source/CADTools/Public/CADOptions.h
--- a/Engine/Plugins/Enterprise/DatasmithCADImporter/Source/CADTools/Public/CADOptions.h
+++ b/Engine/Plugins/Enterprise/DatasmithCADImporter/Source/CADTools/Public/CADOptions.h
@@ -58,7 +58,7 @@ namespace CADLibrary
 
 
 	public:
-		FImportParameters(double InMetricUnit = 0.001, double InScaleFactor = 0.1, FDatasmithUtils::EModelCoordSystem NewCoordinateSystem = FDatasmithUtils::EModelCoordSystem::ZUp_RightHanded)
+		FImportParameters(double InMetricUnit = 0.001, double InScaleFactor = 1., FDatasmithUtils::EModelCoordSystem NewCoordinateSystem = FDatasmithUtils::EModelCoordSystem::ZUp_RightHanded)
 			: MetricUnit(InMetricUnit)
 			, ScaleFactor(InScaleFactor)
 			, ChordTolerance(0.2)
diff --git a/Engine/Plugins/Enterprise/DatasmithCADImporter/Source/ParametricSurface/Private/TechSoft/CADModelToTechSoftConverterBase.cpp b/Engine/Plugins/Enterprise/DatasmithCADImporter/Source/ParametricSurface/Private/TechSoft/CADModelToTechSoftConverterBase.cpp
--- a/Engine/Plugins/Enterprise/DatasmithCADImporter/Source/ParametricSurface/Private/TechSoft/CADModelToTechSoftConverterBase.cpp
+++ b/Engine/Plugins/Enterprise/DatasmithCADImporter/Source/ParametricSurface/Private/TechSoft/CADModelToTechSoftConverterBase.cpp
@@ -23,7 +23,7 @@ bool FCADModelToTechSoftConverterBase::RepairTopology()
 
 		A3DRiBrepModel** OutNewBReps;
 		uint32 OutNewBRepCount;
-		CADLibrary::TechSoftInterface::SewBReps((A3DRiBrepModel**) RiRepresentationItems.GetData(), RiRepresentationItems.Num(), 0.01, SewOptionsData.GetPtr(), &OutNewBReps, OutNewBRepCount);
+		CADLibrary::TechSoftInterface::SewBReps((A3DRiBrepModel**) RiRepresentationItems.GetData(), RiRepresentationItems.Num(), 0.01, 1., SewOptionsData.GetPtr(), &OutNewBReps, OutNewBRepCount);
 
 		RiRepresentationItems.Empty(OutNewBRepCount);
 		for (uint32 Index = 0; Index < OutNewBRepCount; ++Index)
