Fix retessellate with Stitching = Heal
Fix also CVar ds.CADTranslator.StitchingTolerance: this CVar was not associated to a static variable of FImportParameters, so was not transfered to the workers
#swarm https://p4-swarm.epicgames.net/reviews/19384076

diff --git a/Engine/Plugins/Enterprise/DatasmithCADImporter/Source/CADInterfaces/Private/TechSoftFileParser.cpp b/Engine/Plugins/Enterprise/DatasmithCADImporter/Source/CADInterfaces/Private/TechSoftFileParser.cpp
--- a/Engine/Plugins/Enterprise/DatasmithCADImporter/Source/CADInterfaces/Private/TechSoftFileParser.cpp
+++ b/Engine/Plugins/Enterprise/DatasmithCADImporter/Source/CADInterfaces/Private/TechSoftFileParser.cpp
@@ -12,14 +12,6 @@
 #include "Serialization/JsonWriter.h"
 #include "Templates/UnrealTemplate.h"
 
-float GStitchingTolerance = 0.001f;
-FAutoConsoleVariableRef GCADTranslatorStitchingTolerance(
-	TEXT("ds.CADTranslator.StitchingTolerance"),
-	GStitchingTolerance,
-	TEXT("Welding threshold for Heal/Sew stitching methods in cm\n\
-Default value of StitchingTolerance is 0.001 cm\n"),
-	ECVF_Default);
-
 namespace CADLibrary
 {
 
@@ -352,7 +344,7 @@ void FTechSoftFileParser::SewModel()
 	CADLibrary::TUniqueTSObj<A3DSewOptionsData> SewData;
 	SewData->m_bComputePreferredOpenShellOrientation = false;
 	
-	TechSoftInterface::SewModel(ModelFile.Get(), GStitchingTolerance, SewData.GetPtr());
+	TechSoftInterface::SewModel(ModelFile.Get(), CADLibrary::FImportParameters::GStitchingTolerance, SewData.GetPtr());
 }
 
 
@@ -377,7 +369,8 @@ void FTechSoftFileParser::GenerateBodyMesh(A3DRiRepresentationItem* Representati
 	{
 		TUniqueTSObj<A3DSewOptionsData> SewData;
 		SewData->m_bComputePreferredOpenShellOrientation = false;
-		A3DStatus Status = TechSoftInterface::SewBReps(&Representation, 1, GStitchingTolerance, FileUnit, SewData.GetPtr(), &NewBReps, NewBRepCount);
+		const uint32 BRepCount = 1;
+		A3DStatus Status = TechSoftInterface::SewBReps(&Representation, BRepCount, CADLibrary::FImportParameters::GStitchingTolerance, FileUnit, SewData.GetPtr(), &NewBReps, NewBRepCount);
 		if (Status != A3DStatus::A3D_SUCCESS)
 		{
 			CADFileData.AddWarningMessages(TEXT("A body healing failed. A body could be missing."));
diff --git a/Engine/Plugins/Enterprise/DatasmithCADImporter/Source/CADInterfaces/Private/TechSoftUtils.cpp b/Engine/Plugins/Enterprise/DatasmithCADImporter/Source/CADInterfaces/Private/TechSoftUtils.cpp
--- a/Engine/Plugins/Enterprise/DatasmithCADImporter/Source/CADInterfaces/Private/TechSoftUtils.cpp
+++ b/Engine/Plugins/Enterprise/DatasmithCADImporter/Source/CADInterfaces/Private/TechSoftUtils.cpp
@@ -44,6 +44,14 @@ bool GetBodyFromPcrFile(const FString& Filename, const FImportParameters& Import
 		return false;
 	}
 
+	if(ImportParameters.GetStitchingTechnique() != CADLibrary::EStitchingTechnique::StitchingNone)
+	{
+		CADLibrary::TUniqueTSObj<A3DSewOptionsData> SewData;
+		SewData->m_bComputePreferredOpenShellOrientation = false;
+
+		TechSoftInterface::SewModel(ModelFile.Get(), CADLibrary::FImportParameters::GStitchingTolerance, SewData.GetPtr());
+	}
+
 	TUniqueTSObj<A3DAsmModelFileData> ModelFileData(ModelFile.Get());
 	if (!ModelFileData.IsValid() || ModelFileData->m_uiPOccurrencesSize == 0)
 	{
@@ -161,6 +169,11 @@ bool FillBodyMesh(void* BodyPtr, const FImportParameters& ImportParameters, doub
 	TessellationParameters->m_eTessellationLevelOfDetail = kA3DTessLODUserDefined; // Enum to specify predefined values for some following members.
 	TessellationParameters->m_bUseHeightInsteadOfRatio = A3D_TRUE;
 	TessellationParameters->m_dMaxChordHeight = ImportParameters.GetChordTolerance() * 10.; // cm to mm
+	if (!FMath::IsNearlyZero(FileUnit))
+	{
+		TessellationParameters->m_dMaxChordHeight /= FileUnit;
+	}
+
 	TessellationParameters->m_dAngleToleranceDeg = ImportParameters.GetMaxNormalAngle();
 	TessellationParameters->m_dMaximalTriangleEdgeLength = 0; //ImportParameters.MaxEdgeLength;
 
diff --git a/Engine/Plugins/Enterprise/DatasmithCADImporter/Source/CADTools/Private/CADOption.cpp b/Engine/Plugins/Enterprise/DatasmithCADImporter/Source/CADTools/Private/CADOption.cpp
--- a/Engine/Plugins/Enterprise/DatasmithCADImporter/Source/CADTools/Private/CADOption.cpp
+++ b/Engine/Plugins/Enterprise/DatasmithCADImporter/Source/CADTools/Private/CADOption.cpp
@@ -64,6 +64,14 @@ FAutoConsoleVariableRef GCADTranslatorJtFileEmbeddedTessellation(
 	TEXT("If both (tessellation and BRep) exist in the file, import embedded tessellation instead of meshing BRep.\n"),
 	ECVF_Default);
 
+float FImportParameters::GStitchingTolerance = 0.001f;
+FAutoConsoleVariableRef GCADTranslatorStitchingTolerance(
+	TEXT("ds.CADTranslator.StitchingTolerance"),
+	FImportParameters::GStitchingTolerance,
+	TEXT("Welding threshold for Heal/Sew stitching methods in cm\n\
+Default value of StitchingTolerance is 0.001 cm\n"),
+ECVF_Default);
+
 uint32 GetTypeHash(const FImportParameters& ImportParameters)
 {
 	uint32 ParametersHash = ::GetTypeHash(ImportParameters.bGDisableCADKernelTessellation);
diff --git a/Engine/Plugins/Enterprise/DatasmithCADImporter/Source/CADTools/Public/CADOptions.h b/Engine/Plugins/Enterprise/DatasmithCADImporter/Source/CADTools/Public/CADOptions.h
--- a/Engine/Plugins/Enterprise/DatasmithCADImporter/Source/CADTools/Public/CADOptions.h
+++ b/Engine/Plugins/Enterprise/DatasmithCADImporter/Source/CADTools/Public/CADOptions.h
@@ -54,9 +54,9 @@ namespace CADLibrary
 		CADTOOLS_API static bool bGEnableCADCache;
 		CADTOOLS_API static bool bGOverwriteCache;
 		CADTOOLS_API static bool bGPreferJtFileEmbeddedTessellation;
+		CADTOOLS_API static float GStitchingTolerance;
 		CADTOOLS_API static FString GCADLibrary;
 
-
 	public:
 		FImportParameters(double InMetricUnit = 0.001, double InScaleFactor = 1., FDatasmithUtils::EModelCoordSystem NewCoordinateSystem = FDatasmithUtils::EModelCoordSystem::ZUp_RightHanded)
 			: MetricUnit(InMetricUnit)
@@ -129,6 +129,7 @@ namespace CADLibrary
 			Ar << ImportParameters.bGEnableTimeControl;
 			Ar << ImportParameters.bGEnableCADCache;
 			Ar << ImportParameters.bGPreferJtFileEmbeddedTessellation;
+			Ar << ImportParameters.GStitchingTolerance;
 			Ar << ImportParameters.GCADLibrary;
 			return Ar;
 		}
