Alias UV is badly built for symmetric patch
In a symmetrical triangle, the order of the vertices is reversed with respect to the initial triangle. It was not done for UV index.

Improve also:
- tessellation parameter unit => As TechSoft works in mm, BRep is converted into mm
#swarm https://p4-swarm.epicgames.net/reviews/19384163

diff --git a/Engine/Plugins/Enterprise/DatasmithCADImporter/Source/CADLibrary/Private/MeshDescriptionHelper.cpp b/Engine/Plugins/Enterprise/DatasmithCADImporter/Source/CADLibrary/Private/MeshDescriptionHelper.cpp
--- a/Engine/Plugins/Enterprise/DatasmithCADImporter/Source/CADLibrary/Private/MeshDescriptionHelper.cpp
+++ b/Engine/Plugins/Enterprise/DatasmithCADImporter/Source/CADLibrary/Private/MeshDescriptionHelper.cpp
@@ -528,7 +528,7 @@ namespace CADLibrary
 						for (int32 VertexIndex = 0; VertexIndex < VertexCountPerFace; VertexIndex++)
 						{
 							const FVertexInstanceID& VertexInstanceID = MeshVertexInstanceIDs[FaceIndex + Orientation[VertexIndex]];
-							VertexInstanceUVs.Set(VertexInstanceID, UVChannel, FVector2f(Tessellation.TexCoordArray[NewFaceIndex[FaceIndex + VertexIndex]]));
+							VertexInstanceUVs.Set(VertexInstanceID, UVChannel, FVector2f(Tessellation.TexCoordArray[NewFaceIndex[FaceIndex + Orientation[VertexIndex]]]));
 						}
 					}
 				}
diff --git a/Engine/Plugins/Enterprise/DatasmithCADImporter/Source/DatasmithWireTranslator/Private/AliasModelToTechSoftConverter.cpp b/Engine/Plugins/Enterprise/DatasmithCADImporter/Source/DatasmithWireTranslator/Private/AliasModelToTechSoftConverter.cpp
--- a/Engine/Plugins/Enterprise/DatasmithCADImporter/Source/DatasmithWireTranslator/Private/AliasModelToTechSoftConverter.cpp
+++ b/Engine/Plugins/Enterprise/DatasmithCADImporter/Source/DatasmithWireTranslator/Private/AliasModelToTechSoftConverter.cpp
@@ -121,9 +121,9 @@ A3DSurfBase* AddNURBSSurface(const Surface_T& AliasSurface, EAliasObjectReferenc
 	
 	TFunction<void(const double*, A3DVector3dData&, A3DDouble&)> SetA3DPole = [](const double* HomogeneousPole, A3DVector3dData & OutTSPoint, A3DDouble & OutWeight)
 	{
-		OutTSPoint.m_dX = HomogeneousPole[0];
-		OutTSPoint.m_dY = HomogeneousPole[1];
-		OutTSPoint.m_dZ = HomogeneousPole[2];
+		OutTSPoint.m_dX = HomogeneousPole[0] * 10.;  // cm (Alias MetricUnit) to mm
+		OutTSPoint.m_dY = HomogeneousPole[1] * 10.;
+		OutTSPoint.m_dZ = HomogeneousPole[2] * 10.;
 		OutWeight = HomogeneousPole[3];
 	};
 
diff --git a/Engine/Plugins/Enterprise/DatasmithCADImporter/Source/DatasmithWireTranslator/Private/DatasmithWireTranslator.cpp b/Engine/Plugins/Enterprise/DatasmithCADImporter/Source/DatasmithWireTranslator/Private/DatasmithWireTranslator.cpp
--- a/Engine/Plugins/Enterprise/DatasmithCADImporter/Source/DatasmithWireTranslator/Private/DatasmithWireTranslator.cpp
+++ b/Engine/Plugins/Enterprise/DatasmithCADImporter/Source/DatasmithWireTranslator/Private/DatasmithWireTranslator.cpp
@@ -227,7 +227,9 @@ public:
 			DatasmithScene->SetProductVersion(TEXT("Alias 2022.2"));
 		}
 
-		CADLibrary::FImportParameters ImportParameters(0.01, 1);
+		const double MetricUnit = 0.001; // ImportParameters MetricUnit is defined in meter
+		const double ScaleFactor = 0.1;  // MetricUnit to UE unit i.e. cm
+		CADLibrary::FImportParameters ImportParameters(MetricUnit, ScaleFactor);
 		if(CADLibrary::FImportParameters::bGDisableCADKernelTessellation)
 		{
 			if (CADLibrary::FImportParameters::GCADLibrary == TEXT("KernelIO"))
diff --git a/Engine/Plugins/Enterprise/DatasmithCADImporter/Source/ParametricSurface/Private/TechSoft/CADModelToTechSoftConverterBase.cpp b/Engine/Plugins/Enterprise/DatasmithCADImporter/Source/ParametricSurface/Private/TechSoft/CADModelToTechSoftConverterBase.cpp
--- a/Engine/Plugins/Enterprise/DatasmithCADImporter/Source/ParametricSurface/Private/TechSoft/CADModelToTechSoftConverterBase.cpp
+++ b/Engine/Plugins/Enterprise/DatasmithCADImporter/Source/ParametricSurface/Private/TechSoft/CADModelToTechSoftConverterBase.cpp
@@ -23,7 +23,9 @@ bool FCADModelToTechSoftConverterBase::RepairTopology()
 
 		A3DRiBrepModel** OutNewBReps;
 		uint32 OutNewBRepCount;
-		CADLibrary::TechSoftInterface::SewBReps((A3DRiBrepModel**) RiRepresentationItems.GetData(), RiRepresentationItems.Num(), 0.01, 1., SewOptionsData.GetPtr(), &OutNewBReps, OutNewBRepCount);
+		const double SewTolerance = CADLibrary::FImportParameters::GStitchingTolerance;
+		const double FileUnit = ImportParameters.GetMetricUnit() * 1000; // ImportParameters MetricUnit is defined in meter, CADLibrary::TechSoftInterface::SewBReps expectes it in mm
+		CADLibrary::TechSoftInterface::SewBReps((A3DRiBrepModel**) RiRepresentationItems.GetData(), RiRepresentationItems.Num(), SewTolerance, FileUnit, SewOptionsData.GetPtr(), &OutNewBReps, OutNewBRepCount);
 
 		RiRepresentationItems.Empty(OutNewBRepCount);
 		for (uint32 Index = 0; Index < OutNewBRepCount; ++Index)
@@ -76,7 +78,8 @@ bool FCADModelToTechSoftConverterBase::Tessellate(const CADLibrary::FMeshParamet
 #ifdef USE_TECHSOFT_SDK
 	for (A3DRiRepresentationItem* Representation : RiRepresentationItems)
 	{
-		CADLibrary::TechSoftUtils::FillBodyMesh(Representation, ImportParameters, 1.0, BodyMesh);
+		const double FileUnit = ImportParameters.GetMetricUnit() * 1000; // ImportParameters MetricUnit is defined in meter, CADLibrary::TechSoftInterface::SewBReps expectes it in mm
+		CADLibrary::TechSoftUtils::FillBodyMesh(Representation, ImportParameters, FileUnit, BodyMesh);
 	}
 #endif
 
