Fixed remaining use of L"" in techSoft code

diff --git a/Engine/Plugins/Enterprise/DatasmithCADImporter/Source/CADInterfaces/Private/TechSoftFileParser.cpp b/Engine/Plugins/Enterprise/DatasmithCADImporter/Source/CADInterfaces/Private/TechSoftFileParser.cpp
--- a/Engine/Plugins/Enterprise/DatasmithCADImporter/Source/CADInterfaces/Private/TechSoftFileParser.cpp
+++ b/Engine/Plugins/Enterprise/DatasmithCADImporter/Source/CADInterfaces/Private/TechSoftFileParser.cpp
@@ -1566,13 +1566,13 @@ void FTechSoftFileParser::ExtractSpecificMetaData(const A3DAsmProductOccurrence*
 			if (CatiaV5SpecificData->m_psVersion)
 			{
 				FString Version = UTF8_TO_TCHAR(CatiaV5SpecificData->m_psVersion);
-				OutMetaData.MetaData.Emplace(L"CatiaVersion", Version);
+				OutMetaData.MetaData.Emplace(TEXT("CatiaVersion"), Version);
 			}
 
 			if (CatiaV5SpecificData->m_psPartNumber)
 			{
 				FString PartNumber = UTF8_TO_TCHAR(CatiaV5SpecificData->m_psPartNumber);
-				OutMetaData.MetaData.Emplace(L"CatiaPartNumber", PartNumber);
+				OutMetaData.MetaData.Emplace(TEXT("CatiaPartNumber"), PartNumber);
 			}
 		}
 		break;
