"ds.CADTranslator.DisableCADKernelTessellation" CVar is removed to disable access to CADKernel module
CADKernel module is not stable enough for 5.0.
This module will be reactivated for 5.0.1.

diff --git a/Engine/Plugins/Enterprise/DatasmithCADImporter/Source/CADTools/Private/CADOption.cpp b/Engine/Plugins/Enterprise/DatasmithCADImporter/Source/CADTools/Private/CADOption.cpp
--- a/Engine/Plugins/Enterprise/DatasmithCADImporter/Source/CADTools/Private/CADOption.cpp
+++ b/Engine/Plugins/Enterprise/DatasmithCADImporter/Source/CADTools/Private/CADOption.cpp
@@ -27,11 +27,13 @@ Default is MaxImportThreads = 0\n\
 	ECVF_Default);
 
 bool FImportParameters::bGDisableCADKernelTessellation = true;
+#ifdef ENABLE_IN_5_0_1
 FAutoConsoleVariableRef GCADTranslatorDisableCADKernelTessellation(
 	TEXT("ds.CADTranslator.DisableCADKernelTessellation"),
 	FImportParameters::bGDisableCADKernelTessellation,
 	TEXT("Disable to use CAD import library tessellator.\n"),
 	ECVF_Default);
+#endif
 
 bool FImportParameters::bGEnableCADCache = true;
 FAutoConsoleVariableRef GCADTranslatorEnableCADCache(
