CADKernel imports Meshes in UE in mm instead of cm
- A scale is apply at the mesh during the transfer to UE
- Alias/CADKernel: BRep body is transferred to CADKernel  in cm instead of mm
- Rhino/CADKernel: fix nullptr
#swarm https://p4-swarm.epicgames.net/reviews/19409797

diff --git a/Engine/Plugins/Enterprise/DatasmithCADImporter/Source/CADLibrary/Private/CADKernelTools.cpp b/Engine/Plugins/Enterprise/DatasmithCADImporter/Source/CADLibrary/Private/CADKernelTools.cpp
--- a/Engine/Plugins/Enterprise/DatasmithCADImporter/Source/CADLibrary/Private/CADKernelTools.cpp
+++ b/Engine/Plugins/Enterprise/DatasmithCADImporter/Source/CADLibrary/Private/CADKernelTools.cpp
@@ -342,6 +342,11 @@ namespace CADLibrary
 	{
 		ModelMesh.GetNodeCoordinates(OutBodyMesh.VertexArray);
 
+		for (FVector& Vertex : OutBodyMesh.VertexArray)
+		{
+			Vertex *= 0.1;
+		}
+
 		uint32 FaceSize = Body.FaceCount();
 
 		// Allocate memory space for tessellation data
diff --git a/Engine/Plugins/Enterprise/DatasmithCADImporter/Source/DatasmithOpenNurbsTranslator/Private/OpenNurbsBRepToCADKernelConverter.cpp b/Engine/Plugins/Enterprise/DatasmithCADImporter/Source/DatasmithOpenNurbsTranslator/Private/OpenNurbsBRepToCADKernelConverter.cpp
--- a/Engine/Plugins/Enterprise/DatasmithCADImporter/Source/DatasmithOpenNurbsTranslator/Private/OpenNurbsBRepToCADKernelConverter.cpp
+++ b/Engine/Plugins/Enterprise/DatasmithCADImporter/Source/DatasmithOpenNurbsTranslator/Private/OpenNurbsBRepToCADKernelConverter.cpp
@@ -176,7 +176,7 @@ void FOpenNurbsBRepToCADKernelConverter::LinkEdgesLoop(const ON_BrepLoop& OpenNu
 			}
 
 			TSharedPtr<CADKernel::FTopologicalEdge>* TwinEdge = OpenNurbsTrimId2CADKernelEdge.Find(LinkedEdgeId);
-			if (TwinEdge->IsValid() && !(*TwinEdge)->IsDeleted() && !(*TwinEdge)->IsDegenerated())
+			if (TwinEdge != nullptr && TwinEdge->IsValid() && !(*TwinEdge)->IsDeleted() && !(*TwinEdge)->IsDegenerated())
 			{
 				(*Edge)->Link(**TwinEdge, SquareTolerance);
 				break;
diff --git a/Engine/Plugins/Enterprise/DatasmithCADImporter/Source/DatasmithWireTranslator/Private/AliasModelToCADKernelConverter.cpp b/Engine/Plugins/Enterprise/DatasmithCADImporter/Source/DatasmithWireTranslator/Private/AliasModelToCADKernelConverter.cpp
--- a/Engine/Plugins/Enterprise/DatasmithCADImporter/Source/DatasmithWireTranslator/Private/AliasModelToCADKernelConverter.cpp
+++ b/Engine/Plugins/Enterprise/DatasmithCADImporter/Source/DatasmithWireTranslator/Private/AliasModelToCADKernelConverter.cpp
@@ -89,6 +89,15 @@ namespace AliasToCADKernelUtils
 			AliasSurface.CVsUnaffectedPositionInclMultiples(NURBSData.HomogeneousPoles.GetData());
 		}
 
+		// convert cm to mm
+		double* Poles = NURBSData.HomogeneousPoles.GetData();
+		for (int32 Index = 0; Index < NURBSData.HomogeneousPoles.Num(); Index +=4)
+		{
+			Poles[Index + 0] *= 10.;
+			Poles[Index + 1] *= 10.;
+			Poles[Index + 2] *= 10.;
+		}
+
 		return CADKernel::FEntity::MakeShared<CADKernel::FNURBSSurface>(GeometricTolerance, NURBSData);
 	}
 }
