The algorithm of UV Scaling is reimplemented based on TechSoft SDK
#swarm https://p4-swarm.epicgames.net/reviews/19435639

diff --git a/Engine/Plugins/Enterprise/DatasmithCADImporter/Source/CADInterfaces/Private/TechSoftFileParser.cpp b/Engine/Plugins/Enterprise/DatasmithCADImporter/Source/CADInterfaces/Private/TechSoftFileParser.cpp
--- a/Engine/Plugins/Enterprise/DatasmithCADImporter/Source/CADInterfaces/Private/TechSoftFileParser.cpp
+++ b/Engine/Plugins/Enterprise/DatasmithCADImporter/Source/CADInterfaces/Private/TechSoftFileParser.cpp
@@ -434,9 +434,9 @@ void FTechSoftFileParser::GenerateBodyMesh(A3DRiRepresentationItem* Representati
 		{
 			TSharedPtr<FJsonObject> JsonObject = MakeShared<FJsonObject>();
 
-			// Save file unit and default color and material attributes in a json string
+			// Save body unit and default color and material attributes in a json string
 			// This will be used when the file is reloaded
-			JsonObject->SetNumberField(JSON_ENTRY_FILE_UNIT, Body.BodyUnit);
+			JsonObject->SetNumberField(JSON_ENTRY_BODY_UNIT, Body.BodyUnit);
 			JsonObject->SetNumberField(JSON_ENTRY_COLOR_NAME, DefaultColorName);
 			JsonObject->SetNumberField(JSON_ENTRY_MATERIAL_NAME, DefaultMaterialName);
 
diff --git a/Engine/Plugins/Enterprise/DatasmithCADImporter/Source/CADInterfaces/Private/TechSoftInterface.cpp b/Engine/Plugins/Enterprise/DatasmithCADImporter/Source/CADInterfaces/Private/TechSoftInterface.cpp
--- a/Engine/Plugins/Enterprise/DatasmithCADImporter/Source/CADInterfaces/Private/TechSoftInterface.cpp
+++ b/Engine/Plugins/Enterprise/DatasmithCADImporter/Source/CADInterfaces/Private/TechSoftInterface.cpp
@@ -156,6 +156,16 @@ A3DStatus GetSurfaceAsNurbs(const A3DSurfBase* SurfacePtr, A3DSurfNurbsData* Out
 	return A3DSurfBaseGetAsNurbs(SurfacePtr, Tolerance, bUseSameParameterization, OutDataPtr);
 }
 
+A3DStatus GetSurfaceDomain(const A3DSurfBase* SurfacePtr, A3DDomainData& OutDomain)
+{
+	return A3DSrfGetDomain(SurfacePtr, &OutDomain);
+}
+
+A3DStatus Evaluate(const A3DSurfBase* SurfacePtr, const A3DVector2dData& UVParameter, A3DUns32 Derivatives, A3DVector3dData* OutPointAndDerivatives)
+{
+	return A3DSurfEvaluate(SurfacePtr, &UVParameter, Derivatives, OutPointAndDerivatives);
+}
+
 A3DStatus GetCurveAsNurbs(const A3DCrvBase* CurvePtr, A3DCrvNurbsData* OutDataPtr, A3DDouble Tolerance, A3DBool bUseSameParameterization)
 {
 	return A3DCrvBaseGetAsNurbs(CurvePtr, Tolerance, bUseSameParameterization, OutDataPtr);
@@ -432,6 +442,11 @@ double GetModelFileUnit(const A3DAsmModelFile* ModelFile)
 	return FileUnit * 0.1;
 }
 
+A3DStatus AddAttribute(A3DEntity* EntityPtr, const TCHAR* Title, const TCHAR* Value)
+{
+	return A3DRootBaseAttributeAdd(EntityPtr, TCHAR_TO_UTF8(Title), TCHAR_TO_UTF8(Value));
+}
+
 A3DStatus SewModel(A3DAsmModelFile* ModelPtr, double ToleranceInCM, A3DSewOptionsData const* SewOptions)
 {
 	const double ToleranceInMM = ToleranceInCM * 10.; // cm => mm
@@ -541,6 +556,11 @@ void TUniqueTSObj<A3DCrvPolyLineData>::InitializeData()
 	A3D_INITIALIZE_DATA(A3DCrvPolyLineData, Data);
 }
 template<>
+void TUniqueTSObj<A3DDomainData>::InitializeData()
+{
+	A3D_INITIALIZE_DATA(A3DDomainData, Data);
+}
+template<>
 void TUniqueTSObj<A3DCrvTransformData>::InitializeData()
 {
 	A3D_INITIALIZE_DATA(A3DCrvTransformData, Data);
@@ -805,7 +825,16 @@ void TUniqueTSObj<A3DUTF8Char*>::InitializeData()
 {
 	Data = nullptr;
 }
-
+template<>
+void TUniqueTSObj<A3DVector2dData>::InitializeData()
+{
+	A3D_INITIALIZE_DATA(A3DVector2dData, Data);
+}
+template<>
+void TUniqueTSObj<A3DVector3dData>::InitializeData()
+{
+	A3D_INITIALIZE_DATA(A3DVector3dData, Data);
+}
 
 
 // TUniqueTSObjFromIndex InitializeData -----------------------------------
@@ -946,6 +975,12 @@ A3DStatus TUniqueTSObj<A3DCrvPolyLineData>::GetData(const A3DEntity* InEntityPtr
 	return A3DCrvPolyLineGet(InEntityPtr, &Data);
 }
 
+template<>
+A3DStatus TUniqueTSObj<A3DDomainData>::GetData(const A3DEntity* InEntityPtr)
+{
+	return A3DStatus::A3D_ERROR;
+}
+
 template<>
 A3DStatus TUniqueTSObj<A3DCrvTransformData>::GetData(const A3DEntity* InEntityPtr)
 {
@@ -1264,6 +1299,18 @@ A3DStatus TUniqueTSObj<A3DUTF8Char*>::GetData(const A3DEntity* InEntityPtr)
 	return A3DStatus::A3D_ERROR;
 }
 
+template<>
+A3DStatus TUniqueTSObj<A3DVector2dData>::GetData(const A3DEntity* InEntityPtr)
+{
+	return A3DStatus::A3D_ERROR;
+}
+
+template<>
+A3DStatus TUniqueTSObj<A3DVector3dData>::GetData(const A3DEntity* InEntityPtr)
+{
+	return A3DStatus::A3D_ERROR;
+}
+
 
 // TUniqueTSObjFromIndex GetData -----------------------------------
 template<>
@@ -1333,6 +1380,8 @@ const A3DEntity* TUniqueTSObj<A3DCrvParabolaData>::GetDefaultIndexerValue() cons
 template<>
 const A3DEntity* TUniqueTSObj<A3DCrvPolyLineData>::GetDefaultIndexerValue() const { return nullptr; }
 template<>
+const A3DEntity* TUniqueTSObj<A3DDomainData>::GetDefaultIndexerValue() const { return nullptr; }
+template<>
 const A3DEntity* TUniqueTSObj<A3DCrvTransformData>::GetDefaultIndexerValue() const { return nullptr; }
 template<>
 const A3DEntity* TUniqueTSObj<A3DGlobalData>::GetDefaultIndexerValue() const { return nullptr; }
@@ -1436,6 +1485,10 @@ template<>
 const A3DEntity* TUniqueTSObj<A3DTopoUniqueVertexData>::GetDefaultIndexerValue() const { return nullptr; }
 template<>
 const A3DEntity* TUniqueTSObj<A3DTopoWireEdgeData>::GetDefaultIndexerValue() const { return nullptr; }
+template<>
+const A3DEntity* TUniqueTSObj<A3DVector2dData>::GetDefaultIndexerValue() const { return nullptr; }
+template<>
+const A3DEntity* TUniqueTSObj<A3DVector3dData>::GetDefaultIndexerValue() const { return nullptr; }
 
 template<>
 uint32 TUniqueTSObjFromIndex<A3DGraphMaterialData>::GetDefaultIndexerValue() const { return A3D_DEFAULT_MATERIAL_INDEX; }
diff --git a/Engine/Plugins/Enterprise/DatasmithCADImporter/Source/CADInterfaces/Private/TechSoftInterfaceUtils.inl b/Engine/Plugins/Enterprise/DatasmithCADImporter/Source/CADInterfaces/Private/TechSoftInterfaceUtils.inl
--- a/Engine/Plugins/Enterprise/DatasmithCADImporter/Source/CADInterfaces/Private/TechSoftInterfaceUtils.inl
+++ b/Engine/Plugins/Enterprise/DatasmithCADImporter/Source/CADInterfaces/Private/TechSoftInterfaceUtils.inl
@@ -5,749 +5,944 @@
 #include "CADSceneGraph.h"
 
 #include "TUniqueTechSoftObj.h"
+#include "TechSoftInterface.h"
 
 #ifdef USE_TECHSOFT_SDK
 
 namespace TechSoftInterfaceUtils
 {
-	class FTechSoftTessellationExtractor
+
+class FTechSoftTessellationExtractor
+{
+public:
+	FTechSoftTessellationExtractor(const A3DRiRepresentationItem* RepresentationItemPtr, const double InBodyUnit)
+		: BodyUnit(InBodyUnit)
+		, TextureUnit(InBodyUnit)
 	{
-	public:
-		FTechSoftTessellationExtractor(const A3DTess3D* InTessellationPtr)
-			: TessellationPtr(InTessellationPtr)
+		CADLibrary::TUniqueTSObj<A3DRiRepresentationItemData> RepresentationItemData(RepresentationItemPtr);
+		if (!RepresentationItemData.IsValid())
 		{
+			return;
 		}
 
-		bool  FillBodyMesh(CADLibrary::FBodyMesh& BodyMesh, double FileUnit)
+		A3DEEntityType Type;
+		A3DEntityGetType(RepresentationItemData->m_pTessBase, &Type);
+		if (Type != kA3DTypeTess3D)
 		{
-			int32 VertexOffset = BodyMesh.VertexArray.Num();
-			FillVertexArray(FileUnit, BodyMesh.VertexArray);
+			return;
+		}
 
-			if (BodyMesh.VertexArray.Num() == VertexOffset)
+		TessellationPtr = RepresentationItemData->m_pTessBase;
+
+		A3DEntityGetType(RepresentationItemPtr, &Type);
+		if (Type == kA3DTypeRiBrepModel)
+		{
+			CADLibrary::TUniqueTSObj<A3DRiBrepModelData> BRepModelData(RepresentationItemPtr);
+			if (BRepModelData.IsValid())
 			{
-				return false;
+				A3DBrepData = BRepModelData->m_pBrepData;
 			}
+		}
+	}
 
-			FillFaceArray(BodyMesh, VertexOffset);
-
-			return BodyMesh.Faces.Num() > 0 ? true : false;
+	bool FillBodyMesh(CADLibrary::FBodyMesh& BodyMesh)
+	{
+		if (TessellationPtr == nullptr)
+		{
+			return false;
 		}
 
-	private:
-		void FillVertexArray(double FileUnit, TArray<FVector>& VertexArray)
+		int32 VertexOffset = BodyMesh.VertexArray.Num();
+		FillVertexArray(BodyMesh.VertexArray);
+
+		if (BodyMesh.VertexArray.Num() == VertexOffset)
 		{
-			using namespace CADLibrary;
+			return false;
+		}
 
-			TUniqueTSObj<A3DTessBaseData> TessellationBaseData(TessellationPtr);
+		BodyFaces.Empty();
+		GetBRepFaces();
+		FillFaceArray(BodyMesh, VertexOffset);
 
-			if (!TessellationBaseData.IsValid() || TessellationBaseData->m_uiCoordSize == 0)
-			{
-				return;
-			}
+		return BodyMesh.Faces.Num() > 0 ? true : false;
+	}
 
-			int32 VertexCount = TessellationBaseData->m_uiCoordSize / 3;
-			VertexArray.Reserve(VertexArray.Num() + VertexCount);
+private:
+	void FillVertexArray(TArray<FVector>& VertexArray)
+	{
+		using namespace CADLibrary;
 
-			double* Coordinates = TessellationBaseData->m_pdCoords;
-			for (unsigned int Index = 0; Index < TessellationBaseData->m_uiCoordSize; ++Index)
-			{
-				Coordinates[Index] *= FileUnit;
-			}
+		CADLibrary::TUniqueTSObj<A3DTessBaseData> TessellationBaseData(TessellationPtr);
 
-			for (unsigned int Index = 0; Index < TessellationBaseData->m_uiCoordSize; Index += 3)
-			{
-				VertexArray.Emplace(Coordinates[Index], Coordinates[Index + 1], Coordinates[Index + 2]);
-			}
+		if (!TessellationBaseData.IsValid() || TessellationBaseData->m_uiCoordSize == 0)
+		{
+			return;
 		}
 
-		// #ueent_techsoft: TODO: Make it more in line with the actual implementation of FillFaceArray
-		uint32 CountTriangles(const A3DTessFaceData& FaceTessData)
+		int32 VertexCount = TessellationBaseData->m_uiCoordSize / 3;
+		VertexArray.Reserve(VertexArray.Num() + VertexCount);
+
+		double* Coordinates = TessellationBaseData->m_pdCoords;
+		for (unsigned int Index = 0; Index < TessellationBaseData->m_uiCoordSize; ++Index)
 		{
-			const int32 TessellationFaceDataWithTriangle = 0x2222;
-			const int32 TessellationFaceDataWithFan = 0x4444;
-			const int32 TessellationFaceDataWithStripe = 0x8888;
-			const int32 TessellationFaceDataWithOneNormal = 0xE0E0;
+			Coordinates[Index] *= BodyUnit;
+		}
 
-			uint32 UsedEntitiesFlags = FaceTessData.m_usUsedEntitiesFlags;
+		for (unsigned int Index = 0; Index < TessellationBaseData->m_uiCoordSize; Index += 3)
+		{
+			VertexArray.Emplace(Coordinates[Index], Coordinates[Index + 1], Coordinates[Index + 2]);
+		}
+	}
 
-			uint32 TriangleCount = 0;
-			uint32 FaceSetIndex = 0;
-			if (UsedEntitiesFlags & TessellationFaceDataWithTriangle)
-			{
-				TriangleCount += FaceTessData.m_puiSizesTriangulated[FaceSetIndex];
-				FaceSetIndex++;
-			}
+	// #ueent_techsoft: TODO: Make it more in line with the actual implementation of FillFaceArray
+	uint32 CountTriangles(const A3DTessFaceData& FaceTessData)
+	{
+		const int32 TessellationFaceDataWithTriangle = 0x2222;
+		const int32 TessellationFaceDataWithFan = 0x4444;
+		const int32 TessellationFaceDataWithStripe = 0x8888;
+		const int32 TessellationFaceDataWithOneNormal = 0xE0E0;
 
-			if (FaceTessData.m_uiSizesTriangulatedSize > FaceSetIndex)
+		uint32 UsedEntitiesFlags = FaceTessData.m_usUsedEntitiesFlags;
+
+		uint32 TriangleCount = 0;
+		uint32 FaceSetIndex = 0;
+		if (UsedEntitiesFlags & TessellationFaceDataWithTriangle)
+		{
+			TriangleCount += FaceTessData.m_puiSizesTriangulated[FaceSetIndex];
+			FaceSetIndex++;
+		}
+
+		if (FaceTessData.m_uiSizesTriangulatedSize > FaceSetIndex)
+		{
+			if (UsedEntitiesFlags & TessellationFaceDataWithFan)
 			{
-				if (UsedEntitiesFlags & TessellationFaceDataWithFan)
+				uint32 LastFanIndex = 1 + FaceSetIndex + FaceTessData.m_puiSizesTriangulated[FaceSetIndex];
+				FaceSetIndex++;
+				for (; FaceSetIndex < LastFanIndex; FaceSetIndex++)
 				{
-					uint32 LastFanIndex = 1 + FaceSetIndex + FaceTessData.m_puiSizesTriangulated[FaceSetIndex];
-					FaceSetIndex++;
-					for (; FaceSetIndex < LastFanIndex; FaceSetIndex++)
-					{
-						uint32 FanSize = (FaceTessData.m_puiSizesTriangulated[FaceSetIndex] & kA3DTessFaceDataNormalMask);
-						TriangleCount += (FanSize - 2);
-					}
+					uint32 FanSize = (FaceTessData.m_puiSizesTriangulated[FaceSetIndex] & kA3DTessFaceDataNormalMask);
+					TriangleCount += (FanSize - 2);
 				}
 			}
+		}
 
-			if (FaceTessData.m_uiSizesTriangulatedSize > FaceSetIndex)
+		if (FaceTessData.m_uiSizesTriangulatedSize > FaceSetIndex)
+		{
+			FaceSetIndex++;
+			for (; FaceSetIndex < FaceTessData.m_uiSizesTriangulatedSize; FaceSetIndex++)
 			{
-				FaceSetIndex++;
-				for (; FaceSetIndex < FaceTessData.m_uiSizesTriangulatedSize; FaceSetIndex++)
-				{
-					uint32 StripeSize = (FaceTessData.m_puiSizesTriangulated[FaceSetIndex] & kA3DTessFaceDataNormalMask);
-					TriangleCount += (StripeSize - 2);
-				}
+				uint32 StripeSize = (FaceTessData.m_puiSizesTriangulated[FaceSetIndex] & kA3DTessFaceDataNormalMask);
+				TriangleCount += (StripeSize - 2);
 			}
-			return TriangleCount;
 		}
+		return TriangleCount;
+	}
+
+	void FillFaceArray(CADLibrary::FBodyMesh& BodyMesh, int32 VertexOffset = 0)
+	{
+		using namespace CADLibrary;
+
+		TArray<FTessellationData>& Faces = BodyMesh.Faces;
 
-		void FillFaceArray(CADLibrary::FBodyMesh& BodyMesh, int32 VertexOffset = 0)
+		CADLibrary::TUniqueTSObj<A3DTess3DData> A3DTessellationData(TessellationPtr);
+
+		if (!A3DTessellationData.IsValid() || A3DTessellationData->m_uiFaceTessSize == 0)
 		{
-			using namespace CADLibrary;
+			return;
+		}
+
+		TessellationNormals = A3DTessellationData->m_pdNormals;
+		TessellationTexCoords = A3DTessellationData->m_pdTextureCoords;
+		TriangulatedIndexes = A3DTessellationData->m_puiTriangulatedIndexes;
 
-			TArray<FTessellationData>& Faces = BodyMesh.Faces;
+		if (A3DTessellationData->m_uiFaceTessSize != BodyFaces.Num())
+		{
+			BodyFaces.Empty();
+			A3DBrepData = nullptr;
+		}
 
-			TUniqueTSObj<A3DTess3DData> TessellationData(TessellationPtr);
+		for (unsigned int Index = 0; Index < A3DTessellationData->m_uiFaceTessSize; ++Index)
+		{
+			const A3DTessFaceData& FaceTessData = A3DTessellationData->m_psFaceTessData[Index];
+			FTessellationData& Tessellation = Faces.Emplace_GetRef();
 
-			if (!TessellationData.IsValid() || TessellationData->m_uiFaceTessSize == 0)
+			Tessellation.MaterialName = FTechSoftInterface::InvalidScriptIndex;
+			if (FaceTessData.m_uiStyleIndexesSize > 0)
 			{
-				return;
+				// Store the StyleIndex on the MaterialName. It will be processed after tessellation
+				Tessellation.MaterialName = FaceTessData.m_puiStyleIndexes[0];
 			}
 
-			TessellationNormals = TessellationData->m_pdNormals;
-			TessellationTexCoords = TessellationData->m_pdTextureCoords;
-			TriangulatedIndexes = TessellationData->m_puiTriangulatedIndexes;
-
-			for (unsigned int Index = 0; Index < TessellationData->m_uiFaceTessSize; ++Index)
+			// Pre-allocate memory for triangles' data
+			uint32 TriangleCount = CountTriangles(FaceTessData);
+			Tessellation.PositionIndices.Reserve(3 * TriangleCount);
+			Tessellation.VertexIndices.Reserve(3 * TriangleCount);
+			Tessellation.NormalArray.Reserve(3 * TriangleCount);
+			if (FaceTessData.m_uiTextureCoordIndexesSize > 0)
 			{
-				const A3DTessFaceData& FaceTessData = TessellationData->m_psFaceTessData[Index];
-				FTessellationData& Tessellation = Faces.Emplace_GetRef();
+				Tessellation.TexCoordArray.Reserve(3 * TriangleCount);
+			}
 
-				Tessellation.MaterialName = FTechSoftInterface::InvalidScriptIndex;
-				if (FaceTessData.m_uiStyleIndexesSize > 0)
-				{
-					// Store the StyleIndex on the MaterialName. It will be processed after tessellation
-					Tessellation.MaterialName = FaceTessData.m_puiStyleIndexes[0];
-				}
+			uint32 UsedEntitiesFlags = FaceTessData.m_usUsedEntitiesFlags;
 
-				// Pre-allocate memory for triangles' data
-				uint32 TriangleCount = CountTriangles(FaceTessData);
-				Tessellation.PositionIndices.Reserve(3 * TriangleCount);
-				Tessellation.VertexIndices.Reserve(3 * TriangleCount);
-				Tessellation.NormalArray.Reserve(3 * TriangleCount);
-				if (FaceTessData.m_uiTextureCoordIndexesSize > 0)
-				{
-					Tessellation.TexCoordArray.Reserve(3 * TriangleCount);
-				}
+			uint32 LastTrianguleIndex = FaceTessData.m_uiStartTriangulated;
+			LastVertexIndex = 0;
 
-				uint32 UsedEntitiesFlags = FaceTessData.m_usUsedEntitiesFlags;
+			uint32 FaceSetIndex = 0;
+			bool bMustProcess = FaceTessData.m_uiSizesTriangulatedSize > FaceSetIndex;
 
-				LastTrianguleIndex = FaceTessData.m_uiStartTriangulated;
-				LastVertexIndex = 0;
+			if (bMustProcess && UsedEntitiesFlags & kA3DTessFaceDataTriangle)
+			{
+				AddFaceTriangle(Tessellation, FaceTessData.m_puiSizesTriangulated[FaceSetIndex++], LastTrianguleIndex, LastVertexIndex);
+				bMustProcess = FaceTessData.m_uiSizesTriangulatedSize > FaceSetIndex;
+			}
 
-				uint32 FaceSetIndex = 0;
-				bool bMustProcess = FaceTessData.m_uiSizesTriangulatedSize > FaceSetIndex;
+			if (bMustProcess && UsedEntitiesFlags & kA3DTessFaceDataTriangleOneNormal)
+			{
+				AddFaceTriangleWithUniqueNormal(Tessellation, FaceTessData.m_puiSizesTriangulated[FaceSetIndex++], LastTrianguleIndex, LastVertexIndex);
+				bMustProcess = FaceTessData.m_uiSizesTriangulatedSize > FaceSetIndex;
+			}
 
-				if (bMustProcess && UsedEntitiesFlags & kA3DTessFaceDataTriangle)
-				{
-					AddFaceTriangle(Tessellation, FaceTessData.m_puiSizesTriangulated[FaceSetIndex++], LastTrianguleIndex, LastVertexIndex);
-					bMustProcess = FaceTessData.m_uiSizesTriangulatedSize > FaceSetIndex;
-				}
+			if (bMustProcess && UsedEntitiesFlags & kA3DTessFaceDataTriangleTextured)
+			{
+				AddFaceTriangleWithTexture(Tessellation, FaceTessData.m_puiSizesTriangulated[FaceSetIndex++], FaceTessData.m_uiTextureCoordIndexesSize, LastTrianguleIndex, LastVertexIndex);
+				bMustProcess = FaceTessData.m_uiSizesTriangulatedSize > FaceSetIndex;
+			}
 
-				if (bMustProcess && UsedEntitiesFlags & kA3DTessFaceDataTriangleOneNormal)
-				{
-					AddFaceTriangleWithUniqueNormal(Tessellation, FaceTessData.m_puiSizesTriangulated[FaceSetIndex++], LastTrianguleIndex, LastVertexIndex);
-					bMustProcess = FaceTessData.m_uiSizesTriangulatedSize > FaceSetIndex;
-				}
+			if (bMustProcess && UsedEntitiesFlags & kA3DTessFaceDataTriangleOneNormalTextured)
+			{
+				AddFaceTriangleWithUniqueNormalAndTexture(Tessellation, FaceTessData.m_puiSizesTriangulated[FaceSetIndex++], FaceTessData.m_uiTextureCoordIndexesSize, LastTrianguleIndex, LastVertexIndex);
+				bMustProcess = FaceTessData.m_uiSizesTriangulatedSize > FaceSetIndex;
+			}
 
-				if (bMustProcess && UsedEntitiesFlags & kA3DTessFaceDataTriangleTextured)
+			if (bMustProcess && UsedEntitiesFlags & kA3DTessFaceDataTriangleFan)
+			{
+				uint32 FanCount = FaceTessData.m_puiSizesTriangulated[FaceSetIndex++];
+				for (uint32 FanIndex = 0; FanIndex < FanCount; ++FanIndex)
 				{
-					AddFaceTriangleWithTexture(Tessellation, FaceTessData.m_puiSizesTriangulated[FaceSetIndex++], FaceTessData.m_uiTextureCoordIndexesSize, LastTrianguleIndex, LastVertexIndex);
-					bMustProcess = FaceTessData.m_uiSizesTriangulatedSize > FaceSetIndex;
+					uint32 VertexCount = FaceTessData.m_puiSizesTriangulated[FaceSetIndex++];
+					AddFaceTriangleFan(Tessellation, VertexCount, LastTrianguleIndex, LastVertexIndex);
 				}
+				bMustProcess = FaceTessData.m_uiSizesTriangulatedSize > FaceSetIndex;
+			}
 
-				if (bMustProcess && UsedEntitiesFlags & kA3DTessFaceDataTriangleOneNormalTextured)
+			if (bMustProcess && UsedEntitiesFlags & kA3DTessFaceDataTriangleFanOneNormal)
+			{
+				uint32 FanCount = FaceTessData.m_puiSizesTriangulated[FaceSetIndex++] & kA3DTessFaceDataNormalMask;
+				for (uint32 FanIndex = 0; FanIndex < FanCount; ++FanIndex)
 				{
-					AddFaceTriangleWithUniqueNormalAndTexture(Tessellation, FaceTessData.m_puiSizesTriangulated[FaceSetIndex++], FaceTessData.m_uiTextureCoordIndexesSize, LastTrianguleIndex, LastVertexIndex);
-					bMustProcess = FaceTessData.m_uiSizesTriangulatedSize > FaceSetIndex;
+					ensure((FaceTessData.m_puiSizesTriangulated[FaceSetIndex] & kA3DTessFaceDataNormalSingle) != 0);
+
+					uint32 VertexCount = FaceTessData.m_puiSizesTriangulated[FaceSetIndex++] & kA3DTessFaceDataNormalMask;
+					AddFaceTriangleFanWithUniqueNormal(Tessellation, VertexCount, LastTrianguleIndex, LastVertexIndex);
 				}
+				bMustProcess = FaceTessData.m_uiSizesTriangulatedSize > FaceSetIndex;
+			}
 
-				if (bMustProcess && UsedEntitiesFlags & kA3DTessFaceDataTriangleFan)
+			if (bMustProcess && UsedEntitiesFlags & kA3DTessFaceDataTriangleFanTextured)
+			{
+				uint32 FanCount = FaceTessData.m_puiSizesTriangulated[FaceSetIndex++];
+				for (uint32 FanIndex = 0; FanIndex < FanCount; ++FanIndex)
 				{
-					uint32 FanCount = FaceTessData.m_puiSizesTriangulated[FaceSetIndex++];
-					for (uint32 FanIndex = 0; FanIndex < FanCount; ++FanIndex)
-					{
-						uint32 VertexCount = FaceTessData.m_puiSizesTriangulated[FaceSetIndex++];
-						AddFaceTriangleFan(Tessellation, VertexCount, LastTrianguleIndex, LastVertexIndex);
-					}
-					bMustProcess = FaceTessData.m_uiSizesTriangulatedSize > FaceSetIndex;
+					uint32 VertexCount = FaceTessData.m_puiSizesTriangulated[FaceSetIndex++];
+					AddFaceTriangleFanWithTexture(Tessellation, VertexCount, FaceTessData.m_uiTextureCoordIndexesSize, LastTrianguleIndex, LastVertexIndex);
 				}
+				bMustProcess = FaceTessData.m_uiSizesTriangulatedSize > FaceSetIndex;
+			}
 
-				if (bMustProcess && UsedEntitiesFlags & kA3DTessFaceDataTriangleFanOneNormal)
+			if (bMustProcess && UsedEntitiesFlags & kA3DTessFaceDataTriangleFanOneNormalTextured)
+			{
+				uint32 FanCount = FaceTessData.m_puiSizesTriangulated[FaceSetIndex++] & kA3DTessFaceDataNormalMask;
+				for (uint32 FanIndex = 0; FanIndex < FanCount; ++FanIndex)
 				{
-					uint32 FanCount = FaceTessData.m_puiSizesTriangulated[FaceSetIndex++] & kA3DTessFaceDataNormalMask;
-					for (uint32 FanIndex = 0; FanIndex < FanCount; ++FanIndex)
-					{
-						ensure((FaceTessData.m_puiSizesTriangulated[FaceSetIndex] & kA3DTessFaceDataNormalSingle) != 0);
+					ensure((FaceTessData.m_puiSizesTriangulated[FaceSetIndex] & kA3DTessFaceDataNormalSingle) != 0);
 
-						uint32 VertexCount = FaceTessData.m_puiSizesTriangulated[FaceSetIndex++] & kA3DTessFaceDataNormalMask;
-						AddFaceTriangleFanWithUniqueNormal(Tessellation, VertexCount, LastTrianguleIndex, LastVertexIndex);
-					}
-					bMustProcess = FaceTessData.m_uiSizesTriangulatedSize > FaceSetIndex;
+					uint32 VertexCount = FaceTessData.m_puiSizesTriangulated[FaceSetIndex++] & kA3DTessFaceDataNormalMask;
+					AddFaceTriangleFanWithUniqueNormalAndTexture(Tessellation, VertexCount, FaceTessData.m_uiTextureCoordIndexesSize, LastTrianguleIndex, LastVertexIndex);
 				}
+				bMustProcess = FaceTessData.m_uiSizesTriangulatedSize > FaceSetIndex;
+			}
 
-				if (bMustProcess && UsedEntitiesFlags & kA3DTessFaceDataTriangleFanTextured)
+			if (bMustProcess && UsedEntitiesFlags & kA3DTessFaceDataTriangleStripe)
+			{
+				A3DUns32 StripeSize = FaceTessData.m_puiSizesTriangulated[FaceSetIndex++];
+				for (A3DUns32 StripeIndex = 0; StripeIndex < StripeSize; ++StripeIndex)
 				{
-					uint32 FanCount = FaceTessData.m_puiSizesTriangulated[FaceSetIndex++];
-					for (uint32 FanIndex = 0; FanIndex < FanCount; ++FanIndex)
-					{
-						uint32 VertexCount = FaceTessData.m_puiSizesTriangulated[FaceSetIndex++];
-						AddFaceTriangleFanWithTexture(Tessellation, VertexCount, FaceTessData.m_uiTextureCoordIndexesSize, LastTrianguleIndex, LastVertexIndex);
-					}
-					bMustProcess = FaceTessData.m_uiSizesTriangulatedSize > FaceSetIndex;
+					A3DUns32 PointCount = FaceTessData.m_puiSizesTriangulated[FaceSetIndex++];
+					AddFaceTriangleStripe(Tessellation, PointCount, LastTrianguleIndex, LastVertexIndex);
 				}
+				bMustProcess = FaceTessData.m_uiSizesTriangulatedSize > FaceSetIndex;
+			}
 
-				if (bMustProcess && UsedEntitiesFlags & kA3DTessFaceDataTriangleFanOneNormalTextured)
+			if (bMustProcess && UsedEntitiesFlags & kA3DTessFaceDataTriangleStripeOneNormal)
+			{
+				A3DUns32 StripeSize = FaceTessData.m_puiSizesTriangulated[FaceSetIndex++] & kA3DTessFaceDataNormalMask;
+				for (A3DUns32 StripeIndex = 0; StripeIndex < StripeSize; ++StripeIndex)
 				{
-					uint32 FanCount = FaceTessData.m_puiSizesTriangulated[FaceSetIndex++] & kA3DTessFaceDataNormalMask;
-					for (uint32 FanIndex = 0; FanIndex < FanCount; ++FanIndex)
-					{
-						ensure((FaceTessData.m_puiSizesTriangulated[FaceSetIndex] & kA3DTessFaceDataNormalSingle) != 0);
+					bool bIsOneNormal = (FaceTessData.m_puiSizesTriangulated[FaceSetIndex] & kA3DTessFaceDataNormalSingle) != 0;
+					A3DUns32 PointCount = FaceTessData.m_puiSizesTriangulated[FaceSetIndex++] & kA3DTessFaceDataNormalMask;
 
-						uint32 VertexCount = FaceTessData.m_puiSizesTriangulated[FaceSetIndex++] & kA3DTessFaceDataNormalMask;
-						AddFaceTriangleFanWithUniqueNormalAndTexture(Tessellation, VertexCount, FaceTessData.m_uiTextureCoordIndexesSize, LastTrianguleIndex, LastVertexIndex);
-					}
-					bMustProcess = FaceTessData.m_uiSizesTriangulatedSize > FaceSetIndex;
-				}
-
-				if (bMustProcess && UsedEntitiesFlags & kA3DTessFaceDataTriangleStripe)
-				{
-					A3DUns32 StripeSize = FaceTessData.m_puiSizesTriangulated[FaceSetIndex++];
-					for (A3DUns32 StripeIndex = 0; StripeIndex < StripeSize; ++StripeIndex)
+					// Is there only one normal for the entire stripe?
+					if (bIsOneNormal == false)
 					{
-						A3DUns32 PointCount = FaceTessData.m_puiSizesTriangulated[FaceSetIndex++];
 						AddFaceTriangleStripe(Tessellation, PointCount, LastTrianguleIndex, LastVertexIndex);
+						continue;
 					}
-					bMustProcess = FaceTessData.m_uiSizesTriangulatedSize > FaceSetIndex;
-				}
-
-				if (bMustProcess && UsedEntitiesFlags & kA3DTessFaceDataTriangleStripeOneNormal)
-				{
-					A3DUns32 StripeSize = FaceTessData.m_puiSizesTriangulated[FaceSetIndex++] & kA3DTessFaceDataNormalMask;
-					for (A3DUns32 StripeIndex = 0; StripeIndex < StripeSize; ++StripeIndex)
-					{
-						bool bIsOneNormal = (FaceTessData.m_puiSizesTriangulated[FaceSetIndex] & kA3DTessFaceDataNormalSingle) != 0;
-						A3DUns32 PointCount = FaceTessData.m_puiSizesTriangulated[FaceSetIndex++] & kA3DTessFaceDataNormalMask;
-
-						// Is there only one normal for the entire stripe?
-						if (bIsOneNormal == false)
-						{
-							AddFaceTriangleStripe(Tessellation, PointCount, LastTrianguleIndex, LastVertexIndex);
-							continue;
-						}
 
-						AddFaceTriangleStripeWithUniqueNormal(Tessellation, PointCount, LastTrianguleIndex, LastVertexIndex);
-					}
-					bMustProcess = FaceTessData.m_uiSizesTriangulatedSize > FaceSetIndex;
+					AddFaceTriangleStripeWithUniqueNormal(Tessellation, PointCount, LastTrianguleIndex, LastVertexIndex);
 				}
+				bMustProcess = FaceTessData.m_uiSizesTriangulatedSize > FaceSetIndex;
+			}
 
-				if (bMustProcess && UsedEntitiesFlags & kA3DTessFaceDataTriangleStripeTextured)
+			if (bMustProcess && UsedEntitiesFlags & kA3DTessFaceDataTriangleStripeTextured)
+			{
+				A3DUns32 StripeSize = FaceTessData.m_puiSizesTriangulated[FaceSetIndex++];
+				for (A3DUns32 StripeIndex = 0; StripeIndex < StripeSize; ++StripeIndex)
 				{
-					A3DUns32 StripeSize = FaceTessData.m_puiSizesTriangulated[FaceSetIndex++];
-					for (A3DUns32 StripeIndex = 0; StripeIndex < StripeSize; ++StripeIndex)
-					{
-						A3DUns32 PointCount = FaceTessData.m_puiSizesTriangulated[FaceSetIndex++];
-						AddFaceTriangleStripeWithTexture(Tessellation, PointCount, FaceTessData.m_uiTextureCoordIndexesSize, LastTrianguleIndex, LastVertexIndex);
-					}
-					bMustProcess = FaceTessData.m_uiSizesTriangulatedSize > FaceSetIndex;
+					A3DUns32 PointCount = FaceTessData.m_puiSizesTriangulated[FaceSetIndex++];
+					AddFaceTriangleStripeWithTexture(Tessellation, PointCount, FaceTessData.m_uiTextureCoordIndexesSize, LastTrianguleIndex, LastVertexIndex);
 				}
+				bMustProcess = FaceTessData.m_uiSizesTriangulatedSize > FaceSetIndex;
+			}
 
-				if (bMustProcess && UsedEntitiesFlags & kA3DTessFaceDataTriangleStripeOneNormalTextured)
+			if (bMustProcess && UsedEntitiesFlags & kA3DTessFaceDataTriangleStripeOneNormalTextured)
+			{
+				A3DUns32 StripeSize = FaceTessData.m_puiSizesTriangulated[FaceSetIndex++] & kA3DTessFaceDataNormalMask;
+				for (A3DUns32 StripeIndex = 0; StripeIndex < StripeSize; ++StripeIndex)
 				{
-					A3DUns32 StripeSize = FaceTessData.m_puiSizesTriangulated[FaceSetIndex++] & kA3DTessFaceDataNormalMask;
-					for (A3DUns32 StripeIndex = 0; StripeIndex < StripeSize; ++StripeIndex)
-					{
-						A3DUns32 PointCount = FaceTessData.m_puiSizesTriangulated[FaceSetIndex++] & kA3DTessFaceDataNormalMask;
-						AddFaceTriangleStripeWithUniqueNormalAndTexture(Tessellation, PointCount, FaceTessData.m_uiTextureCoordIndexesSize, LastTrianguleIndex, LastVertexIndex);
-					}
-					bMustProcess = FaceTessData.m_uiSizesTriangulatedSize > FaceSetIndex;
+					A3DUns32 PointCount = FaceTessData.m_puiSizesTriangulated[FaceSetIndex++] & kA3DTessFaceDataNormalMask;
+					AddFaceTriangleStripeWithUniqueNormalAndTexture(Tessellation, PointCount, FaceTessData.m_uiTextureCoordIndexesSize, LastTrianguleIndex, LastVertexIndex);
 				}
-				
-				if (VertexOffset > 0)
-				{
-					for (int32& PositionIndex : Tessellation.PositionIndices)
-					{
-						PositionIndex += VertexOffset;
-					}
-				}
-
-				ensure(!bMustProcess);
+				bMustProcess = FaceTessData.m_uiSizesTriangulatedSize > FaceSetIndex;
 			}
-		}
 
-		typedef double A3DDouble;
-		bool AddFace(int32 FaceIndex[3], CADLibrary::FTessellationData& Tessellation, int32& InOutVertexIndex)
-		{
-			if (FaceIndex[0] == FaceIndex[1] || FaceIndex[0] == FaceIndex[2] || FaceIndex[1] == FaceIndex[2])
+			if (VertexOffset > 0)
 			{
-				return false;
+				for (int32& PositionIndex : Tessellation.PositionIndices)
+				{
+					PositionIndex += VertexOffset;
+				}
 			}
 
-			for (int32 Index = 0; Index < 3; ++Index)
-			{
-				Tessellation.VertexIndices.Add(InOutVertexIndex++);
-			}
-			Tessellation.PositionIndices.Append(FaceIndex, 3);
-			return true;
-		};
+			ensure(!bMustProcess);
 
-		void AddNormals(const A3DDouble* Normals, const int32 Indices[3], TArray<FVector>& NormalsArray)
-		{
-			for (int32 Index = 0; Index < 3; ++Index)
+			if (!BodyFaces.IsEmpty())
 			{
-				int32 NormalIndex = Indices[Index];
-				NormalsArray.Emplace(Normals[NormalIndex], Normals[NormalIndex + 1], Normals[NormalIndex + 2]);
+				const A3DTopoFace* TopoFace = BodyFaces[Index];
+				ScaleUV(TopoFace, Tessellation.TexCoordArray);
 			}
-		};
+		}
+	}
 
-		void AddTextureCoordinates(const A3DDouble* TextureCoords, const int32 Indices[3], TArray<FVector2D>& TessellationTextures)
+	typedef double A3DDouble;
+	bool AddFace(int32 FaceIndex[3], CADLibrary::FTessellationData& Tessellation, int32& InOutVertexIndex)
+	{
+		if (FaceIndex[0] == FaceIndex[1] || FaceIndex[0] == FaceIndex[2] || FaceIndex[1] == FaceIndex[2])
 		{
-			for (int32 Index = 0; Index < 3; ++Index)
-			{
-				int32 TextureIndex = Indices[Index];
-				TessellationTextures.Emplace(TextureCoords[TextureIndex], TextureCoords[TextureIndex + 1]);
-			}
-		};
+			return false;
+		}
 
-		void AddFaceTriangle(CADLibrary::FTessellationData& Tessellation, const uint32 InTriangleCount, uint32& InOutStartIndex, int32& InOutLastVertexIndex)
+		for (int32 Index = 0; Index < 3; ++Index)
 		{
-			int32 FaceIndex[3] = { 0, 0, 0 };
-			int32 NormalIndex[3] = { 0, 0, 0 };
-
-			// Get Triangles
-			for (uint32 TriangleIndex = 0; TriangleIndex < InTriangleCount; TriangleIndex++)
-			{
-				NormalIndex[0] = TriangulatedIndexes[InOutStartIndex++];
-				FaceIndex[0] = TriangulatedIndexes[InOutStartIndex++] / 3;
-				NormalIndex[1] = TriangulatedIndexes[InOutStartIndex++];
-				FaceIndex[1] = TriangulatedIndexes[InOutStartIndex++] / 3;
-				NormalIndex[2] = TriangulatedIndexes[InOutStartIndex++];
-				FaceIndex[2] = TriangulatedIndexes[InOutStartIndex++] / 3;
+			Tessellation.VertexIndices.Add(InOutVertexIndex++);
+		}
+		Tessellation.PositionIndices.Append(FaceIndex, 3);
+		return true;
+	};
 
-				if (AddFace(FaceIndex, Tessellation, InOutLastVertexIndex))
-				{
-					AddNormals(TessellationNormals, NormalIndex, Tessellation.NormalArray);
-				}
-			}
+	void AddNormals(const A3DDouble* Normals, const int32 Indices[3], TArray<FVector>& NormalsArray)
+	{
+		for (int32 Index = 0; Index < 3; ++Index)
+		{
+			int32 NormalIndex = Indices[Index];
+			NormalsArray.Emplace(Normals[NormalIndex], Normals[NormalIndex + 1], Normals[NormalIndex + 2]);
 		}
+	};
 
-		void AddFaceTriangleWithUniqueNormal(CADLibrary::FTessellationData& Tessellation, const uint32 InTriangleCount, uint32& InOutLastTriangleIndex, int32& InOutLastVertexIndex)
+	void AddTextureCoordinates(const A3DDouble* TextureCoords, const int32 Indices[3], TArray<FVector2D>& TessellationTextures)
+	{
+		for (int32 Index = 0; Index < 3; ++Index)
 		{
-			int32 FaceIndex[3] = { 0, 0, 0 };
-			int32 NormalIndex[3] = { 0, 0, 0 };
+			int32 TextureIndex = Indices[Index];
+			TessellationTextures.Emplace(TextureCoords[TextureIndex], TextureCoords[TextureIndex + 1]);
+		}
+	};
 
-			// Get Triangles
-			for (uint32 TriangleIndex = 0; TriangleIndex < InTriangleCount; TriangleIndex++)
+	void AddFaceTriangle(CADLibrary::FTessellationData& Tessellation, const uint32 InTriangleCount, uint32& InOutStartIndex, int32& InOutLastVertexIndex)
+	{
+		int32 FaceIndex[3] = { 0, 0, 0 };
+		int32 NormalIndex[3] = { 0, 0, 0 };
+
+		// Get Triangles
+		for (uint32 TriangleIndex = 0; TriangleIndex < InTriangleCount; TriangleIndex++)
+		{
+			NormalIndex[0] = TriangulatedIndexes[InOutStartIndex++];
+			FaceIndex[0] = TriangulatedIndexes[InOutStartIndex++] / 3;
+			NormalIndex[1] = TriangulatedIndexes[InOutStartIndex++];
+			FaceIndex[1] = TriangulatedIndexes[InOutStartIndex++] / 3;
+			NormalIndex[2] = TriangulatedIndexes[InOutStartIndex++];
+			FaceIndex[2] = TriangulatedIndexes[InOutStartIndex++] / 3;
+
+			if (AddFace(FaceIndex, Tessellation, InOutLastVertexIndex))
 			{
-				NormalIndex[0] = TriangulatedIndexes[InOutLastTriangleIndex++];
-				NormalIndex[1] = NormalIndex[0];
-				NormalIndex[2] = NormalIndex[0];
+				AddNormals(TessellationNormals, NormalIndex, Tessellation.NormalArray);
+			}
+		}
+	}
 
-				FaceIndex[0] = TriangulatedIndexes[InOutLastTriangleIndex++] / 3;
-				FaceIndex[1] = TriangulatedIndexes[InOutLastTriangleIndex++] / 3;
-				FaceIndex[2] = TriangulatedIndexes[InOutLastTriangleIndex++] / 3;
+	void AddFaceTriangleWithUniqueNormal(CADLibrary::FTessellationData& Tessellation, const uint32 InTriangleCount, uint32& InOutLastTriangleIndex, int32& InOutLastVertexIndex)
+	{
+		int32 FaceIndex[3] = { 0, 0, 0 };
+		int32 NormalIndex[3] = { 0, 0, 0 };
 
-				if (!AddFace(FaceIndex, Tessellation, InOutLastVertexIndex))
-				{
-					continue;
-				}
+		// Get Triangles
+		for (uint32 TriangleIndex = 0; TriangleIndex < InTriangleCount; TriangleIndex++)
+		{
+			NormalIndex[0] = TriangulatedIndexes[InOutLastTriangleIndex++];
+			NormalIndex[1] = NormalIndex[0];
+			NormalIndex[2] = NormalIndex[0];
 
-				AddNormals(TessellationNormals, NormalIndex, Tessellation.NormalArray);
+			FaceIndex[0] = TriangulatedIndexes[InOutLastTriangleIndex++] / 3;
+			FaceIndex[1] = TriangulatedIndexes[InOutLastTriangleIndex++] / 3;
+			FaceIndex[2] = TriangulatedIndexes[InOutLastTriangleIndex++] / 3;
+
+			if (!AddFace(FaceIndex, Tessellation, InOutLastVertexIndex))
+			{
+				continue;
 			}
+
+			AddNormals(TessellationNormals, NormalIndex, Tessellation.NormalArray);
 		}
+	}
+
+	void AddFaceTriangleWithUniqueNormalAndTexture(CADLibrary::FTessellationData& Tessellation, const uint32 InTriangleCount, const uint32 TextureCount, uint32& InOutStartIndex, int32& InOutLastVertexIndex)
+	{
+		int32 FaceIndex[3] = { 0, 0, 0 };
+		int32 NormalIndex[3] = { 0, 0, 0 };
+		int32 TextureIndex[3] = { 0, 0, 0 };
 
-		void AddFaceTriangleWithUniqueNormalAndTexture(CADLibrary::FTessellationData& Tessellation, const uint32 InTriangleCount, const uint32 TextureCount, uint32& InOutStartIndex, int32& InOutLastVertexIndex)
+		// Get Triangles
+		for (uint32 TriangleIndex = 0; TriangleIndex < InTriangleCount; TriangleIndex++)
 		{
-			int32 FaceIndex[3] = { 0, 0, 0 };
-			int32 NormalIndex[3] = { 0, 0, 0 };
-			int32 TextureIndex[3] = { 0, 0, 0 };
+			NormalIndex[0] = TriangulatedIndexes[InOutStartIndex++];
+			NormalIndex[1] = NormalIndex[0];
+			NormalIndex[2] = NormalIndex[0];
 
-			// Get Triangles
-			for (uint32 TriangleIndex = 0; TriangleIndex < InTriangleCount; TriangleIndex++)
+			TextureIndex[0] = TriangulatedIndexes[InOutStartIndex];
+			InOutStartIndex += TextureCount;
+			FaceIndex[0] = TriangulatedIndexes[InOutStartIndex++] / 3;
+			TextureIndex[1] = TriangulatedIndexes[InOutStartIndex];
+			InOutStartIndex += TextureCount;
+			FaceIndex[1] = TriangulatedIndexes[InOutStartIndex++] / 3;
+			TextureIndex[2] = TriangulatedIndexes[InOutStartIndex];
+			InOutStartIndex += TextureCount;
+			FaceIndex[2] = TriangulatedIndexes[InOutStartIndex++] / 3;
+
+			if (!AddFace(FaceIndex, Tessellation, InOutLastVertexIndex))
 			{
-				NormalIndex[0] = TriangulatedIndexes[InOutStartIndex++];
-				NormalIndex[1] = NormalIndex[0];
-				NormalIndex[2] = NormalIndex[0];
+				continue;
+			}
 
-				TextureIndex[0] = TriangulatedIndexes[InOutStartIndex];
-				InOutStartIndex += TextureCount;
-				FaceIndex[0] = TriangulatedIndexes[InOutStartIndex++] / 3;
-				TextureIndex[1] = TriangulatedIndexes[InOutStartIndex];
-				InOutStartIndex += TextureCount;
-				FaceIndex[1] = TriangulatedIndexes[InOutStartIndex++] / 3;
-				TextureIndex[2] = TriangulatedIndexes[InOutStartIndex];
-				InOutStartIndex += TextureCount;
-				FaceIndex[2] = TriangulatedIndexes[InOutStartIndex++] / 3;
+			AddNormals(TessellationNormals, NormalIndex, Tessellation.NormalArray);
+			AddTextureCoordinates(TessellationTexCoords, TextureIndex, Tessellation.TexCoordArray);
+		}
+	}
 
-				if (!AddFace(FaceIndex, Tessellation, InOutLastVertexIndex))
-				{
-					continue;
-				}
+	void AddFaceTriangleWithTexture(CADLibrary::FTessellationData& Tessellation, const uint32 InTriangleCount, const uint32 InTextureCount, uint32& InOutStartIndex, int32& inOutLastVertexIndex)
+	{
+		int32 FaceIndex[3] = { 0, 0, 0 };
+		int32 NormalIndex[3] = { 0, 0, 0 };
+		int32 TextureIndex[3] = { 0, 0, 0 };
 
+		// Get Triangles
+		for (uint64 TriangleIndex = 0; TriangleIndex < InTriangleCount; TriangleIndex++)
+		{
+			NormalIndex[0] = TriangulatedIndexes[InOutStartIndex++];
+			TextureIndex[0] = TriangulatedIndexes[InOutStartIndex];
+			InOutStartIndex += InTextureCount;
+			FaceIndex[0] = TriangulatedIndexes[InOutStartIndex++] / 3;
+			NormalIndex[1] = TriangulatedIndexes[InOutStartIndex++];
+			TextureIndex[1] = TriangulatedIndexes[InOutStartIndex];
+			InOutStartIndex += InTextureCount;
+			FaceIndex[1] = TriangulatedIndexes[InOutStartIndex++] / 3;
+			NormalIndex[2] = TriangulatedIndexes[InOutStartIndex++];
+			TextureIndex[2] = TriangulatedIndexes[InOutStartIndex];
+			InOutStartIndex += InTextureCount;
+			FaceIndex[2] = TriangulatedIndexes[InOutStartIndex++] / 3;
+
+			if (AddFace(FaceIndex, Tessellation, inOutLastVertexIndex))
+			{
 				AddNormals(TessellationNormals, NormalIndex, Tessellation.NormalArray);
 				AddTextureCoordinates(TessellationTexCoords, TextureIndex, Tessellation.TexCoordArray);
 			}
 		}
+	}
+
+	void AddFaceTriangleFan(CADLibrary::FTessellationData& Tessellation, const uint32 InTriangleCount, uint32& InOutLastTriangleIndex, int32& InOutLastVertexIndex)
+	{
+		int32 FaceIndex[3] = { 0, 0, 0 };
+		int32 NormalIndex[3] = { 0, 0, 0 };
+
+		NormalIndex[0] = TriangulatedIndexes[InOutLastTriangleIndex++];
+		FaceIndex[0] = TriangulatedIndexes[InOutLastTriangleIndex++] / 3;
+		NormalIndex[1] = TriangulatedIndexes[InOutLastTriangleIndex++];
+		FaceIndex[1] = TriangulatedIndexes[InOutLastTriangleIndex++] / 3;
 
-		void AddFaceTriangleWithTexture(CADLibrary::FTessellationData& Tessellation, const uint32 InTriangleCount, const uint32 InTextureCount, uint32& InOutStartIndex, int32& inOutLastVertexIndex)
+		for (unsigned long TriangleIndex = 2; TriangleIndex < InTriangleCount; TriangleIndex++)
 		{
-			int32 FaceIndex[3] = { 0, 0, 0 };
-			int32 NormalIndex[3] = { 0, 0, 0 };
-			int32 TextureIndex[3] = { 0, 0, 0 };
+			NormalIndex[2] = TriangulatedIndexes[InOutLastTriangleIndex++];
+			FaceIndex[2] = TriangulatedIndexes[InOutLastTriangleIndex++] / 3;
 
-			// Get Triangles
-			for (uint64 TriangleIndex = 0; TriangleIndex < InTriangleCount; TriangleIndex++)
+			if (AddFace(FaceIndex, Tessellation, InOutLastVertexIndex))
 			{
-				NormalIndex[0] = TriangulatedIndexes[InOutStartIndex++];
-				TextureIndex[0] = TriangulatedIndexes[InOutStartIndex];
-				InOutStartIndex += InTextureCount;
-				FaceIndex[0] = TriangulatedIndexes[InOutStartIndex++] / 3;
-				NormalIndex[1] = TriangulatedIndexes[InOutStartIndex++];
-				TextureIndex[1] = TriangulatedIndexes[InOutStartIndex];
-				InOutStartIndex += InTextureCount;
-				FaceIndex[1] = TriangulatedIndexes[InOutStartIndex++] / 3;
-				NormalIndex[2] = TriangulatedIndexes[InOutStartIndex++];
-				TextureIndex[2] = TriangulatedIndexes[InOutStartIndex];
-				InOutStartIndex += InTextureCount;
-				FaceIndex[2] = TriangulatedIndexes[InOutStartIndex++] / 3;
-
-				if (AddFace(FaceIndex, Tessellation, inOutLastVertexIndex))
-				{
-					AddNormals(TessellationNormals, NormalIndex, Tessellation.NormalArray);
-					AddTextureCoordinates(TessellationTexCoords, TextureIndex, Tessellation.TexCoordArray);
-				}
+				AddNormals(TessellationNormals, NormalIndex, Tessellation.NormalArray);
 			}
+
+			NormalIndex[1] = NormalIndex[2];
+			FaceIndex[1] = FaceIndex[2];
 		}
+	}
 
-		void AddFaceTriangleFan(CADLibrary::FTessellationData& Tessellation, const uint32 InTriangleCount, uint32& InOutLastTriangleIndex, int32& InOutLastVertexIndex)
-		{
-			int32 FaceIndex[3] = { 0, 0, 0 };
-			int32 NormalIndex[3] = { 0, 0, 0 };
+	void AddFaceTriangleFanWithUniqueNormal(CADLibrary::FTessellationData& Tessellation, const uint32 InTriangleCount, uint32& InOutLastTriangleIndex, int32& /*LastVertexIndex*/)
+	{
+		int32 FaceIndex[3] = { 0, 0, 0 };
+		int32 NormalIndex[3] = { 0, 0, 0 };
 
-			NormalIndex[0] = TriangulatedIndexes[InOutLastTriangleIndex++];
-			FaceIndex[0] = TriangulatedIndexes[InOutLastTriangleIndex++] / 3;
-			NormalIndex[1] = TriangulatedIndexes[InOutLastTriangleIndex++];
-			FaceIndex[1] = TriangulatedIndexes[InOutLastTriangleIndex++] / 3;
+		NormalIndex[0] = TriangulatedIndexes[InOutLastTriangleIndex++];
+		NormalIndex[1] = NormalIndex[0];
+		NormalIndex[2] = NormalIndex[0];
 
-			for (unsigned long TriangleIndex = 2; TriangleIndex < InTriangleCount; TriangleIndex++)
-			{
-				NormalIndex[2] = TriangulatedIndexes[InOutLastTriangleIndex++];
-				FaceIndex[2] = TriangulatedIndexes[InOutLastTriangleIndex++] / 3;
+		FaceIndex[0] = TriangulatedIndexes[InOutLastTriangleIndex++] / 3;
+		FaceIndex[1] = TriangulatedIndexes[InOutLastTriangleIndex++] / 3;
 
-				if (AddFace(FaceIndex, Tessellation, InOutLastVertexIndex))
-				{
-					AddNormals(TessellationNormals, NormalIndex, Tessellation.NormalArray);
-				}
+		// Get Triangles
+		for (uint32 TriangleIndex = 2; TriangleIndex < InTriangleCount; TriangleIndex++)
+		{
+			FaceIndex[2] = TriangulatedIndexes[InOutLastTriangleIndex++] / 3;
 
-				NormalIndex[1] = NormalIndex[2];
-				FaceIndex[1] = FaceIndex[2];
+			if (AddFace(FaceIndex, Tessellation, LastVertexIndex))
+			{
+				AddNormals(TessellationNormals, NormalIndex, Tessellation.NormalArray);
 			}
+
+			FaceIndex[1] = FaceIndex[2];
 		}
+	}
 
-		void AddFaceTriangleFanWithUniqueNormal(CADLibrary::FTessellationData& Tessellation, const uint32 InTriangleCount, uint32& InOutLastTriangleIndex, int32& /*LastVertexIndex*/)
-		{
-			int32 FaceIndex[3] = { 0, 0, 0 };
-			int32 NormalIndex[3] = { 0, 0, 0 };
+	void AddFaceTriangleFanWithUniqueNormalAndTexture(CADLibrary::FTessellationData& Tessellation, const uint32 InTriangleCount, const uint32 TextureCount, uint32& InOutLastTriangleIndex, int32& InOutLastVertexIndex)
+	{
+		int32 FaceIndex[3] = { 0, 0, 0 };
+		int32 NormalIndex[3] = { 0, 0, 0 };
+		int32 TextureIndex[3] = { 0, 0, 0 };
 
-			NormalIndex[0] = TriangulatedIndexes[InOutLastTriangleIndex++];
-			NormalIndex[1] = NormalIndex[0];
-			NormalIndex[2] = NormalIndex[0];
+		NormalIndex[0] = TriangulatedIndexes[InOutLastTriangleIndex++];
+		NormalIndex[1] = NormalIndex[0];
+		NormalIndex[2] = NormalIndex[0];
 
-			FaceIndex[0] = TriangulatedIndexes[InOutLastTriangleIndex++] / 3;
-			FaceIndex[1] = TriangulatedIndexes[InOutLastTriangleIndex++] / 3;
+		TextureIndex[0] = TriangulatedIndexes[InOutLastTriangleIndex];
+		InOutLastTriangleIndex += TextureCount;
+		FaceIndex[0] = TriangulatedIndexes[InOutLastTriangleIndex++] / 3;
 
-			// Get Triangles
-			for (uint32 TriangleIndex = 2; TriangleIndex < InTriangleCount; TriangleIndex++)
-			{
-				FaceIndex[2] = TriangulatedIndexes[InOutLastTriangleIndex++] / 3;
+		TextureIndex[1] = TriangulatedIndexes[InOutLastTriangleIndex];
+		InOutLastTriangleIndex += TextureCount;
+		FaceIndex[1] = TriangulatedIndexes[InOutLastTriangleIndex++] / 3;
 
-				if (AddFace(FaceIndex, Tessellation, LastVertexIndex))
-				{
-					AddNormals(TessellationNormals, NormalIndex, Tessellation.NormalArray);
-				}
+		for (uint32 TriangleIndex = 2; TriangleIndex < InTriangleCount; TriangleIndex++)
+		{
+			TextureIndex[2] = TriangulatedIndexes[InOutLastTriangleIndex];
+			InOutLastTriangleIndex += TextureCount;
+			FaceIndex[2] = TriangulatedIndexes[InOutLastTriangleIndex++] / 3;
 
-				FaceIndex[1] = FaceIndex[2];
+			if (AddFace(FaceIndex, Tessellation, InOutLastVertexIndex))
+			{
+				AddNormals(TessellationNormals, NormalIndex, Tessellation.NormalArray);
+				AddTextureCoordinates(TessellationTexCoords, TextureIndex, Tessellation.TexCoordArray);
 			}
+
+			FaceIndex[1] = FaceIndex[2];
+			TextureIndex[1] = TextureIndex[2];
 		}
+	}
 
-		void AddFaceTriangleFanWithUniqueNormalAndTexture(CADLibrary::FTessellationData& Tessellation, const uint32 InTriangleCount, const uint32 TextureCount, uint32& InOutLastTriangleIndex, int32& InOutLastVertexIndex)
-		{
-			int32 FaceIndex[3] = { 0, 0, 0 };
-			int32 NormalIndex[3] = { 0, 0, 0 };
-			int32 TextureIndex[3] = { 0, 0, 0 };
+	void AddFaceTriangleFanWithTexture(CADLibrary::FTessellationData& Tessellation, const uint32 InTriangleCount, const uint32 TextureCount, uint32& InOutLastTriangleIndex, int32& InOutLastVertexIndex)
+	{
+		int32 FaceIndex[3] = { 0, 0, 0 };
+		int32 NormalIndex[3] = { 0, 0, 0 };
+		int32 TextureIndex[3] = { 0, 0, 0 };
 
-			NormalIndex[0] = TriangulatedIndexes[InOutLastTriangleIndex++];
-			NormalIndex[1] = NormalIndex[0];
-			NormalIndex[2] = NormalIndex[0];
+		NormalIndex[0] = TriangulatedIndexes[InOutLastTriangleIndex++];
+		TextureIndex[0] = TriangulatedIndexes[InOutLastTriangleIndex];
+		InOutLastTriangleIndex += TextureCount;
+		FaceIndex[0] = TriangulatedIndexes[InOutLastTriangleIndex++] / 3;
 
-			TextureIndex[0] = TriangulatedIndexes[InOutLastTriangleIndex];
-			InOutLastTriangleIndex += TextureCount;
-			FaceIndex[0] = TriangulatedIndexes[InOutLastTriangleIndex++] / 3;
+		NormalIndex[1] = TriangulatedIndexes[InOutLastTriangleIndex++];
+		TextureIndex[1] = TriangulatedIndexes[InOutLastTriangleIndex];
+		InOutLastTriangleIndex += TextureCount;
+		FaceIndex[1] = TriangulatedIndexes[InOutLastTriangleIndex++] / 3;
 
-			TextureIndex[1] = TriangulatedIndexes[InOutLastTriangleIndex];
+		for (uint32 TriangleIndex = 2; TriangleIndex < InTriangleCount; TriangleIndex++)
+		{
+			NormalIndex[2] = TriangulatedIndexes[InOutLastTriangleIndex++];
+			TextureIndex[2] = TriangulatedIndexes[InOutLastTriangleIndex];
 			InOutLastTriangleIndex += TextureCount;
-			FaceIndex[1] = TriangulatedIndexes[InOutLastTriangleIndex++] / 3;
+			FaceIndex[2] = TriangulatedIndexes[InOutLastTriangleIndex++] / 3;
 
-			for (uint32 TriangleIndex = 2; TriangleIndex < InTriangleCount; TriangleIndex++)
+			if (AddFace(FaceIndex, Tessellation, InOutLastVertexIndex))
 			{
-				TextureIndex[2] = TriangulatedIndexes[InOutLastTriangleIndex];
-				InOutLastTriangleIndex += TextureCount;
-				FaceIndex[2] = TriangulatedIndexes[InOutLastTriangleIndex++] / 3;
-
-				if (AddFace(FaceIndex, Tessellation, InOutLastVertexIndex))
-				{
-					AddNormals(TessellationNormals, NormalIndex, Tessellation.NormalArray);
-					AddTextureCoordinates(TessellationTexCoords, TextureIndex, Tessellation.TexCoordArray);
-				}
-
-				FaceIndex[1] = FaceIndex[2];
-				TextureIndex[1] = TextureIndex[2];
+				AddNormals(TessellationNormals, NormalIndex, Tessellation.NormalArray);
+				AddTextureCoordinates(TessellationTexCoords, TextureIndex, Tessellation.TexCoordArray);
 			}
+
+			NormalIndex[1] = NormalIndex[2];
+			TextureIndex[1] = TextureIndex[2];
+			FaceIndex[1] = FaceIndex[2];
 		}
+	}
+
+	void AddFaceTriangleStripe(CADLibrary::FTessellationData& Tessellation, const uint32 InTriangleCount, uint32& InOutLastTriangleIndex, int32& /*LastVertexIndex*/)
+	{
+		int32 FaceIndex[3] = { 0, 0, 0 };
+		int32 NormalIndex[3] = { 0, 0, 0 };
 
-		void AddFaceTriangleFanWithTexture(CADLibrary::FTessellationData& Tessellation, const uint32 InTriangleCount, const uint32 TextureCount, uint32& InOutLastTriangleIndex, int32& InOutLastVertexIndex)
+		NormalIndex[0] = TriangulatedIndexes[InOutLastTriangleIndex++];
+		FaceIndex[0] = TriangulatedIndexes[InOutLastTriangleIndex++] / 3;
+		NormalIndex[1] = TriangulatedIndexes[InOutLastTriangleIndex++];
+		FaceIndex[1] = TriangulatedIndexes[InOutLastTriangleIndex++] / 3;
+
+		for (unsigned long TriangleIndex = 2; TriangleIndex < InTriangleCount; TriangleIndex++)
 		{
-			int32 FaceIndex[3] = { 0, 0, 0 };
-			int32 NormalIndex[3] = { 0, 0, 0 };
-			int32 TextureIndex[3] = { 0, 0, 0 };
+			NormalIndex[2] = TriangulatedIndexes[InOutLastTriangleIndex++];
+			FaceIndex[2] = TriangulatedIndexes[InOutLastTriangleIndex++] / 3;
+
+			if (AddFace(FaceIndex, Tessellation, LastVertexIndex))
+			{
+				AddNormals(TessellationNormals, NormalIndex, Tessellation.NormalArray);
+			}
+
+			TriangleIndex++;
+			if (TriangleIndex == InTriangleCount)
+			{
+				break;
+			}
+
+			Swap(FaceIndex[1], FaceIndex[2]);
+			Swap(NormalIndex[1], NormalIndex[2]);
 
 			NormalIndex[0] = TriangulatedIndexes[InOutLastTriangleIndex++];
-			TextureIndex[0] = TriangulatedIndexes[InOutLastTriangleIndex];
-			InOutLastTriangleIndex += TextureCount;
 			FaceIndex[0] = TriangulatedIndexes[InOutLastTriangleIndex++] / 3;
 
-			NormalIndex[1] = TriangulatedIndexes[InOutLastTriangleIndex++];
-			TextureIndex[1] = TriangulatedIndexes[InOutLastTriangleIndex];
+			if (AddFace(FaceIndex, Tessellation, LastVertexIndex))
+			{
+				AddNormals(TessellationNormals, NormalIndex, Tessellation.NormalArray);
+			}
+
+			Swap(FaceIndex[0], FaceIndex[1]);
+			Swap(NormalIndex[0], NormalIndex[1]);
+		}
+	}
+
+	void AddFaceTriangleStripeWithTexture(CADLibrary::FTessellationData& Tessellation, const uint32 InTriangleCount, const uint32 TextureCount, uint32& InOutLastTriangleIndex, int32& InOutLastVertexIndex)
+	{
+		int32 FaceIndex[3] = { 0, 0, 0 };
+		int32 NormalIndex[3] = { 0, 0, 0 };
+		int32 TextureIndex[3] = { 0, 0, 0 };
+
+		NormalIndex[0] = TriangulatedIndexes[InOutLastTriangleIndex++];
+		TextureIndex[0] = TriangulatedIndexes[InOutLastTriangleIndex];
+		InOutLastTriangleIndex += TextureCount;
+		FaceIndex[0] = TriangulatedIndexes[InOutLastTriangleIndex++] / 3;
+		NormalIndex[1] = TriangulatedIndexes[InOutLastTriangleIndex++];
+		TextureIndex[1] = TriangulatedIndexes[InOutLastTriangleIndex];
+		InOutLastTriangleIndex += TextureCount;
+		FaceIndex[1] = TriangulatedIndexes[InOutLastTriangleIndex++] / 3;
+
+		for (unsigned long TriangleIndex = 2; TriangleIndex < InTriangleCount; TriangleIndex++)
+		{
+			NormalIndex[2] = TriangulatedIndexes[InOutLastTriangleIndex++];
+			TextureIndex[2] = TriangulatedIndexes[InOutLastTriangleIndex];
 			InOutLastTriangleIndex += TextureCount;
-			FaceIndex[1] = TriangulatedIndexes[InOutLastTriangleIndex++] / 3;
+			FaceIndex[2] = TriangulatedIndexes[InOutLastTriangleIndex++] / 3;
 
-			for (uint32 TriangleIndex = 2; TriangleIndex < InTriangleCount; TriangleIndex++)
+			if (AddFace(FaceIndex, Tessellation, InOutLastVertexIndex))
 			{
-				NormalIndex[2] = TriangulatedIndexes[InOutLastTriangleIndex++];
-				TextureIndex[2] = TriangulatedIndexes[InOutLastTriangleIndex];
-				InOutLastTriangleIndex += TextureCount;
-				FaceIndex[2] = TriangulatedIndexes[InOutLastTriangleIndex++] / 3;
-
-				if (AddFace(FaceIndex, Tessellation, InOutLastVertexIndex))
-				{
-					AddNormals(TessellationNormals, NormalIndex, Tessellation.NormalArray);
-					AddTextureCoordinates(TessellationTexCoords, TextureIndex, Tessellation.TexCoordArray);
-				}
+				AddNormals(TessellationNormals, NormalIndex, Tessellation.NormalArray);
+				AddTextureCoordinates(TessellationTexCoords, TextureIndex, Tessellation.TexCoordArray);
+			}
 
-				NormalIndex[1] = NormalIndex[2];
-				TextureIndex[1] = TextureIndex[2];
-				FaceIndex[1] = FaceIndex[2];
+			TriangleIndex++;
+			if (TriangleIndex == InTriangleCount)
+			{
+				break;
 			}
-		}
 
-		void AddFaceTriangleStripe(CADLibrary::FTessellationData& Tessellation, const uint32 InTriangleCount, uint32& InOutLastTriangleIndex, int32& /*LastVertexIndex*/)
-		{
-			int32 FaceIndex[3] = { 0, 0, 0 };
-			int32 NormalIndex[3] = { 0, 0, 0 };
+			Swap(FaceIndex[1], FaceIndex[2]);
+			Swap(NormalIndex[1], NormalIndex[2]);
+			Swap(TextureIndex[1], TextureIndex[2]);
 
 			NormalIndex[0] = TriangulatedIndexes[InOutLastTriangleIndex++];
 			FaceIndex[0] = TriangulatedIndexes[InOutLastTriangleIndex++] / 3;
-			NormalIndex[1] = TriangulatedIndexes[InOutLastTriangleIndex++];
-			FaceIndex[1] = TriangulatedIndexes[InOutLastTriangleIndex++] / 3;
 
-			for (unsigned long TriangleIndex = 2; TriangleIndex < InTriangleCount; TriangleIndex++)
+			if (AddFace(FaceIndex, Tessellation, InOutLastVertexIndex))
 			{
-				NormalIndex[2] = TriangulatedIndexes[InOutLastTriangleIndex++];
-				FaceIndex[2] = TriangulatedIndexes[InOutLastTriangleIndex++] / 3;
+				AddNormals(TessellationNormals, NormalIndex, Tessellation.NormalArray);
+			}
 
-				if (AddFace(FaceIndex, Tessellation, LastVertexIndex))
-				{
-					AddNormals(TessellationNormals, NormalIndex, Tessellation.NormalArray);
-				}
+			Swap(FaceIndex[0], FaceIndex[1]);
+			Swap(NormalIndex[0], NormalIndex[1]);
+			Swap(TextureIndex[0], TextureIndex[1]);
+		}
+	}
 
-				TriangleIndex++;
-				if (TriangleIndex == InTriangleCount)
-				{
-					break;
-				}
+	void AddFaceTriangleStripeWithUniqueNormal(CADLibrary::FTessellationData& Tessellation, const uint32 InTriangleCount, uint32& InOutLastTriangleIndex, int32& InOutLastVertexIndex)
+	{
+		int32 FaceIndex[3] = { 0, 0, 0 };
+		int32 NormalIndex[3] = { 0, 0, 0 };
 
-				Swap(FaceIndex[1], FaceIndex[2]);
-				Swap(NormalIndex[1], NormalIndex[2]);
+		NormalIndex[0] = TriangulatedIndexes[InOutLastTriangleIndex++];
+		NormalIndex[1] = NormalIndex[0];
+		NormalIndex[2] = NormalIndex[0];
 
-				NormalIndex[0] = TriangulatedIndexes[InOutLastTriangleIndex++];
-				FaceIndex[0] = TriangulatedIndexes[InOutLastTriangleIndex++] / 3;
+		FaceIndex[0] = TriangulatedIndexes[InOutLastTriangleIndex++] / 3;
+		FaceIndex[1] = TriangulatedIndexes[InOutLastTriangleIndex++] / 3;
 
-				if (AddFace(FaceIndex, Tessellation, LastVertexIndex))
-				{
-					AddNormals(TessellationNormals, NormalIndex, Tessellation.NormalArray);
-				}
+		for (unsigned long TriangleIndex = 2; TriangleIndex < InTriangleCount; TriangleIndex++)
+		{
+			FaceIndex[2] = TriangulatedIndexes[InOutLastTriangleIndex++] / 3;
 
-				Swap(FaceIndex[0], FaceIndex[1]);
-				Swap(NormalIndex[0], NormalIndex[1]);
+			if (AddFace(FaceIndex, Tessellation, InOutLastVertexIndex))
+			{
+				AddNormals(TessellationNormals, NormalIndex, Tessellation.NormalArray);
 			}
-		}
 
-		void AddFaceTriangleStripeWithTexture(CADLibrary::FTessellationData& Tessellation, const uint32 InTriangleCount, const uint32 TextureCount, uint32& InOutLastTriangleIndex, int32& InOutLastVertexIndex)
-		{
-			int32 FaceIndex[3] = { 0, 0, 0 };
-			int32 NormalIndex[3] = { 0, 0, 0 };
-			int32 TextureIndex[3] = { 0, 0, 0 };
+			TriangleIndex++;
+			if (TriangleIndex == InTriangleCount)
+			{
+				break;
+			}
+
+			Swap(FaceIndex[1], FaceIndex[2]);
 
 			NormalIndex[0] = TriangulatedIndexes[InOutLastTriangleIndex++];
-			TextureIndex[0] = TriangulatedIndexes[InOutLastTriangleIndex];
-			InOutLastTriangleIndex += TextureCount;
 			FaceIndex[0] = TriangulatedIndexes[InOutLastTriangleIndex++] / 3;
-			NormalIndex[1] = TriangulatedIndexes[InOutLastTriangleIndex++];
-			TextureIndex[1] = TriangulatedIndexes[InOutLastTriangleIndex];
-			InOutLastTriangleIndex += TextureCount;
-			FaceIndex[1] = TriangulatedIndexes[InOutLastTriangleIndex++] / 3;
 
-			for (unsigned long TriangleIndex = 2; TriangleIndex < InTriangleCount; TriangleIndex++)
+			if (AddFace(FaceIndex, Tessellation, InOutLastVertexIndex))
 			{
-				NormalIndex[2] = TriangulatedIndexes[InOutLastTriangleIndex++];
-				TextureIndex[2] = TriangulatedIndexes[InOutLastTriangleIndex];
-				InOutLastTriangleIndex += TextureCount;
-				FaceIndex[2] = TriangulatedIndexes[InOutLastTriangleIndex++] / 3;
+				AddNormals(TessellationNormals, NormalIndex, Tessellation.NormalArray);
+			}
 
-				if (AddFace(FaceIndex, Tessellation, InOutLastVertexIndex))
-				{
-					AddNormals(TessellationNormals, NormalIndex, Tessellation.NormalArray);
-					AddTextureCoordinates(TessellationTexCoords, TextureIndex, Tessellation.TexCoordArray);
-				}
+			Swap(FaceIndex[0], FaceIndex[1]);
+		}
+	}
 
-				TriangleIndex++;
-				if (TriangleIndex == InTriangleCount)
-				{
-					break;
-				}
+	void AddFaceTriangleStripeWithUniqueNormalAndTexture(CADLibrary::FTessellationData& Tessellation, const uint32 InTriangleCount, const uint32 TextureCount, uint32& InOutLastTriangleIndex, int32& InOutLastVertexIndex)
+	{
+		int32 FaceIndex[3] = { 0, 0, 0 };
+		int32 NormalIndex[3] = { 0, 0, 0 };
+		int32 TextureIndex[3] = { 0, 0, 0 };
 
-				Swap(FaceIndex[1], FaceIndex[2]);
-				Swap(NormalIndex[1], NormalIndex[2]);
-				Swap(TextureIndex[1], TextureIndex[2]);
+		NormalIndex[0] = TriangulatedIndexes[InOutLastTriangleIndex++];
+		NormalIndex[1] = NormalIndex[0];
+		NormalIndex[2] = NormalIndex[0];
 
-				NormalIndex[0] = TriangulatedIndexes[InOutLastTriangleIndex++];
-				FaceIndex[0] = TriangulatedIndexes[InOutLastTriangleIndex++] / 3;
+		TextureIndex[0] = TriangulatedIndexes[InOutLastTriangleIndex];
+		InOutLastTriangleIndex += TextureCount;
+		FaceIndex[0] = TriangulatedIndexes[InOutLastTriangleIndex++] / 3;
 
-				if (AddFace(FaceIndex, Tessellation, InOutLastVertexIndex))
-				{
-					AddNormals(TessellationNormals, NormalIndex, Tessellation.NormalArray);
-				}
+		TextureIndex[1] = TriangulatedIndexes[InOutLastTriangleIndex];
+		InOutLastTriangleIndex += TextureCount;
+		FaceIndex[1] = TriangulatedIndexes[InOutLastTriangleIndex++] / 3;
 
-				Swap(FaceIndex[0], FaceIndex[1]);
-				Swap(NormalIndex[0], NormalIndex[1]);
-				Swap(TextureIndex[0], TextureIndex[1]);
+		for (unsigned long TriangleIndex = 2; TriangleIndex < InTriangleCount; TriangleIndex++)
+		{
+			TextureIndex[2] = TriangulatedIndexes[InOutLastTriangleIndex];
+			InOutLastTriangleIndex += TextureCount;
+			FaceIndex[2] = TriangulatedIndexes[InOutLastTriangleIndex++] / 3;
+
+			if (AddFace(FaceIndex, Tessellation, InOutLastVertexIndex))
+			{
+				AddNormals(TessellationNormals, NormalIndex, Tessellation.NormalArray);
+				AddTextureCoordinates(TessellationTexCoords, TextureIndex, Tessellation.TexCoordArray);
 			}
-		}
 
-		void AddFaceTriangleStripeWithUniqueNormal(CADLibrary::FTessellationData& Tessellation, const uint32 InTriangleCount, uint32& InOutLastTriangleIndex, int32& InOutLastVertexIndex)
-		{
-			int32 FaceIndex[3] = { 0, 0, 0 };
-			int32 NormalIndex[3] = { 0, 0, 0 };
+			TriangleIndex++;
+			if (TriangleIndex == InTriangleCount)
+			{
+				break;
+			}
 
-			NormalIndex[0] = TriangulatedIndexes[InOutLastTriangleIndex++];
-			NormalIndex[1] = NormalIndex[0];
-			NormalIndex[2] = NormalIndex[0];
+			Swap(FaceIndex[1], FaceIndex[2]);
+			Swap(TextureIndex[1], TextureIndex[2]);
 
 			FaceIndex[0] = TriangulatedIndexes[InOutLastTriangleIndex++] / 3;
-			FaceIndex[1] = TriangulatedIndexes[InOutLastTriangleIndex++] / 3;
 
-			for (unsigned long TriangleIndex = 2; TriangleIndex < InTriangleCount; TriangleIndex++)
+			if (AddFace(FaceIndex, Tessellation, InOutLastVertexIndex))
 			{
-				FaceIndex[2] = TriangulatedIndexes[InOutLastTriangleIndex++] / 3;
+				AddNormals(TessellationNormals, NormalIndex, Tessellation.NormalArray);
+			}
 
-				if (AddFace(FaceIndex, Tessellation, InOutLastVertexIndex))
-				{
-					AddNormals(TessellationNormals, NormalIndex, Tessellation.NormalArray);
-				}
+			Swap(FaceIndex[0], FaceIndex[1]);
+			Swap(TextureIndex[0], TextureIndex[1]);
+		}
+	}
 
-				TriangleIndex++;
-				if (TriangleIndex == InTriangleCount)
-				{
-					break;
-				}
+	void ScaleUV(const A3DTopoFace* TopoFace, TArray<FVector2D>& TexCoordArray)
+	{
+		CADLibrary::TUniqueTSObj<A3DTopoFaceData> TopoFaceData(TopoFace);
+		if (!TopoFaceData.IsValid())
+		{
+			return;
+		}
 
-				Swap(FaceIndex[1], FaceIndex[2]);
+		CADLibrary::TUniqueTSObj<A3DDomainData> Domain;
+		if (TopoFaceData->m_bHasTrimDomain)
+		{
+			*Domain = TopoFaceData->m_sSurfaceDomain;
+		}
+		else
+		{
+			CADLibrary::TechSoftInterface::GetSurfaceDomain(TopoFaceData->m_pSurface, *Domain);
+		}
 
-				NormalIndex[0] = TriangulatedIndexes[InOutLastTriangleIndex++];
-				FaceIndex[0] = TriangulatedIndexes[InOutLastTriangleIndex++] / 3;
+		const int32 IsoCurveCount = 7;
+		const double DeltaU = (Domain->m_sMax.m_dX - Domain->m_sMin.m_dX) / (IsoCurveCount - 1.);
+		const double DeltaV = (Domain->m_sMax.m_dY - Domain->m_sMin.m_dY) / (IsoCurveCount - 1.);
 
-				if (AddFace(FaceIndex, Tessellation, InOutLastVertexIndex))
-				{
-					AddNormals(TessellationNormals, NormalIndex, Tessellation.NormalArray);
-				}
+		const A3DSurfBase* A3DSurface = TopoFaceData->m_pSurface;
+
+		FVector NodeMatrix[IsoCurveCount * IsoCurveCount];
+
+		CADLibrary::TUniqueTSObj<A3DVector3dData> Point3D;
+		CADLibrary::TUniqueTSObj<A3DVector2dData> CoordinateObj;
+		A3DVector2dData& Coordinate = *CoordinateObj;
+		Coordinate.m_dX = Domain->m_sMin.m_dX;
+		Coordinate.m_dY = Domain->m_sMin.m_dY;
 
-				Swap(FaceIndex[0], FaceIndex[1]);
+		for (int32 IndexI = 0; IndexI < IsoCurveCount; IndexI++)
+		{
+			for (int32 IndexJ = 0; IndexJ < IsoCurveCount; IndexJ++)
+			{
+				CADLibrary::TechSoftInterface::Evaluate(A3DSurface, Coordinate, 0, Point3D.GetPtr());
+				NodeMatrix[IndexI * IsoCurveCount + IndexJ].X = Point3D->m_dX;
+				NodeMatrix[IndexI * IsoCurveCount + IndexJ].Y = Point3D->m_dY;
+				NodeMatrix[IndexI * IsoCurveCount + IndexJ].Z = Point3D->m_dZ;
+				Coordinate.m_dY += DeltaV;
 			}
+			Coordinate.m_dX += DeltaU;
+			Coordinate.m_dY = Domain->m_sMin.m_dY;
 		}
 
-		void AddFaceTriangleStripeWithUniqueNormalAndTexture(CADLibrary::FTessellationData& Tessellation, const uint32 InTriangleCount, const uint32 TextureCount, uint32& InOutLastTriangleIndex, int32& InOutLastVertexIndex)
+		// Compute length of 7 iso V line
+		double LengthU[IsoCurveCount];
+		double LengthUMax = 0;
+		double LengthUMed = 0;
+
+		for (int32 IndexJ = 0; IndexJ < IsoCurveCount; IndexJ++)
 		{
-			int32 FaceIndex[3] = { 0, 0, 0 };
-			int32 NormalIndex[3] = { 0, 0, 0 };
-			int32 TextureIndex[3] = { 0, 0, 0 };
+			LengthU[IndexJ] = 0;
+			for (int32 IndexI = 0; IndexI < (IsoCurveCount - 1); IndexI++)
+			{
+				LengthU[IndexJ] += FVector::Distance(NodeMatrix[IndexI * IsoCurveCount + IndexJ], NodeMatrix[(IndexI + 1) * IsoCurveCount + IndexJ]);
+			}
+			LengthUMed += LengthU[IndexJ];
+			LengthUMax = FMath::Max(LengthU[IndexJ], LengthUMax);
+		}
+		LengthUMed /= IsoCurveCount;
+		LengthUMed = LengthUMed * 2 / 3 + LengthUMax / 3;
 
-			NormalIndex[0] = TriangulatedIndexes[InOutLastTriangleIndex++];
-			NormalIndex[1] = NormalIndex[0];
-			NormalIndex[2] = NormalIndex[0];
+		// Compute length of 7 iso U line
+		double LengthV[IsoCurveCount];
+		double LengthVMax = 0;
+		double LengthVMed = 0;
 
-			TextureIndex[0] = TriangulatedIndexes[InOutLastTriangleIndex];
-			InOutLastTriangleIndex += TextureCount;
-			FaceIndex[0] = TriangulatedIndexes[InOutLastTriangleIndex++] / 3;
+		for (int32 IndexI = 0; IndexI < IsoCurveCount; IndexI++)
+		{
+			LengthV[IndexI] = 0;
+			for (int32 IndexJ = 0; IndexJ < (IsoCurveCount - 1); IndexJ++)
+			{
+				LengthV[IndexI] += FVector::Distance(NodeMatrix[IndexI * IsoCurveCount + IndexJ], NodeMatrix[IndexI * IsoCurveCount + IndexJ + 1]);
+			}
+			LengthVMed += LengthV[IndexI];
+			LengthVMax = FMath::Max(LengthV[IndexI], LengthVMax);
+		}
+		LengthVMed /= IsoCurveCount;
+		LengthVMed = LengthVMed * 2 / 3 + LengthVMax / 3;
 
-			TextureIndex[1] = TriangulatedIndexes[InOutLastTriangleIndex];
-			InOutLastTriangleIndex += TextureCount;
-			FaceIndex[1] = TriangulatedIndexes[InOutLastTriangleIndex++] / 3;
+		// Texture unit is meter, Coord unit from TechSoft is mm, so TextureScale = 0.001 to convert mm into m
+		const double TextureScale = 0.01;
+		const double UScale = TextureUnit * TextureScale * LengthUMed / (Domain->m_sMax.m_dX - Domain->m_sMin.m_dX);
+		const double VScale = TextureUnit * TextureScale * LengthVMed / (Domain->m_sMax.m_dY - Domain->m_sMin.m_dY);
 
-			for (unsigned long TriangleIndex = 2; TriangleIndex < InTriangleCount; TriangleIndex++)
-			{
-				TextureIndex[2] = TriangulatedIndexes[InOutLastTriangleIndex];
-				InOutLastTriangleIndex += TextureCount;
-				FaceIndex[2] = TriangulatedIndexes[InOutLastTriangleIndex++] / 3;
+		for (FVector2D& TexCoord : TexCoordArray)
+		{
+			TexCoord[0] *= UScale;
+			TexCoord[1] *= VScale;
+		}
+	}
 
-				if (AddFace(FaceIndex, Tessellation, InOutLastVertexIndex))
-				{
-					AddNormals(TessellationNormals, NormalIndex, Tessellation.NormalArray);
-					AddTextureCoordinates(TessellationTexCoords, TextureIndex, Tessellation.TexCoordArray);
-				}
+	void GetBRepFaces()
+	{
+		if (A3DBrepData == nullptr)
+		{
+			return;
+		}
 
-				TriangleIndex++;
-				if (TriangleIndex == InTriangleCount)
+		CADLibrary::TUniqueTSObj<A3DTopoBodyData> TopoBodyData(A3DBrepData);
+		if (TopoBodyData.IsValid())
+		{
+			if (TopoBodyData->m_pContext)
+			{
+				CADLibrary::TUniqueTSObj<A3DTopoContextData> TopoContextData(TopoBodyData->m_pContext);
+				if (TopoContextData.IsValid())
 				{
-					break;
+					if (TopoContextData->m_bHaveScale)
+					{
+						TextureUnit *= TopoContextData->m_dScale;
+					}
 				}
+			}
+		}
 
-				Swap(FaceIndex[1], FaceIndex[2]);
-				Swap(TextureIndex[1], TextureIndex[2]);
+		CADLibrary::TUniqueTSObj<A3DTopoBrepDataData> TopoBrepData(A3DBrepData);
+		if (TopoBrepData.IsValid())
+		{
 
-				FaceIndex[0] = TriangulatedIndexes[InOutLastTriangleIndex++] / 3;
 
-				if (AddFace(FaceIndex, Tessellation, InOutLastVertexIndex))
+			for (A3DUns32 Index = 0; Index < TopoBrepData->m_uiConnexSize; ++Index)
+			{
+				CADLibrary::TUniqueTSObj<A3DTopoConnexData> TopoConnexData(TopoBrepData->m_ppConnexes[Index]);
+				if (TopoConnexData.IsValid())
 				{
-					AddNormals(TessellationNormals, NormalIndex, Tessellation.NormalArray);
+					for (A3DUns32 Sndex = 0; Sndex < TopoConnexData->m_uiShellSize; ++Sndex)
+					{
+						CADLibrary::TUniqueTSObj<A3DTopoShellData> ShellData(TopoConnexData->m_ppShells[Sndex]);
+						if (ShellData.IsValid())
+						{
+							for (A3DUns32 Fndex = 0; Fndex < ShellData->m_uiFaceSize; ++Fndex)
+							{
+								const A3DTopoFace* A3DFace = ShellData->m_ppFaces[Fndex];
+								BodyFaces.Add(A3DFace);
+							}
+						}
+					}
 				}
-
-				Swap(FaceIndex[0], FaceIndex[1]);
-				Swap(TextureIndex[0], TextureIndex[1]);
 			}
 		}
+	}
 
-	private:
-		const A3DTess3D* TessellationPtr;
+private:
+	const A3DTess3D* TessellationPtr = nullptr;
+	const A3DTopoBrepData* A3DBrepData = nullptr;
+	const double BodyUnit = 1;
+	double TextureUnit = 1;
+
+	TArray<const A3DTopoFace*> BodyFaces;
+
+	int32 LastVertexIndex = 0;
+	A3DUns32* TriangulatedIndexes = nullptr;
+	A3DDouble* TessellationNormals;
+	A3DDouble* TessellationTexCoords;				/*!< Array of \ref A3DDouble, as texture coordinates. */
+};
 
-		uint32 LastTrianguleIndex = 0;
-		int32 LastVertexIndex = 0;
-		A3DUns32* TriangulatedIndexes = nullptr;
-		A3DDouble* TessellationNormals;
-		A3DDouble* TessellationTexCoords;				/*!< Array of \ref A3DDouble, as texture coordinates. */
-	};
 } // ns TechSoftUtils
 
 #endif
\ No newline at end of file
diff --git a/Engine/Plugins/Enterprise/DatasmithCADImporter/Source/CADInterfaces/Private/TechSoftUtils.cpp b/Engine/Plugins/Enterprise/DatasmithCADImporter/Source/CADInterfaces/Private/TechSoftUtils.cpp
--- a/Engine/Plugins/Enterprise/DatasmithCADImporter/Source/CADInterfaces/Private/TechSoftUtils.cpp
+++ b/Engine/Plugins/Enterprise/DatasmithCADImporter/Source/CADInterfaces/Private/TechSoftUtils.cpp
@@ -74,7 +74,7 @@ bool GetBodyFromPcrFile(const FString& Filename, const FImportParameters& Import
 	if (JsonObject.IsValid())
 	{
 		double BodyUnit = 1.0;
-		JsonObject->TryGetNumberField(JSON_ENTRY_FILE_UNIT, BodyUnit);
+		JsonObject->TryGetNumberField(JSON_ENTRY_BODY_UNIT, BodyUnit);
 
 		for (A3DUns32 Index = 0; Index < PartDefinitionData->m_uiRepItemsSize; ++Index)
 		{
@@ -149,7 +149,7 @@ FUniqueTechSoftModelFile SaveBodiesToPrcFile(void** Bodies, uint32 BodyCount, co
 #endif
 }
 
-bool FillBodyMesh(void* BodyPtr, const FImportParameters& ImportParameters, double FileUnit, FBodyMesh& BodyMesh)
+bool FillBodyMesh(void* BodyPtr, const FImportParameters& ImportParameters, double BodyUnit, FBodyMesh& BodyMesh)
 {
 #if defined USE_TECHSOFT_SDK && !defined CADKERNEL_DEV
 	A3DRiRepresentationItem* RepresentationItemPtr = (A3DRiRepresentationItem*)BodyPtr;
@@ -158,9 +158,8 @@ bool FillBodyMesh(void* BodyPtr, const FImportParameters& ImportParameters, doub
 	A3DEntityGetType(RepresentationItemPtr, &Type);
 	if (Type == kA3DTypeRiPolyBrepModel)
 	{
-		TUniqueTSObj<A3DRiRepresentationItemData> RepresentationItemData(RepresentationItemPtr);
-		TechSoftInterfaceUtils::FTechSoftTessellationExtractor Extractor(RepresentationItemData->m_pTessBase);
-		return Extractor.FillBodyMesh(BodyMesh, FileUnit);
+		TechSoftInterfaceUtils::FTechSoftTessellationExtractor Extractor(RepresentationItemPtr, BodyUnit);
+		return Extractor.FillBodyMesh(BodyMesh);
 	}
 
 	// TUniqueTechSoftObj does not work in this case
@@ -169,9 +168,9 @@ bool FillBodyMesh(void* BodyPtr, const FImportParameters& ImportParameters, doub
 	TessellationParameters->m_eTessellationLevelOfDetail = kA3DTessLODUserDefined; // Enum to specify predefined values for some following members.
 	TessellationParameters->m_bUseHeightInsteadOfRatio = A3D_TRUE;
 	TessellationParameters->m_dMaxChordHeight = ImportParameters.GetChordTolerance() * 10.; // cm to mm
-	if (!FMath::IsNearlyZero(FileUnit))
+	if (!FMath::IsNearlyZero(BodyUnit))
 	{
-		TessellationParameters->m_dMaxChordHeight /= FileUnit;
+		TessellationParameters->m_dMaxChordHeight /= BodyUnit;
 	}
 
 	TessellationParameters->m_dAngleToleranceDeg = ImportParameters.GetMaxNormalAngle();
@@ -199,8 +198,9 @@ bool FillBodyMesh(void* BodyPtr, const FImportParameters& ImportParameters, doub
 		}
 	}
 
-	TechSoftInterfaceUtils::FTechSoftTessellationExtractor Extractor(RepresentationItemData->m_pTessBase);
-	return Extractor.FillBodyMesh(BodyMesh, FileUnit);
+	TechSoftInterfaceUtils::FTechSoftTessellationExtractor Extractor(RepresentationItemPtr, BodyUnit);
+	return Extractor.FillBodyMesh(BodyMesh);
+
 #else
 	return false;
 #endif
diff --git a/Engine/Plugins/Enterprise/DatasmithCADImporter/Source/CADInterfaces/Public/TUniqueTechSoftObj.h b/Engine/Plugins/Enterprise/DatasmithCADImporter/Source/CADInterfaces/Public/TUniqueTechSoftObj.h
--- a/Engine/Plugins/Enterprise/DatasmithCADImporter/Source/CADInterfaces/Public/TUniqueTechSoftObj.h
+++ b/Engine/Plugins/Enterprise/DatasmithCADImporter/Source/CADInterfaces/Public/TUniqueTechSoftObj.h
@@ -262,6 +262,8 @@ CADINTERFACES_API void TUniqueTSObj<A3DCrvPolyLineData>::InitializeData();
 template<>
 CADINTERFACES_API void TUniqueTSObj<A3DCrvTransformData>::InitializeData();
 template<>
+CADINTERFACES_API void TUniqueTSObj<A3DDomainData>::InitializeData();
+template<>
 CADINTERFACES_API void TUniqueTSObj<A3DGlobalData>::InitializeData();
 template<>
 CADINTERFACES_API void TUniqueTSObj<A3DGraphicsData>::InitializeData();
@@ -363,6 +365,10 @@ template<>
 CADINTERFACES_API void TUniqueTSObj<A3DTopoUniqueVertexData>::InitializeData();
 template<>
 CADINTERFACES_API void TUniqueTSObj<A3DTopoWireEdgeData>::InitializeData();
+template<>
+CADINTERFACES_API void TUniqueTSObj<A3DVector2dData>::InitializeData();
+template<>
+CADINTERFACES_API void TUniqueTSObj<A3DVector3dData>::InitializeData();
 
 template<>
 CADINTERFACES_API A3DStatus TUniqueTSObj<A3DAsmModelFileData>::GetData(const A3DEntity* InEntityPtr);
@@ -401,6 +407,8 @@ CADINTERFACES_API A3DStatus TUniqueTSObj<A3DCrvPolyLineData>::GetData(const A3DE
 template<>
 CADINTERFACES_API A3DStatus TUniqueTSObj<A3DCrvTransformData>::GetData(const A3DEntity* InEntityPtr);
 template<>
+CADINTERFACES_API A3DStatus TUniqueTSObj<A3DDomainData>::GetData(const A3DEntity* InEntityPtr);
+template<>
 CADINTERFACES_API A3DStatus TUniqueTSObj<A3DGlobalData>::GetData(const A3DEntity* InEntityPtr);
 template<>
 CADINTERFACES_API A3DStatus TUniqueTSObj<A3DGraphicsData>::GetData(const A3DEntity* InEntityPtr);
@@ -502,6 +510,10 @@ template<>
 CADINTERFACES_API A3DStatus TUniqueTSObj<A3DTopoUniqueVertexData>::GetData(const A3DEntity* InEntityPtr);
 template<>
 CADINTERFACES_API A3DStatus TUniqueTSObj<A3DTopoWireEdgeData>::GetData(const A3DEntity* InEntityPtr);
+template<>
+CADINTERFACES_API A3DStatus TUniqueTSObj<A3DVector2dData>::GetData(const A3DEntity* InEntityPtr);
+template<>
+CADINTERFACES_API A3DStatus TUniqueTSObj<A3DVector3dData>::GetData(const A3DEntity* InEntityPtr);
 
 template<>
 CADINTERFACES_API const A3DEntity* TUniqueTSObj<A3DAsmModelFileData>::GetDefaultIndexerValue() const;
@@ -540,6 +552,8 @@ CADINTERFACES_API const A3DEntity* TUniqueTSObj<A3DCrvPolyLineData>::GetDefaultI
 template<>
 CADINTERFACES_API const A3DEntity* TUniqueTSObj<A3DCrvTransformData>::GetDefaultIndexerValue() const;
 template<>
+CADINTERFACES_API const A3DEntity* TUniqueTSObj<A3DDomainData>::GetDefaultIndexerValue() const;
+template<>
 CADINTERFACES_API const A3DEntity* TUniqueTSObj<A3DGlobalData>::GetDefaultIndexerValue() const;
 template<>
 CADINTERFACES_API const A3DEntity* TUniqueTSObj<A3DGraphicsData>::GetDefaultIndexerValue() const;
@@ -641,6 +655,10 @@ template<>
 CADINTERFACES_API const A3DEntity* TUniqueTSObj<A3DTopoUniqueVertexData>::GetDefaultIndexerValue() const;
 template<>
 CADINTERFACES_API const A3DEntity* TUniqueTSObj<A3DTopoWireEdgeData>::GetDefaultIndexerValue() const;
+template<>
+CADINTERFACES_API const A3DEntity* TUniqueTSObj<A3DVector2dData>::GetDefaultIndexerValue() const;
+template<>
+CADINTERFACES_API const A3DEntity* TUniqueTSObj<A3DVector3dData>::GetDefaultIndexerValue() const;
 
 // TUniqueTSObjFromIndex -----------------------------------
 
diff --git a/Engine/Plugins/Enterprise/DatasmithCADImporter/Source/CADInterfaces/Public/TechSoftInterface.h b/Engine/Plugins/Enterprise/DatasmithCADImporter/Source/CADInterfaces/Public/TechSoftInterface.h
--- a/Engine/Plugins/Enterprise/DatasmithCADImporter/Source/CADInterfaces/Public/TechSoftInterface.h
+++ b/Engine/Plugins/Enterprise/DatasmithCADImporter/Source/CADInterfaces/Public/TechSoftInterface.h
@@ -26,7 +26,7 @@ THIRD_PARTY_INCLUDES_END
 #include "Windows/HideWindowsPlatformTypes.h"
 #endif
 
-#define JSON_ENTRY_FILE_UNIT		TEXT("FileUnit")
+#define JSON_ENTRY_BODY_UNIT		TEXT("BodyUnit")
 #define JSON_ENTRY_COLOR_NAME		TEXT("ColorName")
 #define JSON_ENTRY_MATERIAL_NAME	TEXT("MaterialName")
 
@@ -89,6 +89,16 @@ CADINTERFACES_API A3DGlobal* GetGlobalPointer();
 CADINTERFACES_API A3DEntity* GetPointerFromIndex(const uint32 Index, const A3DEEntityType Type);
 
 CADINTERFACES_API A3DStatus GetSurfaceAsNurbs(const A3DSurfBase* SurfacePtr, A3DSurfNurbsData* DataPtr, A3DDouble Tolerance, A3DBool bUseSameParameterization);
+CADINTERFACES_API A3DStatus GetSurfaceDomain(const A3DSurfBase* SurfacePtr, A3DDomainData& OutDomain);
+
+/**
+ * @param OutPointAndDerivatives size must be equal to (Derivatives+1)*(Derivatives+1)
+ * @param Derivatives The number of derivatives must be less or equal to 2
+ * @param OutPointAndDerivatives = Point, Derivatives=1 OutPointAndDerivatives = Point + dU + dV, ...
+ * @see A3DSDKAdvancedTools.h for more details
+ */
+CADINTERFACES_API A3DStatus Evaluate(const A3DSurfBase* SurfacePtr, const A3DVector2dData& UVParameter, A3DUns32 Derivatives, A3DVector3dData* OutPointAndDerivatives);
+
 CADINTERFACES_API A3DStatus GetCurveAsNurbs(const A3DCrvBase* A3DCurve, A3DCrvNurbsData* DataPtr, A3DDouble Tolerance, A3DBool bUseSameParameterization);
 
 CADINTERFACES_API A3DStatus GetOriginalFilePathName(const A3DAsmProductOccurrence* A3DOccurrencePtr, A3DUTF8Char** FilePathUTF8Ptr);
@@ -136,6 +146,8 @@ CADINTERFACES_API A3DUns32 InsertGraphStyle(const A3DGraphStyleData& InStyleData
 
 CADINTERFACES_API double GetModelFileUnit(const A3DAsmModelFile* pModelFile);
 
+CADINTERFACES_API A3DStatus AddAttribute(A3DEntity* EntityPtr, const TCHAR*, const TCHAR*);
+
 /**
  * @param ToleranceInCM : The maximum tolerance for the sewing (in cm).
  */
diff --git a/Engine/Plugins/Enterprise/DatasmithCADImporter/Source/CADInterfaces/Public/TechSoftUtils.h b/Engine/Plugins/Enterprise/DatasmithCADImporter/Source/CADInterfaces/Public/TechSoftUtils.h
--- a/Engine/Plugins/Enterprise/DatasmithCADImporter/Source/CADInterfaces/Public/TechSoftUtils.h
+++ b/Engine/Plugins/Enterprise/DatasmithCADImporter/Source/CADInterfaces/Public/TechSoftUtils.h
@@ -6,10 +6,6 @@
 #include "CADData.h"
 #include "TechSoftInterface.h"
 
-#define JSON_ENTRY_BODY_UNIT		TEXT("BodyUnit")
-#define JSON_ENTRY_COLOR_NAME		TEXT("ColorName")
-#define JSON_ENTRY_MATERIAL_NAME	TEXT("MaterialName")
-
 struct FColor;
 
 typedef void A3DAsmPartDefinition;
diff --git a/Engine/Plugins/Enterprise/DatasmithCADImporter/Source/CADLibrary/Private/CADKernelTools.cpp b/Engine/Plugins/Enterprise/DatasmithCADImporter/Source/CADLibrary/Private/CADKernelTools.cpp
--- a/Engine/Plugins/Enterprise/DatasmithCADImporter/Source/CADLibrary/Private/CADKernelTools.cpp
+++ b/Engine/Plugins/Enterprise/DatasmithCADImporter/Source/CADLibrary/Private/CADKernelTools.cpp
@@ -53,7 +53,7 @@ namespace CADLibrary
 
 		for (FVector& Vertex : VertexArray)
 		{
-			Vertex *= Context.ImportParams.GetScaleFactor();
+			Vertex *= 0.1; // mm (CADKernel unit) to cm (UE unit)
 		}
 
 		int32 VertexCount = VertexArray.Num();
@@ -206,13 +206,17 @@ namespace CADLibrary
 					// Set patch id attribute
 					PatchGroups[PolygonID] = Face.GetPatchId();
 				}
-
+				
 				for (int32 IndexFace = 0; IndexFace < MeshVertexInstanceIDs.Num(); IndexFace += 3)
 				{
 					for (int32 Index = 0; Index < TriangleCount; Index++)
 					{
 						const FVertexInstanceID VertexInstanceID = MeshVertexInstanceIDs[IndexFace + Index];
-						VertexInstanceUVs.Set(VertexInstanceID, UVChannel, FVector2f(FaceMesh->UVMap[TriangleVertexIndices[IndexFace + Orientation[Index]]]));
+
+						// TODO
+						// The UV scaling should be done in CADKernel
+						const double ScaleUV = 0.001; // mm to m
+						VertexInstanceUVs.Set(VertexInstanceID, UVChannel, FVector2f(FaceMesh->UVMap[TriangleVertexIndices[IndexFace + Orientation[Index]]] * ScaleUV));
 
 						VertexInstanceColors[VertexInstanceID] = FLinearColor::White;
 						VertexInstanceTangents[VertexInstanceID] = FVector3f(ForceInitToZero);
diff --git a/Engine/Plugins/Enterprise/DatasmithCADImporter/Source/DatasmithOpenNurbsTranslator/Private/DatasmithOpenNurbsTranslator.cpp b/Engine/Plugins/Enterprise/DatasmithCADImporter/Source/DatasmithOpenNurbsTranslator/Private/DatasmithOpenNurbsTranslator.cpp
--- a/Engine/Plugins/Enterprise/DatasmithCADImporter/Source/DatasmithOpenNurbsTranslator/Private/DatasmithOpenNurbsTranslator.cpp
+++ b/Engine/Plugins/Enterprise/DatasmithCADImporter/Source/DatasmithOpenNurbsTranslator/Private/DatasmithOpenNurbsTranslator.cpp
@@ -647,8 +647,9 @@ public:
 			TranslationCache = MakeShared<FTranslationCache>();
 		}
 
-		CADLibrary::FImportParameters ImportParameters(0.001, 0.1, FDatasmithUtils::EModelCoordSystem::ZUp_RightHanded_FBXLegacy);
-		ImportParameters.SwitchOffUVMapScaling();
+		const double IPMetricUnit = 0.001; // CAD Tools works in mm
+		const double IPScaleFactor = 1.;
+		CADLibrary::FImportParameters ImportParameters(IPMetricUnit, IPScaleFactor, FDatasmithUtils::EModelCoordSystem::ZUp_RightHanded_FBXLegacy);
 
 		if (CADLibrary::FImportParameters::bGDisableCADKernelTessellation)
 		{
@@ -751,7 +752,22 @@ private:
 
 	double GetScaleFactor() const
 	{
-		return CADModelConverter->GetScaleFactor();
+		return ScaleFactor;
+	}
+
+	/**
+	 * @param NewMetricUnit compare to mm
+	 */
+	void SetMetricUnit(const double NewMetricUnit)
+	{
+		if (FMath::IsNearlyZero(NewMetricUnit))
+		{
+			return;
+		}
+
+		// default unit is mm
+		ScaleFactor = 0.1 / NewMetricUnit;
+		OpenNurbsBRepConverter->SetScaleFactor(1 / NewMetricUnit);
 	}
 
 private:
@@ -765,6 +781,9 @@ private:
 	uint32 OpenNurbsOptionsHash;
 	FDatasmithImportBaseOptions BaseOptions;
 
+	// default unit is mm
+	double ScaleFactor = 0.1;
+
 private:
 	// For OpenNurbs archive parsing
 	// Ref. GP3DMReader
@@ -2345,8 +2364,8 @@ bool FOpenNurbsTranslatorImpl::Read(ON_BinaryFile& Archive, TSharedRef<IDatasmit
 		bResult = false;
 	}
 
-	double MetricUnit = 1 / Settings.m_ModelUnitsAndTolerances.Scale(ON::LengthUnitSystem::Meters);
-	CADModelConverter->SetMetricUnit(MetricUnit);
+	const double MetricUnit = Settings.m_ModelUnitsAndTolerances.Scale(ON::LengthUnitSystem::Millimeters);
+	SetMetricUnit(MetricUnit);
 
 	// Step 4: REQUIRED - Read bitmap table (it can be empty)
 	int Count = 0;
diff --git a/Engine/Plugins/Enterprise/DatasmithCADImporter/Source/DatasmithOpenNurbsTranslator/Private/OpenNurbsBRepConverter.h b/Engine/Plugins/Enterprise/DatasmithCADImporter/Source/DatasmithOpenNurbsTranslator/Private/OpenNurbsBRepConverter.h
--- a/Engine/Plugins/Enterprise/DatasmithCADImporter/Source/DatasmithOpenNurbsTranslator/Private/OpenNurbsBRepConverter.h
+++ b/Engine/Plugins/Enterprise/DatasmithCADImporter/Source/DatasmithOpenNurbsTranslator/Private/OpenNurbsBRepConverter.h
@@ -9,5 +9,13 @@ class IOpenNurbsBRepConverter
 {
 public:
 	virtual bool AddBRep(ON_Brep& Brep, const ON_3dVector& Offset) = 0;
+	void SetScaleFactor(double NewScaleFactor)
+	{
+		ensure(!FMath::IsNearlyZero(NewScaleFactor));
+		ScaleFactor = NewScaleFactor;
+	}
+
+protected:
+	double ScaleFactor = 1; // mm to mm 
 };
 
diff --git a/Engine/Plugins/Enterprise/DatasmithCADImporter/Source/DatasmithOpenNurbsTranslator/Private/OpenNurbsBRepToCADKernelConverter.cpp b/Engine/Plugins/Enterprise/DatasmithCADImporter/Source/DatasmithOpenNurbsTranslator/Private/OpenNurbsBRepToCADKernelConverter.cpp
--- a/Engine/Plugins/Enterprise/DatasmithCADImporter/Source/DatasmithOpenNurbsTranslator/Private/OpenNurbsBRepToCADKernelConverter.cpp
+++ b/Engine/Plugins/Enterprise/DatasmithCADImporter/Source/DatasmithOpenNurbsTranslator/Private/OpenNurbsBRepToCADKernelConverter.cpp
@@ -84,6 +84,18 @@ TSharedRef<CADKernel::FSurface> FOpenNurbsBRepToCADKernelConverter::AddSurface(O
 		}
 	}
 
+	// Scale ControlPoints into mm
+	if(!FMath::IsNearlyEqual(ScaleFactor, 1.))
+	{
+		int32 Offset = NurbsData.bIsRational ? 4 : 3;
+		for (int32 Index = 0; Index < NurbsData.HomogeneousPoles.Num(); Index += Offset)
+		{
+			ControlPoints[Index + 0] *= ScaleFactor;
+			ControlPoints[Index + 1] *= ScaleFactor;
+			ControlPoints[Index + 2] *= ScaleFactor;
+		}
+	}
+
 #ifdef REMOVE_NEGATIVE_WEIGHT
 	if (NurbsData.bIsRational)
 	{
diff --git a/Engine/Plugins/Enterprise/DatasmithCADImporter/Source/DatasmithOpenNurbsTranslator/Private/OpenNurbsBRepToTechSoftConverter.cpp b/Engine/Plugins/Enterprise/DatasmithCADImporter/Source/DatasmithOpenNurbsTranslator/Private/OpenNurbsBRepToTechSoftConverter.cpp
--- a/Engine/Plugins/Enterprise/DatasmithCADImporter/Source/DatasmithOpenNurbsTranslator/Private/OpenNurbsBRepToTechSoftConverter.cpp
+++ b/Engine/Plugins/Enterprise/DatasmithCADImporter/Source/DatasmithOpenNurbsTranslator/Private/OpenNurbsBRepToTechSoftConverter.cpp
@@ -19,32 +19,18 @@ namespace OpenNurbsUtils
 {
 enum EAxis { U, V };
 
-template<typename PointType>
-void SetA3DVector(const PointType& OpenNurbsPoint, A3DVector3dData& OutTSPoint)
-{
-	OutTSPoint.m_dX = OpenNurbsPoint.x;
-	OutTSPoint.m_dY = OpenNurbsPoint.y;
-	OutTSPoint.m_dZ = OpenNurbsPoint.z;
-};
-
-template<typename PointType>
-void SetA3DVectorAndWeight(const PointType& OpenNurbsPoint, A3DVector3dData& OutTSPoint, A3DDouble& OutWeight)
-{
-	SetA3DVector(OpenNurbsPoint, OutTSPoint);
-	OutWeight = OpenNurbsPoint.w;
-};
-
-
 struct FOpenNurbsSurfaceInfo
 {
-	FOpenNurbsSurfaceInfo(const ON_NurbsSurface& InOpenNurbsSurface, A3DSurfNurbsData& InData)
+	FOpenNurbsSurfaceInfo(const ON_NurbsSurface& InOpenNurbsSurface, A3DSurfNurbsData& InData, const double InScale)
 		: OpenNurbsSurface(InOpenNurbsSurface)
 		, Data(InData)
+		, Scale(InScale)
 	{
 	}
 
 	const ON_NurbsSurface& OpenNurbsSurface;
 	A3DSurfNurbsData& Data;
+	const double Scale;
 
 	TArray<A3DDouble> UNodalVector;
 	TArray<A3DDouble> VNodalVector;
@@ -52,6 +38,21 @@ struct FOpenNurbsSurfaceInfo
 	TArray<A3DDouble> Weights;
 	TArray<A3DVector3dData> ControlPoints;
 
+	template<typename PointType>
+	void SetA3DVector(const PointType& OpenNurbsPoint, A3DVector3dData& OutTSPoint)
+	{
+		OutTSPoint.m_dX = OpenNurbsPoint.x * Scale;
+		OutTSPoint.m_dY = OpenNurbsPoint.y * Scale;
+		OutTSPoint.m_dZ = OpenNurbsPoint.z * Scale;
+	};
+
+	template<typename PointType>
+	void SetA3DVectorAndWeight(const PointType& OpenNurbsPoint, A3DVector3dData& OutTSPoint, A3DDouble& OutWeight)
+	{
+		SetA3DVector(OpenNurbsPoint, OutTSPoint);
+		OutWeight = OpenNurbsPoint.w;
+	};
+
 	void FillPerAxisInfo(EAxis Axis)
 	{
 		A3DUns32& Degree = Axis == EAxis::U ? Data.m_uiUDegree : Data.m_uiVDegree;
@@ -134,7 +135,7 @@ A3DSurfBase* FOpenNurbsBRepToTechSoftConverter::CreateSurface(const ON_NurbsSurf
 {
 	CADLibrary::TUniqueTSObj<A3DSurfNurbsData> NurbsSurfaceData;
 
-	OpenNurbsUtils::FOpenNurbsSurfaceInfo NurbsInfo(OpenNurbsSurface, *NurbsSurfaceData);
+	OpenNurbsUtils::FOpenNurbsSurfaceInfo NurbsInfo(OpenNurbsSurface, *NurbsSurfaceData, ScaleFactor);
 	NurbsInfo.Populate();
 
 	return CADLibrary::TechSoftInterface::CreateSurfaceNurbs(*NurbsSurfaceData);
diff --git a/Engine/Plugins/Enterprise/DatasmithCADImporter/Source/DatasmithWireTranslator/Private/DatasmithWireTranslator.cpp b/Engine/Plugins/Enterprise/DatasmithCADImporter/Source/DatasmithWireTranslator/Private/DatasmithWireTranslator.cpp
--- a/Engine/Plugins/Enterprise/DatasmithCADImporter/Source/DatasmithWireTranslator/Private/DatasmithWireTranslator.cpp
+++ b/Engine/Plugins/Enterprise/DatasmithCADImporter/Source/DatasmithWireTranslator/Private/DatasmithWireTranslator.cpp
@@ -228,7 +228,7 @@ public:
 		}
 
 		const double MetricUnit = 0.001; // ImportParameters MetricUnit is defined in meter
-		const double ScaleFactor = 0.1;  // MetricUnit to UE unit i.e. cm
+		const double ScaleFactor = 1;    // MetricUnit to Modeler unit i.e. mm
 		CADLibrary::FImportParameters ImportParameters(MetricUnit, ScaleFactor);
 		if(CADLibrary::FImportParameters::bGDisableCADKernelTessellation)
 		{
diff --git a/Engine/Plugins/Enterprise/DatasmithCADImporter/Source/ParametricSurface/Private/TechSoft/CADModelToTechSoftConverterBase.cpp b/Engine/Plugins/Enterprise/DatasmithCADImporter/Source/ParametricSurface/Private/TechSoft/CADModelToTechSoftConverterBase.cpp
--- a/Engine/Plugins/Enterprise/DatasmithCADImporter/Source/ParametricSurface/Private/TechSoft/CADModelToTechSoftConverterBase.cpp
+++ b/Engine/Plugins/Enterprise/DatasmithCADImporter/Source/ParametricSurface/Private/TechSoft/CADModelToTechSoftConverterBase.cpp
@@ -51,10 +51,10 @@ bool FCADModelToTechSoftConverterBase::SaveModel(const TCHAR* InFolderPath, TSha
 
 	FString JsonString;
 	{
-		// Save file unit and default color and material attributes in a json string
+		// Save body unit and default color and material attributes in a json string
 		// This will be used when the file is reloaded
 		TSharedPtr<FJsonObject> JsonObject = MakeShared<FJsonObject>();
-		JsonObject->SetNumberField(JSON_ENTRY_FILE_UNIT, 1);
+		JsonObject->SetNumberField(JSON_ENTRY_BODY_UNIT, 0.1); // ALL BRep from Alias or Rhino are defined in mm
 		JsonObject->SetNumberField(JSON_ENTRY_COLOR_NAME, 0);
 		JsonObject->SetNumberField(JSON_ENTRY_MATERIAL_NAME, 0);
 
@@ -78,8 +78,8 @@ bool FCADModelToTechSoftConverterBase::Tessellate(const CADLibrary::FMeshParamet
 #ifdef USE_TECHSOFT_SDK
 	for (A3DRiRepresentationItem* Representation : RiRepresentationItems)
 	{
-		const double FileUnit = ImportParameters.GetMetricUnit() * 1000; // ImportParameters MetricUnit is defined in meter, CADLibrary::TechSoftInterface::SewBReps expectes it in mm
-		CADLibrary::TechSoftUtils::FillBodyMesh(Representation, ImportParameters, FileUnit, BodyMesh);
+		const double BodyUnit = 0.1; // CAD Models from (Wire or Rhino) imported in TechSoft are imported in mm so BodyUnit = 0.1 
+		CADLibrary::TechSoftUtils::FillBodyMesh(Representation, ImportParameters, BodyUnit, BodyMesh);
 	}
 #endif
 
