PatchId wasn't set at the import
#swarm https://p4-swarm.epicgames.net/reviews/19605483

diff --git a/Engine/Plugins/Enterprise/DatasmithCADImporter/Source/CADInterfaces/Private/TechSoftInterfaceUtils.inl b/Engine/Plugins/Enterprise/DatasmithCADImporter/Source/CADInterfaces/Private/TechSoftInterfaceUtils.inl
--- a/Engine/Plugins/Enterprise/DatasmithCADImporter/Source/CADInterfaces/Private/TechSoftInterfaceUtils.inl
+++ b/Engine/Plugins/Enterprise/DatasmithCADImporter/Source/CADInterfaces/Private/TechSoftInterfaceUtils.inl
@@ -166,6 +166,9 @@ private:
 			const A3DTessFaceData& FaceTessData = A3DTessellationData->m_psFaceTessData[Index];
 			FTessellationData& Tessellation = Faces.Emplace_GetRef();
 
+			// there is a bijection between A3DTess3DData->m_psFaceTessData and A3DTopoShellData->m_ppFaces
+			Tessellation.PatchId = Index;
+
 			Tessellation.MaterialName = FTechSoftInterface::InvalidScriptIndex;
 			if (FaceTessData.m_uiStyleIndexesSize > 0)
 			{
