In the jira case, the material is defined at the part level but not at the BRepModel level. The material of the part is assigned to the Body
#swarm https://p4-swarm.epicgames.net/reviews/19605813

diff --git a/Engine/Plugins/Enterprise/DatasmithCADImporter/Source/CADInterfaces/Private/TechSoftFileParser.cpp b/Engine/Plugins/Enterprise/DatasmithCADImporter/Source/CADInterfaces/Private/TechSoftFileParser.cpp
--- a/Engine/Plugins/Enterprise/DatasmithCADImporter/Source/CADInterfaces/Private/TechSoftFileParser.cpp
+++ b/Engine/Plugins/Enterprise/DatasmithCADImporter/Source/CADInterfaces/Private/TechSoftFileParser.cpp
@@ -1223,6 +1223,21 @@ FCadId FTechSoftFileParser::TraverseBRepModel(A3DRiBrepModel* BRepModelPtr, cons
 		return CADFileData.GetBodyAt(*BodyIndexPtr).ObjectId;
 	}
 
+	// if BRep model has not material or color, add part one
+	if (BRepMetaData.MaterialName == 0 && BRepMetaData.ColorName == 0)
+	{
+		if (PartMetaData.MaterialName)
+		{
+			BRepMetaData.MaterialName = PartMetaData.MaterialName;
+			BRepMetaData.MetaData.Add(TEXT("MaterialName"), FString::Printf(TEXT("%u"), PartMetaData.MaterialName));
+		}
+		if (PartMetaData.ColorName)
+		{
+			BRepMetaData.ColorName = PartMetaData.ColorName;
+			BRepMetaData.MetaData.Add(TEXT("ColorName"), FString::Printf(TEXT("%u"), PartMetaData.ColorName));
+		}
+	}
+
 	ExtractSpecificMetaData(BRepModelPtr, BRepMetaData);
 
 	FMatrix Matrix = FMatrix::Identity;
