the error appears with empty mesh (mesh without triangles). To avoid the error, meshActors with empty mesh are not build
#swarm https://p4-swarm.epicgames.net/reviews/19608475

diff --git a/Engine/Plugins/Enterprise/DatasmithCADImporter/Source/CADInterfaces/Private/TechSoftFileParser.cpp b/Engine/Plugins/Enterprise/DatasmithCADImporter/Source/CADInterfaces/Private/TechSoftFileParser.cpp
--- a/Engine/Plugins/Enterprise/DatasmithCADImporter/Source/CADInterfaces/Private/TechSoftFileParser.cpp
+++ b/Engine/Plugins/Enterprise/DatasmithCADImporter/Source/CADInterfaces/Private/TechSoftFileParser.cpp
@@ -389,6 +389,14 @@ void FTechSoftFileParser::GenerateBodyMesh(A3DRiRepresentationItem* Representati
 		TechSoftUtils::FillBodyMesh(Representation, CADFileData.GetImportParameters(), Body.BodyUnit, BodyMesh);
 	}
 
+	if (BodyMesh.TriangleCount == 0)
+	{
+		// the mesh of the body is empty, the body is deleted.
+		// Todo (jira UETOOL-5148): add a boolean in Body to flag that the body should not be build
+		Body.ParentId = 0;
+		Body.MeshActorName = 0;
+	}
+
 	// Convert material
 	FCADUUID DefaultColorName = Body.ColorFaceSet.Num() > 0 ? *Body.ColorFaceSet.begin() : 0;
 	FCADUUID DefaultMaterialName = Body.MaterialFaceSet.Num() > 0 ? *Body.MaterialFaceSet.begin() : 0;
diff --git a/Engine/Plugins/Enterprise/DatasmithCADImporter/Source/CADInterfaces/Private/TechSoftInterfaceUtils.inl b/Engine/Plugins/Enterprise/DatasmithCADImporter/Source/CADInterfaces/Private/TechSoftInterfaceUtils.inl
--- a/Engine/Plugins/Enterprise/DatasmithCADImporter/Source/CADInterfaces/Private/TechSoftInterfaceUtils.inl
+++ b/Engine/Plugins/Enterprise/DatasmithCADImporter/Source/CADInterfaces/Private/TechSoftInterfaceUtils.inl
@@ -334,6 +334,7 @@ private:
 				const A3DTopoFace* TopoFace = BodyFaces[Index];
 				ScaleUV(TopoFace, Tessellation.TexCoordArray);
 			}
+			BodyMesh.TriangleCount += (Tessellation.VertexIndices.Num() / 3.);
 		}
 	}
 
@@ -908,8 +909,6 @@ private:
 		CADLibrary::TUniqueTSObj<A3DTopoBrepDataData> TopoBrepData(A3DBrepData);
 		if (TopoBrepData.IsValid())
 		{
-
-
 			for (A3DUns32 Index = 0; Index < TopoBrepData->m_uiConnexSize; ++Index)
 			{
 				CADLibrary::TUniqueTSObj<A3DTopoConnexData> TopoConnexData(TopoBrepData->m_ppConnexes[Index]);
