The problem was caused by the default material not being assigned to untrimmed surfaces
#swarm https://p4-swarm.epicgames.net/reviews/19647802

diff --git a/Engine/Plugins/Enterprise/DatasmithCADImporter/Source/DatasmithWireTranslator/Private/AliasModelToTechSoftConverter.cpp b/Engine/Plugins/Enterprise/DatasmithCADImporter/Source/DatasmithWireTranslator/Private/AliasModelToTechSoftConverter.cpp
--- a/Engine/Plugins/Enterprise/DatasmithCADImporter/Source/DatasmithWireTranslator/Private/AliasModelToTechSoftConverter.cpp
+++ b/Engine/Plugins/Enterprise/DatasmithCADImporter/Source/DatasmithWireTranslator/Private/AliasModelToTechSoftConverter.cpp
@@ -28,6 +28,9 @@
 
 namespace UE_DATASMITHWIRETRANSLATOR_NAMESPACE
 {
+
+const FColor FAliasModelToTechSoftConverter::DefaultColor(200, 200, 200);
+
 #ifdef USE_TECHSOFT_SDK
 
 namespace AliasToTechSoftUtils
@@ -337,7 +340,7 @@ A3DTopoFace* FAliasModelToTechSoftConverter::AddTrimRegion(const AlTrimRegion& I
 
 	A3DTopoFace* FacePtr = CADLibrary::TechSoftInterface::CreateTopoFace(*Face);
 
-	CADLibrary::TechSoftUtils::SetEntityGraphicsColor(FacePtr, FColor(200,200,200));
+	CADLibrary::TechSoftUtils::SetEntityGraphicsColor(FacePtr, DefaultColor);
 
 	return FacePtr;
 }
@@ -406,6 +409,8 @@ bool FAliasModelToTechSoftConverter::AddBRep(AlDagNode& DagNode, EAliasObjectRef
 				if (TSSurface != nullptr)
 				{
 					A3DTopoFace* TSFace = CADLibrary::TechSoftUtils::CreateTopoFaceWithNaturalLoop(TSSurface);
+					CADLibrary::TechSoftUtils::SetEntityGraphicsColor(TSFace, DefaultColor);
+
 					if (TSFace != nullptr)
 					{
 						TSFaces.Add(TSFace);
diff --git a/Engine/Plugins/Enterprise/DatasmithCADImporter/Source/DatasmithWireTranslator/Private/AliasModelToTechSoftConverter.h b/Engine/Plugins/Enterprise/DatasmithCADImporter/Source/DatasmithWireTranslator/Private/AliasModelToTechSoftConverter.h
--- a/Engine/Plugins/Enterprise/DatasmithCADImporter/Source/DatasmithWireTranslator/Private/AliasModelToTechSoftConverter.h
+++ b/Engine/Plugins/Enterprise/DatasmithCADImporter/Source/DatasmithWireTranslator/Private/AliasModelToTechSoftConverter.h
@@ -41,6 +41,7 @@ public:
 	virtual bool AddBRep(AlDagNode& DagNode, EAliasObjectReference ObjectReference) override;
 
 protected:
+	static const FColor DefaultColor;
 
 	TMap<void*, A3DTopoCoEdge*> AlEdgeToTSCoEdge;
 
