Fix TechSoft Color definition

diff --git a/Engine/Plugins/Enterprise/DatasmithCADImporter/Source/CADInterfaces/Private/TechSoftUtils.cpp b/Engine/Plugins/Enterprise/DatasmithCADImporter/Source/CADInterfaces/Private/TechSoftUtils.cpp
--- a/Engine/Plugins/Enterprise/DatasmithCADImporter/Source/CADInterfaces/Private/TechSoftUtils.cpp
+++ b/Engine/Plugins/Enterprise/DatasmithCADImporter/Source/CADInterfaces/Private/TechSoftUtils.cpp
@@ -387,7 +387,7 @@ int32 SetEntityGraphicsColor(A3DEntity* InEntity, FColor Color)
 	StyleData->m_bVPicture = false;
 	StyleData->m_dWidth = 0.1; // default
 	A3DUns8 Alpha = Color.A;
-	if (Alpha > 0)
+	if (Alpha < 255)
 	{
 		StyleData->m_bIsTransparencyDefined = true;
 		StyleData->m_ucTransparency = 255 - Alpha;
@@ -395,7 +395,7 @@ int32 SetEntityGraphicsColor(A3DEntity* InEntity, FColor Color)
 	else
 	{
 		StyleData->m_bIsTransparencyDefined = false;
-		StyleData->m_ucTransparency = 255;
+		StyleData->m_ucTransparency = 0;
 	}
 
 	StyleData->m_bSpecialCulling = false;
