Error in AddFaceTriangleStripeWithUniqueNormal method (and also in AddFaceTriangleStripeWithUniqueNormalAndTexture method)
#swarm https://p4-swarm.epicgames.net/reviews/19644838

diff --git a/Engine/Plugins/Enterprise/DatasmithCADImporter/Source/CADInterfaces/Private/TechSoftInterfaceUtils.inl b/Engine/Plugins/Enterprise/DatasmithCADImporter/Source/CADInterfaces/Private/TechSoftInterfaceUtils.inl
--- a/Engine/Plugins/Enterprise/DatasmithCADImporter/Source/CADInterfaces/Private/TechSoftInterfaceUtils.inl
+++ b/Engine/Plugins/Enterprise/DatasmithCADImporter/Source/CADInterfaces/Private/TechSoftInterfaceUtils.inl
@@ -726,7 +726,6 @@ private:
 
 			Swap(FaceIndex[1], FaceIndex[2]);
 
-			NormalIndex[0] = TriangulatedIndexes[InOutLastTriangleIndex++];
 			FaceIndex[0] = TriangulatedIndexes[InOutLastTriangleIndex++] / 3;
 
 			if (AddFace(FaceIndex, Tessellation, InOutLastVertexIndex))
@@ -777,6 +776,8 @@ private:
 			Swap(FaceIndex[1], FaceIndex[2]);
 			Swap(TextureIndex[1], TextureIndex[2]);
 
+			TextureIndex[0] = TriangulatedIndexes[InOutLastTriangleIndex];
+			InOutLastTriangleIndex += TextureCount;
 			FaceIndex[0] = TriangulatedIndexes[InOutLastTriangleIndex++] / 3;
 
 			if (AddFace(FaceIndex, Tessellation, InOutLastVertexIndex))
