Two problems in this jira:
- FCADFileReader::FindFile did not find files when it should have.
- FDatasmithDispatcher and FDatasmithSceneGraphBuilder sort the files by their FileDescription hash (GetDescriptorHash). This hash is based on the file stat (i.e. the file need to be known). When a file is parsed, its external references are not known... => an hash based on the file name is enough
#swarm https://p4-swarm.epicgames.net/reviews/19666293

diff --git a/Engine/Plugins/Enterprise/DatasmithCADImporter/Source/CADInterfaces/Private/CADFileReader.cpp b/Engine/Plugins/Enterprise/DatasmithCADImporter/Source/CADInterfaces/Private/CADFileReader.cpp
--- a/Engine/Plugins/Enterprise/DatasmithCADImporter/Source/CADInterfaces/Private/CADFileReader.cpp
+++ b/Engine/Plugins/Enterprise/DatasmithCADImporter/Source/CADInterfaces/Private/CADFileReader.cpp
@@ -52,10 +52,10 @@ namespace CADLibrary
 		FString FilePath = FPaths::GetPath(File.GetSourcePath());
 		FString RootFilePath = File.GetRootFolder();
 
-		// Basic case: FilePath is, or is in a sub-folder of, RootFilePath
-		if (FilePath.StartsWith(RootFilePath))
+		// Basic case: File exists at the initial path
+		if (IFileManager::Get().FileExists(*File.GetSourcePath()))
 		{
-			return IFileManager::Get().FileExists(*File.GetSourcePath());
+			return true;
 		}
 
 		// Advance case: end of FilePath is in a upper-folder of RootFilePath
diff --git a/Engine/Plugins/Enterprise/DatasmithCADImporter/Source/DatasmithCADTranslator/Private/DatasmithSceneGraphBuilder.cpp b/Engine/Plugins/Enterprise/DatasmithCADImporter/Source/DatasmithCADTranslator/Private/DatasmithSceneGraphBuilder.cpp
--- a/Engine/Plugins/Enterprise/DatasmithCADImporter/Source/DatasmithCADTranslator/Private/DatasmithSceneGraphBuilder.cpp
+++ b/Engine/Plugins/Enterprise/DatasmithCADImporter/Source/DatasmithCADTranslator/Private/DatasmithSceneGraphBuilder.cpp
@@ -86,8 +86,8 @@ bool FDatasmithSceneGraphBuilder::Build()
 {
 	LoadSceneGraphDescriptionFiles();
 
-	uint32 RootHash = RootFileDescription.GetDescriptorHash();
-	SceneGraph = CADFileToSceneGraphArchive.FindRef(RootHash);
+	uint32 RootHash = 0;
+	SceneGraph = FindSceneGraphArchive(RootFileDescription, RootHash);
 	if (!SceneGraph)
 	{
 		return false;
@@ -125,10 +125,17 @@ void FDatasmithSceneGraphBuilder::LoadSceneGraphDescriptionFiles()
 	}
 }
 
+CADLibrary::FArchiveSceneGraph* FDatasmithSceneBaseGraphBuilder::FindSceneGraphArchive(const CADLibrary::FFileDescriptor& File, uint32& FileHash) const
+{
+	FileHash = GetTypeHash(File.GetFileName());
+	return CADFileToSceneGraphArchive.FindRef(FileHash);
+}
+
 void FDatasmithSceneGraphBuilder::FillAnchorActor(const TSharedRef<IDatasmithActorElement>& ActorElement, const FString& CleanFilenameOfCADFile)
 {
 	CADLibrary::FFileDescriptor AnchorDescription(*CleanFilenameOfCADFile);
-	SceneGraph = CADFileToSceneGraphArchive.FindRef(GetTypeHash(AnchorDescription));
+	uint32 AnchorHash = 0;
+	SceneGraph = FindSceneGraphArchive(AnchorDescription, AnchorHash);
 	if (!SceneGraph)
 	{
 		return;
@@ -293,8 +300,8 @@ TSharedPtr< IDatasmithActorElement >  FDatasmithSceneBaseGraphBuilder::BuildInst
 	{
 		if (!Instance.ExternalReference.GetSourcePath().IsEmpty())
 		{
-			uint32 InstanceSceneGraphHash = GetTypeHash(Instance.ExternalReference);
-			SceneGraph = CADFileToSceneGraphArchive.FindRef(InstanceSceneGraphHash);
+			uint32 InstanceSceneGraphHash = 0;
+			SceneGraph = FindSceneGraphArchive(Instance.ExternalReference, InstanceSceneGraphHash);
 			if (SceneGraph)
 			{
 				if (AncestorSceneGraphHash.Find(InstanceSceneGraphHash) == INDEX_NONE)
diff --git a/Engine/Plugins/Enterprise/DatasmithCADImporter/Source/DatasmithCADTranslator/Public/DatasmithSceneGraphBuilder.h b/Engine/Plugins/Enterprise/DatasmithCADImporter/Source/DatasmithCADTranslator/Public/DatasmithSceneGraphBuilder.h
--- a/Engine/Plugins/Enterprise/DatasmithCADImporter/Source/DatasmithCADTranslator/Public/DatasmithSceneGraphBuilder.h
+++ b/Engine/Plugins/Enterprise/DatasmithCADImporter/Source/DatasmithCADTranslator/Public/DatasmithSceneGraphBuilder.h
@@ -82,6 +82,8 @@ protected:
 
 	void GetNodeUUIDAndName(const TMap<FString, FString>& InInstanceNodeMetaDataMap, const TMap<FString, FString>& InReferenceNodeMetaDataMap, int32 InComponentIndex, const TCHAR* InParentUEUUID, FString& OutUEUUID, FString& OutName);
 
+	CADLibrary::FArchiveSceneGraph* FindSceneGraphArchive(const CADLibrary::FFileDescriptor& File, uint32& FileHash) const;
+
 protected:
 	CADLibrary::FArchiveSceneGraph* SceneGraph;
 	const FString& CachePath;
diff --git a/Engine/Plugins/Enterprise/DatasmithCADImporter/Source/DatasmithDispatcher/Private/DatasmithDispatcher.cpp b/Engine/Plugins/Enterprise/DatasmithCADImporter/Source/DatasmithDispatcher/Private/DatasmithDispatcher.cpp
--- a/Engine/Plugins/Enterprise/DatasmithCADImporter/Source/DatasmithDispatcher/Private/DatasmithDispatcher.cpp
+++ b/Engine/Plugins/Enterprise/DatasmithCADImporter/Source/DatasmithDispatcher/Private/DatasmithDispatcher.cpp
@@ -203,7 +203,7 @@ void FDatasmithDispatcher::LinkCTFileToUnrealCacheFile(const CADLibrary::FFileDe
 {
 	FScopeLock Lock(&TaskPoolCriticalSection);
 
-	uint32 FileHash = GetTypeHash(CTFileDescription);
+	uint32 FileHash = GetTypeHash(CTFileDescription.GetFileName());
 	if (!UnrealSceneGraphFile.IsEmpty())
 	{
 		CADFileToUnrealFileMap.Add(FileHash, UnrealSceneGraphFile);
