Fix scale on Meshing parameters. The scale was apply two times

diff --git a/Engine/Plugins/Enterprise/DatasmithCADImporter/Source/CADInterfaces/Private/TechSoftUtils.cpp b/Engine/Plugins/Enterprise/DatasmithCADImporter/Source/CADInterfaces/Private/TechSoftUtils.cpp
--- a/Engine/Plugins/Enterprise/DatasmithCADImporter/Source/CADInterfaces/Private/TechSoftUtils.cpp
+++ b/Engine/Plugins/Enterprise/DatasmithCADImporter/Source/CADInterfaces/Private/TechSoftUtils.cpp
@@ -167,7 +167,7 @@ bool FillBodyMesh(void* BodyPtr, const FImportParameters& ImportParameters, doub
 
 	TessellationParameters->m_eTessellationLevelOfDetail = kA3DTessLODUserDefined; // Enum to specify predefined values for some following members.
 	TessellationParameters->m_bUseHeightInsteadOfRatio = A3D_TRUE;
-	TessellationParameters->m_dMaxChordHeight = ImportParameters.GetChordTolerance() * 10.; // cm to mm
+	TessellationParameters->m_dMaxChordHeight = ImportParameters.GetChordTolerance();
 	if (!FMath::IsNearlyZero(BodyUnit))
 	{
 		TessellationParameters->m_dMaxChordHeight /= BodyUnit;
