Alias import with Stitching option == Sew allows merged BRep to have more than one material assigned (for multi-color surfaces)
Add CVar "ds.CADTranslator.Alias.SewByColor" to sort BRep by color or not
Fix Sew that called RecurseDagForLeavesNoMerge and duplicate sewn BReps in unsewn ones
Fix Sew with mesh => transform was not apply to Mesh before the sew action

Alias material management (to allow sew of BReps of different materials):
To be compatible with "Retessellate" function, Alias material management as to be the same than CAD (TechSoft) import.
As a reminder: the name and slot of UE Material from CAD is based on CAD material/color data i.e. RGBA Color components => "UE Material slot" (int32) and "UE Material name" (FString = FString::FromInt("UE Material slot"))
UE Material Label is free.

During Retessellate step, Color/Material of each CAD Faces are known, so "UE Material slot" can be deduced.

For Alias import:
Alias BRep is exported in CAD modeler (CADKernel, TechSoft, ...)
Material is build in UE,
From an Alias Material, a unique Color is generated.
This Color is associated to the BRep Shell/face in the CAD modeler
The name and slot of the associated UE Material is defined from this color
So at the retessellate step, nothing changes from the CAD Retessellate process

The unique Color of an Alias Material is defined as follows:
TypeHash(Alias Material Name) => uint24 == 3 uint8 => RGB components of the color

#swarm https://p4-swarm.epicgames.net/reviews/19730005

diff --git a/Engine/Plugins/Enterprise/DatasmithCADImporter/Source/CADInterfaces/Private/TechSoftParserUtils.cpp b/Engine/Plugins/Enterprise/DatasmithCADImporter/Source/CADInterfaces/Private/TechSoftParserUtils.cpp
--- a/Engine/Plugins/Enterprise/DatasmithCADImporter/Source/CADInterfaces/Private/TechSoftParserUtils.cpp
+++ b/Engine/Plugins/Enterprise/DatasmithCADImporter/Source/CADInterfaces/Private/TechSoftParserUtils.cpp
@@ -47,25 +47,6 @@ namespace TechSoftParserUtils
 
 	void FTechSoftMeshConverter::ConvertFace(const A3DTessFaceData& FaceTessData, CADLibrary::FTessellationData& Tessellation)
 	{
-		//if (FaceTessData.m_uiStyleIndexesSize == 1)
-		//{
-		//	A3DUns32 StyleIndex = FaceTessData.m_puiStyleIndexes[0];
-		//	FCADUUID ColorName;
-		//	FCADUUID MaterialName;
-		//	TraverseGraphStyleData(StyleIndex, ColorName, MaterialName);
-		//	if (ColorName)
-		//	{
-		//		Tessellation.ColorName = ColorName;
-		//		BodyMesh.ColorSet.Add(ColorName);
-		//	}
-
-		//	if (MaterialName)
-		//	{
-		//		Tessellation.MaterialName = MaterialName;
-		//		BodyMesh.MaterialSet.Add(MaterialName);
-		//	}
-		//}
-
 		uint32 TriangleCount = CountTriangles(FaceTessData);
 		Reserve(Tessellation, TriangleCount, /*bWithTexture*/ FaceTessData.m_uiTextureCoordIndexesSize > 0);
 
diff --git a/Engine/Plugins/Enterprise/DatasmithCADImporter/Source/CADInterfaces/Public/TechSoftUtils.h b/Engine/Plugins/Enterprise/DatasmithCADImporter/Source/CADInterfaces/Public/TechSoftUtils.h
--- a/Engine/Plugins/Enterprise/DatasmithCADImporter/Source/CADInterfaces/Public/TechSoftUtils.h
+++ b/Engine/Plugins/Enterprise/DatasmithCADImporter/Source/CADInterfaces/Public/TechSoftUtils.h
@@ -17,6 +17,8 @@ typedef void A3DTopoEdge;
 typedef void A3DTopoFace;
 typedef void A3DTopoShell;
 
+class FJsonObject;
+
 namespace CADLibrary
 {
 
@@ -57,6 +59,7 @@ CADINTERFACES_API A3DTopoFace* CreateTopoFaceWithNaturalLoop(A3DSurfBase* Carrie
 CADINTERFACES_API A3DCrvNurbs* CreateTrimNurbsCurve(A3DCrvNurbs* CurveNurbsPtr, double UMin, double UMax, bool bIs2D);
 
 CADINTERFACES_API FColor GetColorAt(uint32 ColorIndex);
+CADINTERFACES_API void RestoreMaterials(const TSharedPtr<FJsonObject>& DefaultValues, CADLibrary::FBodyMesh& BodyMesh);
 
 } // NS TechSoftUtils
 
diff --git a/Engine/Plugins/Enterprise/DatasmithCADImporter/Source/DatasmithWireTranslator/Private/AliasBrepConverter.h b/Engine/Plugins/Enterprise/DatasmithCADImporter/Source/DatasmithWireTranslator/Private/AliasBrepConverter.h
--- a/Engine/Plugins/Enterprise/DatasmithCADImporter/Source/DatasmithWireTranslator/Private/AliasBrepConverter.h
+++ b/Engine/Plugins/Enterprise/DatasmithCADImporter/Source/DatasmithWireTranslator/Private/AliasBrepConverter.h
@@ -18,7 +18,7 @@ enum class EAliasObjectReference
 class IAliasBRepConverter
 {
 public:
-	virtual bool AddBRep(AlDagNode& DagNode, EAliasObjectReference ObjectReference) = 0;
+	virtual bool AddBRep(AlDagNode& DagNode, const FColor& Color, EAliasObjectReference ObjectReference) = 0;
 };
 
 }
\ No newline at end of file
diff --git a/Engine/Plugins/Enterprise/DatasmithCADImporter/Source/DatasmithWireTranslator/Private/AliasModelToCADKernelConverter.cpp b/Engine/Plugins/Enterprise/DatasmithCADImporter/Source/DatasmithWireTranslator/Private/AliasModelToCADKernelConverter.cpp
--- a/Engine/Plugins/Enterprise/DatasmithCADImporter/Source/DatasmithWireTranslator/Private/AliasModelToCADKernelConverter.cpp
+++ b/Engine/Plugins/Enterprise/DatasmithCADImporter/Source/DatasmithWireTranslator/Private/AliasModelToCADKernelConverter.cpp
@@ -2,6 +2,7 @@
 
 #include "AliasModelToCADKernelConverter.h"
 
+#include "CADData.h"
 #include "CADKernelTools.h"
 #include "Hal/PlatformMemory.h"
 #include "MeshDescriptionHelper.h"
@@ -275,12 +276,14 @@ void FAliasModelToCADKernelConverter::AddShell(const AlShell& InShell, EAliasObj
 	}
 }
 
-bool FAliasModelToCADKernelConverter::AddBRep(AlDagNode& DagNode, EAliasObjectReference InObjectReference)
+bool FAliasModelToCADKernelConverter::AddBRep(AlDagNode& DagNode, const FColor& Color, EAliasObjectReference InObjectReference)
 {
 	AlEdge2CADKernelEdge.Empty();
 
 	TSharedRef<CADKernel::FBody> CADKernelBody = CADKernel::FEntity::MakeShared<CADKernel::FBody>();
 	TSharedRef<CADKernel::FShell> CADKernelShell = CADKernel::FEntity::MakeShared<CADKernel::FShell>();
+	uint32 ColorId = (uint32) CADLibrary::BuildColorName(Color);
+	CADKernelShell->SetColorId(ColorId);
 	CADKernelBody->AddShell(CADKernelShell);
 
 	boolean bAlOrientation;
diff --git a/Engine/Plugins/Enterprise/DatasmithCADImporter/Source/DatasmithWireTranslator/Private/AliasModelToCADKernelConverter.h b/Engine/Plugins/Enterprise/DatasmithCADImporter/Source/DatasmithWireTranslator/Private/AliasModelToCADKernelConverter.h
--- a/Engine/Plugins/Enterprise/DatasmithCADImporter/Source/DatasmithWireTranslator/Private/AliasModelToCADKernelConverter.h
+++ b/Engine/Plugins/Enterprise/DatasmithCADImporter/Source/DatasmithWireTranslator/Private/AliasModelToCADKernelConverter.h
@@ -40,7 +40,7 @@ public:
 	{
 	}
 
-	virtual bool AddBRep(AlDagNode& DagNode, EAliasObjectReference ObjectReference) override;
+	virtual bool AddBRep(AlDagNode& DagNode, const FColor& Color, EAliasObjectReference ObjectReference) override;
 
 protected:
 	TSharedPtr<CADKernel::FTopologicalEdge> AddEdge(const AlTrimCurve& TrimCurve, TSharedPtr<CADKernel::FSurface>& CarrierSurface);
diff --git a/Engine/Plugins/Enterprise/DatasmithCADImporter/Source/DatasmithWireTranslator/Private/AliasModelToCoretechConverter.cpp b/Engine/Plugins/Enterprise/DatasmithCADImporter/Source/DatasmithWireTranslator/Private/AliasModelToCoretechConverter.cpp
--- a/Engine/Plugins/Enterprise/DatasmithCADImporter/Source/DatasmithWireTranslator/Private/AliasModelToCoretechConverter.cpp
+++ b/Engine/Plugins/Enterprise/DatasmithCADImporter/Source/DatasmithWireTranslator/Private/AliasModelToCoretechConverter.cpp
@@ -194,7 +194,7 @@ void FAliasModelToCoretechConverter::AddShell(const AlShell& Shell, EAliasObject
 	}
 }
 
-bool FAliasModelToCoretechConverter::AddBRep(AlDagNode& DagNode, EAliasObjectReference InObjectReference)
+bool FAliasModelToCoretechConverter::AddBRep(AlDagNode& DagNode, const FColor& Color, EAliasObjectReference InObjectReference)
 {
 	TArray<uint64> FaceList;
 
diff --git a/Engine/Plugins/Enterprise/DatasmithCADImporter/Source/DatasmithWireTranslator/Private/AliasModelToCoretechConverter.h b/Engine/Plugins/Enterprise/DatasmithCADImporter/Source/DatasmithWireTranslator/Private/AliasModelToCoretechConverter.h
--- a/Engine/Plugins/Enterprise/DatasmithCADImporter/Source/DatasmithWireTranslator/Private/AliasModelToCoretechConverter.h
+++ b/Engine/Plugins/Enterprise/DatasmithCADImporter/Source/DatasmithWireTranslator/Private/AliasModelToCoretechConverter.h
@@ -30,7 +30,7 @@ public:
 	{
 	}
 
-	virtual bool AddBRep(AlDagNode& DagNode, EAliasObjectReference ObjectReference) override;
+	virtual bool AddBRep(AlDagNode& DagNode, const FColor& Color, EAliasObjectReference ObjectReference) override;
 
 protected:
 
diff --git a/Engine/Plugins/Enterprise/DatasmithCADImporter/Source/DatasmithWireTranslator/Private/AliasModelToTechSoftConverter.cpp b/Engine/Plugins/Enterprise/DatasmithCADImporter/Source/DatasmithWireTranslator/Private/AliasModelToTechSoftConverter.cpp
--- a/Engine/Plugins/Enterprise/DatasmithCADImporter/Source/DatasmithWireTranslator/Private/AliasModelToTechSoftConverter.cpp
+++ b/Engine/Plugins/Enterprise/DatasmithCADImporter/Source/DatasmithWireTranslator/Private/AliasModelToTechSoftConverter.cpp
@@ -29,8 +29,6 @@
 namespace UE_DATASMITHWIRETRANSLATOR_NAMESPACE
 {
 
-const FColor FAliasModelToTechSoftConverter::DefaultColor(200, 200, 200);
-
 #ifdef USE_TECHSOFT_SDK
 
 namespace AliasToTechSoftUtils
@@ -304,7 +302,7 @@ void FAliasModelToTechSoftConverter::LinkEdgesLoop(const AlTrimBoundary& TrimBou
 	}
 }
 
-A3DTopoFace* FAliasModelToTechSoftConverter::AddTrimRegion(const AlTrimRegion& InTrimRegion, EAliasObjectReference InObjectReference, const AlMatrix4x4& InAlMatrix)
+A3DTopoFace* FAliasModelToTechSoftConverter::AddTrimRegion(const AlTrimRegion& InTrimRegion, const FColor& Color, EAliasObjectReference InObjectReference, const AlMatrix4x4& InAlMatrix)
 {
 	A3DSurfBase* CarrierSurface = AliasToTechSoftUtils::AddNURBSSurface(InTrimRegion, InObjectReference, InAlMatrix);
 	if (CarrierSurface == nullptr)
@@ -340,13 +338,13 @@ A3DTopoFace* FAliasModelToTechSoftConverter::AddTrimRegion(const AlTrimRegion& I
 
 	A3DTopoFace* FacePtr = CADLibrary::TechSoftInterface::CreateTopoFace(*Face);
 
-	CADLibrary::TechSoftUtils::SetEntityGraphicsColor(FacePtr, DefaultColor);
+	CADLibrary::TechSoftUtils::SetEntityGraphicsColor(FacePtr, Color);
 
 	return FacePtr;
 }
 #endif
 
-bool FAliasModelToTechSoftConverter::AddBRep(AlDagNode& DagNode, EAliasObjectReference InObjectReference)
+bool FAliasModelToTechSoftConverter::AddBRep(AlDagNode& DagNode, const FColor& Color, EAliasObjectReference InObjectReference)
 {
 #ifdef USE_TECHSOFT_SDK
 	AlEdgeToTSCoEdge.Empty();
@@ -373,7 +371,7 @@ bool FAliasModelToTechSoftConverter::AddBRep(AlDagNode& DagNode, EAliasObjectRef
 			{
 				for (TUniquePtr<AlTrimRegion> TrimRegion(AliasShell->firstTrimRegion()); TrimRegion.IsValid(); TrimRegion = TUniquePtr<AlTrimRegion>(TrimRegion->nextRegion()))
 				{
-					A3DTopoFace* TSFace = AddTrimRegion(*TrimRegion, InObjectReference, AlMatrix);
+					A3DTopoFace* TSFace = AddTrimRegion(*TrimRegion, Color, InObjectReference, AlMatrix);
 					if (TSFace != nullptr)
 					{
 						TSFaces.Add(TSFace);
@@ -396,7 +394,7 @@ bool FAliasModelToTechSoftConverter::AddBRep(AlDagNode& DagNode, EAliasObjectRef
 				{
 					for (; TrimRegion.IsValid(); TrimRegion = TUniquePtr<AlTrimRegion>(TrimRegion->nextRegion()))
 					{
-						A3DTopoFace* TSFace = AddTrimRegion(*TrimRegion, InObjectReference, AlMatrix);
+						A3DTopoFace* TSFace = AddTrimRegion(*TrimRegion, Color, InObjectReference, AlMatrix);
 						if (TSFace != nullptr)
 						{
 							TSFaces.Add(TSFace);
@@ -409,7 +407,7 @@ bool FAliasModelToTechSoftConverter::AddBRep(AlDagNode& DagNode, EAliasObjectRef
 				if (TSSurface != nullptr)
 				{
 					A3DTopoFace* TSFace = CADLibrary::TechSoftUtils::CreateTopoFaceWithNaturalLoop(TSSurface);
-					CADLibrary::TechSoftUtils::SetEntityGraphicsColor(TSFace, DefaultColor);
+					CADLibrary::TechSoftUtils::SetEntityGraphicsColor(TSFace, Color);
 
 					if (TSFace != nullptr)
 					{
diff --git a/Engine/Plugins/Enterprise/DatasmithCADImporter/Source/DatasmithWireTranslator/Private/AliasModelToTechSoftConverter.h b/Engine/Plugins/Enterprise/DatasmithCADImporter/Source/DatasmithWireTranslator/Private/AliasModelToTechSoftConverter.h
--- a/Engine/Plugins/Enterprise/DatasmithCADImporter/Source/DatasmithWireTranslator/Private/AliasModelToTechSoftConverter.h
+++ b/Engine/Plugins/Enterprise/DatasmithCADImporter/Source/DatasmithWireTranslator/Private/AliasModelToTechSoftConverter.h
@@ -38,14 +38,12 @@ public:
 	{
 	}
 
-	virtual bool AddBRep(AlDagNode& DagNode, EAliasObjectReference ObjectReference) override;
+	virtual bool AddBRep(AlDagNode& DagNode, const FColor& Color, EAliasObjectReference ObjectReference) override;
 
 protected:
-	static const FColor DefaultColor;
-
 	TMap<void*, A3DTopoCoEdge*> AlEdgeToTSCoEdge;
 
-	A3DTopoFace* AddTrimRegion(const AlTrimRegion& InTrimRegion, EAliasObjectReference InObjectReference, const AlMatrix4x4& InAlMatrix);
+	A3DTopoFace* AddTrimRegion(const AlTrimRegion& InTrimRegion, const FColor& Color, EAliasObjectReference InObjectReference, const AlMatrix4x4& InAlMatrix);
 	A3DTopoLoop* CreateTopoLoop(const AlTrimBoundary& TrimBoundary);
 	A3DTopoCoEdge* CreateEdge(const AlTrimCurve& TrimCurve);
 	A3DCrvBase* CreateCurve(const AlTrimCurve& TrimCurve);
diff --git a/Engine/Plugins/Enterprise/DatasmithCADImporter/Source/DatasmithWireTranslator/Private/DatasmithWireTranslator.cpp b/Engine/Plugins/Enterprise/DatasmithCADImporter/Source/DatasmithWireTranslator/Private/DatasmithWireTranslator.cpp
--- a/Engine/Plugins/Enterprise/DatasmithCADImporter/Source/DatasmithWireTranslator/Private/DatasmithWireTranslator.cpp
+++ b/Engine/Plugins/Enterprise/DatasmithCADImporter/Source/DatasmithWireTranslator/Private/DatasmithWireTranslator.cpp
@@ -71,6 +71,17 @@ DEFINE_LOG_CATEGORY_STATIC(LogDatasmithWireTranslator, Log, All);
 namespace UE_DATASMITHWIRETRANSLATOR_NAMESPACE
 {
 
+static bool bGAliasSewByColor = false;
+FAutoConsoleVariableRef GAliasSewByColor(
+	TEXT("ds.CADTranslator.Alias.SewByColor"),
+	bGAliasSewByColor,
+	TEXT("Enable Sew action merges BReps according to their material i.e. only BReps associated with same material can be merged together.\
+Default is disable\n"),
+	ECVF_Default);
+
+
+static const FColor DefaultColor = FColor(200, 200, 200);
+
 #ifdef USE_OPENMODEL
 
 const uint64 LibAliasNext_Version = 17881307937833405;
@@ -106,39 +117,110 @@ const uint64 LibAliasVersionMax = LibAliasNext_Version;
 const FString AliasVersionChar = TEXT("AliasStudio 2022.2, Model files");
 #endif
 
+// Alias material management (to allow sew of BReps of different materials):
+// To be compatible with "Retessellate" function, Alias material management as to be the same than CAD (TechSoft) import. 
+// As a reminder: the name and slot of UE Material from CAD is based on CAD material/color data i.e. RGBA Color components => "UE Material slot" (int32) and "UE Material name" (FString = FString::FromInt("UE Material slot"))
+// UE Material Label is free.
+// 
+// During Retessellate step, Color/Material of each CAD Faces are known, so "UE Material slot" can be deduced.
+// 
+// For Alias import:
+// Alias BRep is exported in CAD modeler (CADKernel, TechSoft, ...)
+// Material is build in UE, 
+// From an Alias Material, a unique Color is generated.
+// This Color is associated to the BRep Shell/face in the CAD modeler
+// The name and slot of the associated UE Material is defined from this color
+// So at the retessellate step, nothing changes from the CAD Retessellate process
+// 
+// The unique Color of an Alias Material is defined as follows:
+// TypeHash(Alias Material Name) => uint24 == 3 uint8 => RGB components of the color  
+
+FColor CreateShaderColorFromShaderName(const FString& ShaderName)
+{
+	const uint32 ShaderHash = GetTypeHash(*ShaderName);
+	const uint32 Red = (ShaderHash & 0xff000000) >> 24;
+	const uint32 Green = (ShaderHash & 0x00ff0000) >> 16;
+	const uint32 Blue = (ShaderHash & 0x0000ff00) >> 8;
+	return FColor(Red, Green, Blue);
+}
+
+int32 CreateShaderId(const FColor& ShaderColor)
+{
+	return FMath::Abs((int32)GetTypeHash(ShaderColor));
+}
+
 class BodyData
 {
+private:
+	TArray<TPair<TSharedPtr<AlDagNode>, FColor>> Shells;
+	TSet<FString> ShaderNames;
+
+	FString Label;
+	FString LayerName;
+	bool bCadData = true;
+
 public:
-	BodyData(const FString& InShaderName, const FString& InLayerName, bool bInCadData)
+
+	BodyData(const FString& InLayerName, bool bInCadData)
 		: bCadData(bInCadData)
 	{
-		ShaderName = InShaderName;
 		LayerName = InLayerName;
 	};
 
-	TArray<TSharedPtr<AlDagNode>> ShellSet;
-	FString ShaderName;
-	FString LayerName;
-	FString Label;
-	bool bCadData = true;
+	bool IsCadData() const 
+	{
+		return bCadData;
+	}
+
+	void SetLabel(const FString& InLabel)
+	{
+		Label = InLabel;
+	}
+
+	void Add(const TSharedPtr<AlDagNode>& Node, FColor Color, const FString& ShaderName)
+	{
+		ShaderNames.Add(ShaderName);
+		Shells.Emplace(Node, Color);
+	}
+
+	void Reserve(int32 MaxElement)
+	{
+		Shells.Reserve(MaxElement);
+		ShaderNames.Reserve(MaxElement);
+	}
+
+	const TSet<FString>& GetShaderNames() const
+	{
+		return ShaderNames;
+	}
+
+	const FString& GetLayerName() const
+	{
+		return LayerName;
+	}
+
+	const TArray<TPair<TSharedPtr<AlDagNode>, FColor>>& GetShells() const
+	{
+		return Shells;
+	}
 
 	// Generates BodyData's unique id from AlDagNode objects
 	uint32 GetUuid(const uint32& ParentUuid)
 	{
-		if (ShellSet.Num() == 0)
+		if (Shells.Num() == 0)
 		{
 			return ParentUuid;
 		}
 
-		if (ShellSet.Num() > 1)
+		if (Shells.Num() > 1)
 		{
-			ShellSet.Sort([&](const TSharedPtr<AlDagNode>& NodeA, const TSharedPtr<AlDagNode>& NodeB) { return OpenModelUtils::GetAlDagNodeUuid(NodeA) < OpenModelUtils::GetAlDagNodeUuid(NodeB); });
+			Shells.Sort([&](const TPair<TSharedPtr<AlDagNode>, FColor>& NodeA, const TPair<TSharedPtr<AlDagNode>, FColor>& NodeB) { return OpenModelUtils::GetAlDagNodeUuid(*NodeA.Key) < OpenModelUtils::GetAlDagNodeUuid(*NodeB.Key); });
 		}
 
 		uint32 BodyUUID = 0;
-		for (const TSharedPtr<AlDagNode>& DagNode : ShellSet)
+		for (const TPair<TSharedPtr<AlDagNode>, FColor>& DagNode : Shells)
 		{
-			BodyUUID = HashCombine(BodyUUID, OpenModelUtils::GetAlDagNodeUuid(DagNode));
+			BodyUUID = HashCombine(BodyUUID, OpenModelUtils::GetAlDagNodeUuid(*DagNode.Key));
 		}
 
 		return HashCombine(ParentUuid, BodyUUID);
@@ -279,18 +361,19 @@ private:
 	bool GetDagLeaves();
 	bool GetShader();
 	bool RecurseDagForLeaves(const TSharedPtr<AlDagNode>& DagNode, const FDagNodeInfo& ParentInfo);
-	bool RecurseDagForLeavesNoMerge(const TSharedPtr<AlDagNode>& DagNode, const FDagNodeInfo& ParentInfo);
+	void RecurseDagForLeavesNoMerge(const TSharedPtr<AlDagNode>& DagNode, const FDagNodeInfo& ParentInfo);
+	void DagForLeavesNoMerge(const TSharedPtr<AlDagNode>& DagNode, const FDagNodeInfo& ParentInfo);
 	bool ProcessAlGroupNode(AlDagNode& GroupNode, const FDagNodeInfo& ParentInfo);
 	bool ProcessAlShellNode(const TSharedPtr<AlDagNode>& ShellNode, const FDagNodeInfo& ParentInfo, const FString& ShaderName);
-	bool ProcessBodyNode(const TSharedPtr<BodyData>& Body, const FDagNodeInfo& ParentInfo);
+	void ProcessBodyNode(const TSharedPtr<BodyData>& Body, const FDagNodeInfo& ParentInfo);
 	TSharedPtr<IDatasmithMeshElement> FindOrAddMeshElement(const TSharedPtr<BodyData>& Body, const FDagNodeInfo& ParentInfo);
 	TSharedPtr<IDatasmithMeshElement> FindOrAddMeshElement(const TSharedPtr<AlDagNode>& ShellNode, const FDagNodeInfo& ParentInfo, const FString& ShaderName);
-	void GetDagNodeInfo(const TSharedPtr<AlDagNode>& GroupNode, const FDagNodeInfo& ParentInfo, FDagNodeInfo& CurrentNodeInfo);
-	void GetDagNodeInfo(const TSharedPtr<BodyData>& CurrentNode, const FDagNodeInfo& ParentInfo, FDagNodeInfo& CurrentNodeInfo);
+	void GetDagNodeInfo(AlDagNode& GroupNode, const FDagNodeInfo& ParentInfo, FDagNodeInfo& CurrentNodeInfo);
+	void GetDagNodeInfo(BodyData& CurrentNode, const FDagNodeInfo& ParentInfo, FDagNodeInfo& CurrentNodeInfo);
 	void GetDagNodeMeta(const AlDagNode& InDagNode, TSharedPtr<IDatasmithActorElement> OutActorElement);
 
 	TOptional<FMeshDescription> GetMeshOfShellNode(AlDagNode& DagNode, TSharedRef<IDatasmithMeshElement> MeshElement, CADLibrary::FMeshParameters& MeshParameters);
-	TOptional<FMeshDescription> GetMeshOfNodeMesh(AlDagNode& DagNode, CADLibrary::FMeshParameters& MeshParameters, AlMatrix4x4* AlMeshInvGlobalMatrix = nullptr);
+	TOptional<FMeshDescription> GetMeshOfNodeMesh(AlDagNode& DagNode, TSharedRef<IDatasmithMeshElement> MeshElement, CADLibrary::FMeshParameters& MeshParameters, AlMatrix4x4* AlMeshInvGlobalMatrix = nullptr);
 	TOptional<FMeshDescription> GetMeshOfShellBody(TSharedRef<BodyData> DagNode, TSharedRef<IDatasmithMeshElement> MeshElement, CADLibrary::FMeshParameters& MeshParameters);
 	TOptional<FMeshDescription> GetMeshOfMeshBody(TSharedRef<BodyData> DagNode, TSharedRef<IDatasmithMeshElement> MeshElement, CADLibrary::FMeshParameters& MeshParameters);
 
@@ -299,7 +382,7 @@ private:
 	TOptional<FMeshDescription> MeshDagNodeWithExternalMesher(AlDagNode& DagNode, TSharedRef<IDatasmithMeshElement> MeshElement, CADLibrary::FMeshParameters& MeshParameters);
 	TOptional<FMeshDescription> MeshDagNodeWithExternalMesher(TSharedRef<BodyData> DagNode, TSharedRef<IDatasmithMeshElement> MeshElement, CADLibrary::FMeshParameters& MeshParameters);
 
-	TOptional<FMeshDescription> ImportMesh(AlMesh& Mesh, CADLibrary::FMeshParameters& MeshParameters);
+	TOptional<FMeshDescription> ImportMesh(AlMesh& Mesh, TSharedRef<IDatasmithMeshElement> MeshElement, CADLibrary::FMeshParameters& MeshParameters);
 
 	FORCEINLINE bool IsTransparent(FColor& TransparencyColor)
 	{
@@ -371,7 +454,9 @@ private:
 	TMap<IDatasmithMeshElement*, TSharedPtr<AlDagNode>> MeshElementToAlDagNodeMap;
 
 	TMap<IDatasmithMeshElement*, TSharedPtr<BodyData>> MeshElementToBodyMap;
-
+	
+	TMap<IDatasmithMeshElement*, FString> MeshElementToShaderName;
+	TMap<FString, FColor> ShaderNameToColor;
 	TMap<FString, TSharedPtr<IDatasmithMaterialIDElement>> ShaderNameToUEMaterialId;
 
 	// start section information
@@ -420,12 +505,7 @@ bool FWireTranslatorImpl::Read()
 
 	// Parse and extract the DAG leaf nodes.
 	// Note that Alias file unit is cm like UE
-	if (!GetDagLeaves())
-	{
-		return false;
-	}
-
-	return true;
+	return GetDagLeaves();
 }
 
 void FWireTranslatorImpl::AddAlBlinnParameters(const TUniquePtr<AlShader>& Shader, TSharedRef<IDatasmithUEPbrMaterialElement> MaterialElement)
@@ -1255,15 +1335,17 @@ bool FWireTranslatorImpl::GetShader()
 {
 	for (TUniquePtr<AlShader> Shader(AlUniverse::firstShader()); Shader.IsValid(); Shader = TUniquePtr<AlShader>(AlUniverse::nextShader(Shader.Get())))
 	{
-		FString ShaderName = UTF8_TO_TCHAR(Shader->name());
-		FString ShaderModelName = Shader->shadingModel();
+		const FString ShaderName = UTF8_TO_TCHAR(Shader->name());
+		const FString ShaderModelName = Shader->shadingModel();
 
-		uint32 ShaderUUID = fabs((float)(int32)GetTypeHash(*ShaderName));
+		const FColor Color = CreateShaderColorFromShaderName(ShaderName);
+		ShaderNameToColor.Add(ShaderName, Color);
 
-		TSharedRef<IDatasmithUEPbrMaterialElement> MaterialElement = FDatasmithSceneFactory::CreateUEPbrMaterial(*ShaderName);
+		const int32 ShaderId = CreateShaderId(Color);
+		const FString MaterialElementName = FString::FromInt(ShaderId);
 
+		TSharedRef<IDatasmithUEPbrMaterialElement> MaterialElement = FDatasmithSceneFactory::CreateUEPbrMaterial(*MaterialElementName);
 		MaterialElement->SetLabel(*ShaderName);
-		MaterialElement->SetName(*FString::FromInt(ShaderUUID));
 
 		if (ShaderModelName.Equals(TEXT("BLINN")))
 		{
@@ -1296,11 +1378,11 @@ bool FWireTranslatorImpl::GetDagLeaves()
 {
 	FDagNodeInfo RootContainer;
 	AlRootNode = TSharedPtr<AlDagNode>(AlUniverse::firstDagNode());
-	if (!AlRootNode.IsValid())
+	if (AlRootNode.IsValid())
 	{
-		return false;
+		return RecurseDagForLeaves(AlRootNode, RootContainer);
 	}
-	return RecurseDagForLeaves(AlRootNode, RootContainer);
+	return false;
 }
 
 
@@ -1315,19 +1397,18 @@ void FWireTranslatorImpl::GetDagNodeMeta(const AlDagNode& InDagNode, TSharedPtr<
 	}
 }
 
-void FWireTranslatorImpl::GetDagNodeInfo(const TSharedPtr<AlDagNode>& CurrentNode, const FDagNodeInfo& ParentInfo, FDagNodeInfo& CurrentNodeInfo)
+void FWireTranslatorImpl::GetDagNodeInfo(AlDagNode& CurrentNode, const FDagNodeInfo& ParentInfo, FDagNodeInfo& CurrentNodeInfo)
 {
-
-	CurrentNodeInfo.Label = UTF8_TO_TCHAR(CurrentNode->name());
+	CurrentNodeInfo.Label = UTF8_TO_TCHAR(CurrentNode.name());
 	uint32 ThisGroupNodeUuid = OpenModelUtils::GetAlDagNodeUuid(CurrentNode);
 	CurrentNodeInfo.Uuid = HashCombine(ParentInfo.Uuid, ThisGroupNodeUuid);
 }
 
-void FWireTranslatorImpl::GetDagNodeInfo(const TSharedPtr<BodyData>& CurrentNode, const FDagNodeInfo& ParentInfo, FDagNodeInfo& CurrentNodeInfo)
+void FWireTranslatorImpl::GetDagNodeInfo(BodyData& CurrentNode, const FDagNodeInfo& ParentInfo, FDagNodeInfo& CurrentNodeInfo)
 {
+	CurrentNode.SetLabel(ParentInfo.Label);
 	CurrentNodeInfo.Label = ParentInfo.Label;
-	CurrentNode->Label = CurrentNodeInfo.Label;
-	CurrentNodeInfo.Uuid = CurrentNode->GetUuid(ParentInfo.Uuid);
+	CurrentNodeInfo.Uuid = CurrentNode.GetUuid(ParentInfo.Uuid);
 }
 
 bool FWireTranslatorImpl::ProcessAlGroupNode(AlDagNode& GroupNode, const FDagNodeInfo& ParentInfo)
@@ -1347,7 +1428,7 @@ bool FWireTranslatorImpl::ProcessAlGroupNode(AlDagNode& GroupNode, const FDagNod
 	TSharedPtr<AlDagNode> ChildNode(AlChildPtr);
 
 	FDagNodeInfo ThisGroupNodeInfo;
-	GetDagNodeInfo(ChildNode, ParentInfo, ThisGroupNodeInfo);
+	GetDagNodeInfo(*ChildNode, ParentInfo, ThisGroupNodeInfo);
 
 	ThisGroupNodeInfo.ActorElement = FDatasmithSceneFactory::CreateActor(*OpenModelUtils::UuidToString(ThisGroupNodeInfo.Uuid));
 	ThisGroupNodeInfo.ActorElement->SetLabel(*ThisGroupNodeInfo.Label);
@@ -1386,10 +1467,19 @@ TSharedPtr<IDatasmithMeshElement> FWireTranslatorImpl::FindOrAddMeshElement(cons
 	MeshElement->SetLabel(*NodeInfo.Label);
 	MeshElement->SetLightmapSourceUV(-1);
 
-	if (*Body->ShaderName)
+	const TSet<FString>& ShaderNames = Body->GetShaderNames();
+	for(const FString& ShaderName : ShaderNames)
 	{
-		TSharedPtr<IDatasmithMaterialIDElement> MaterialElement = ShaderNameToUEMaterialId[Body->ShaderName];
-		MeshElement->SetMaterial(MaterialElement->GetName(), 0);
+		TSharedPtr<IDatasmithMaterialIDElement> MaterialElement = ShaderNameToUEMaterialId[ShaderName];
+		if (MaterialElement.IsValid())
+		{
+			const FColor* ShaderColor = ShaderNameToColor.Find(ShaderName);
+			if(ShaderColor)
+			{
+				int32 ShaderId = CreateShaderId(*ShaderColor);
+				MeshElement->SetMaterial(MaterialElement->GetName(), ShaderId);
+			}
+		}
 	}
 
 	DatasmithScene->AddMesh(MeshElement);
@@ -1404,7 +1494,7 @@ TSharedPtr<IDatasmithMeshElement> FWireTranslatorImpl::FindOrAddMeshElement(cons
 
 TSharedPtr<IDatasmithMeshElement> FWireTranslatorImpl::FindOrAddMeshElement(const TSharedPtr<AlDagNode>& ShellNode, const FDagNodeInfo& ShellNodeInfo, const FString& ShaderName)
 {
-	uint32 ShellUuid = OpenModelUtils::GetAlDagNodeUuid(ShellNode);
+	uint32 ShellUuid = OpenModelUtils::GetAlDagNodeUuid(*ShellNode);
 
 	// Look if geometry has not been already processed, return it if found
 	TSharedPtr<IDatasmithMeshElement>* MeshElementPtr = ShellUuidToMeshElementMap.Find(ShellUuid);
@@ -1429,13 +1519,23 @@ TSharedPtr<IDatasmithMeshElement> FWireTranslatorImpl::FindOrAddMeshElement(cons
 	if (ShaderName.Len())
 	{
 		TSharedPtr<IDatasmithMaterialIDElement> MaterialElement = ShaderNameToUEMaterialId[ShaderName];
-		MeshElement->SetMaterial(MaterialElement->GetName(), 0);
+		if (MaterialElement.IsValid())
+		{
+			const FColor* ShaderColor = ShaderNameToColor.Find(ShaderName);
+			if (ShaderColor)
+			{
+				int32 ShaderId = CreateShaderId(*ShaderColor);
+				MeshElement->SetMaterial(MaterialElement->GetName(), ShaderId);
+			}
+		}
+
 	}
 
 	DatasmithScene->AddMesh(MeshElement);
 
 	ShellUuidToMeshElementMap.Add(ShellUuid, MeshElement);
 	MeshElementToAlDagNodeMap.Add(MeshElement.Get(), ShellNode);
+	MeshElementToShaderName.Add(MeshElement.Get(), ShaderName);
 
 	return MeshElement;
 }
@@ -1443,7 +1543,7 @@ TSharedPtr<IDatasmithMeshElement> FWireTranslatorImpl::FindOrAddMeshElement(cons
 bool FWireTranslatorImpl::ProcessAlShellNode(const TSharedPtr<AlDagNode>& ShellNode, const FDagNodeInfo& ParentInfo, const FString& ShaderName)
 {
 	FDagNodeInfo ShellInfo;
-	GetDagNodeInfo(ShellNode, ParentInfo, ShellInfo);
+	GetDagNodeInfo(*ShellNode, ParentInfo, ShellInfo);
 
 	TSharedPtr<IDatasmithMeshElement> MeshElement = FindOrAddMeshElement(ShellNode, ShellInfo, ShaderName);
 	if (!MeshElement.IsValid())
@@ -1493,51 +1593,51 @@ bool FWireTranslatorImpl::ProcessAlShellNode(const TSharedPtr<AlDagNode>& ShellN
 	return true;
 }
 
-bool FWireTranslatorImpl::ProcessBodyNode(const TSharedPtr<BodyData>& Body, const FDagNodeInfo& ParentInfo)
+void FWireTranslatorImpl::ProcessBodyNode(const TSharedPtr<BodyData>& Body, const FDagNodeInfo& ParentInfo)
 {
 	if (!Body.IsValid())
 	{
-		return false;
+		return;
 	}
 
-	if (Body->ShellSet.Num() == 1)
+	if (Body->GetShells().Num() == 1)
 	{
-		return RecurseDagForLeavesNoMerge(Body->ShellSet[0], ParentInfo);
+		DagForLeavesNoMerge(Body->GetShells()[0].Key, ParentInfo);
+		return;
 	}
 
 	FDagNodeInfo ShellInfo;
-	GetDagNodeInfo(Body, ParentInfo, ShellInfo);
+	GetDagNodeInfo(*Body, ParentInfo, ShellInfo);
 
 	TSharedPtr<IDatasmithMeshElement> MeshElement = FindOrAddMeshElement(Body, ShellInfo);
 	if (!MeshElement.IsValid())
 	{
-		return false;
+		return;
 	}
 
 	TSharedPtr<IDatasmithMeshActorElement> ActorElement = FDatasmithSceneFactory::CreateMeshActor(*OpenModelUtils::UuidToString(ShellInfo.Uuid));
 	if (!ActorElement.IsValid())
 	{
-		return false;
+		return;
 	}
 
 	ActorElement->SetLabel(*ShellInfo.Label);
 	ActorElement->SetStaticMeshPathName(MeshElement->GetName());
 	ShellInfo.ActorElement = ActorElement;
 
-	ActorElement->SetLayer(*Body->LayerName);
+	ActorElement->SetLayer(*Body->GetLayerName());
 
-	//SetActorTransform(ShellInfo.ActorElement, ShellNode);
-
-	//// Apply materials on the current part
-	if (*Body->ShaderName)
+	// Apply materials on the current part
+	int32 Index = 0;
+	const TSet<FString>& ShaderNames = Body->GetShaderNames();
+	for (const FString& ShaderName : ShaderNames)
 	{
-		TSharedPtr<IDatasmithMaterialIDElement> MaterialIDElement = ShaderNameToUEMaterialId[FString(*Body->ShaderName)];
-		if (MaterialIDElement.IsValid()) {
-			for (int32 Index = 0; Index < MeshElement->GetMaterialSlotCount(); ++Index)
-			{
-				MaterialIDElement->SetId(MeshElement->GetMaterialSlotAt(Index)->GetId());
-				ActorElement->AddMaterialOverride(MaterialIDElement);
-			}
+		TSharedPtr<IDatasmithMaterialIDElement> MaterialIDElement = ShaderNameToUEMaterialId[ShaderName];
+		if (MaterialIDElement.IsValid())
+		{
+			MaterialIDElement->SetId(MeshElement->GetMaterialSlotAt(Index)->GetId());
+			ActorElement->AddMaterialOverride(MaterialIDElement);
+			Index++;
 		}
 	}
 
@@ -1552,7 +1652,7 @@ bool FWireTranslatorImpl::ProcessBodyNode(const TSharedPtr<BodyData>& Body, cons
 			DatasmithScene->AddActor(ActorElement);
 		}
 	}
-	return true;
+	return;
 }
 
 TSharedPtr<AlDagNode> GetNextNode(const TSharedPtr<AlDagNode>& DagNode)
@@ -1569,41 +1669,13 @@ TSharedPtr<AlDagNode> GetNextNode(const TSharedPtr<AlDagNode>& DagNode)
 	}
 }
 
-bool IsHidden(const TSharedPtr<AlDagNode>& DagNode)
+uint32 GetBodyGroupUuid(const FString& ShaderName, const FString& LayerName, bool bCadData)
 {
-	/*
-	AlObjectType objectType = DagNode->type();
-	switch (objectType)
+	uint32 Uuid = HashCombine(GetTypeHash(LayerName), GetTypeHash(bCadData)); 
+	if(bGAliasSewByColor)
 	{
-		case kShellNodeType:
-		case kSurfaceNodeType:
-		case kMeshNodeType:
-		{
-			boolean isVisible = !DagNode->isDisplayModeSet(AlDisplayModeType::kDisplayModeInvisible);
-			if (!isVisible)
-			{
-				return true;
-			}
-
-			AlLayer* Layer = DagNode->layer();
-			if (Layer)
-			{
-				if (Layer->invisible())
-				{
-					return true;
-				}
-			}
-		}
+		Uuid = HashCombine(GetTypeHash(ShaderName), Uuid);
 	}
-	*/
-	return false;
-}
-
-
-uint32 GetBodyGroupUuid(const FString& ShaderName, const FString& LayerName, bool bCadData)
-{
-	uint32 Uuid = HashCombine(GetTypeHash(ShaderName), GetTypeHash(bCadData));
-	Uuid = HashCombine(GetTypeHash(LayerName), Uuid);
 	return Uuid;
 }
 
@@ -1639,19 +1711,22 @@ void FWireTranslatorImpl::AddNodeInBodyGroup(TSharedPtr<AlDagNode>& DagNode, con
 	}
 	else
 	{
-		TSharedRef<BodyData> BodyRef = MakeShared<BodyData>(ShaderName, LayerName, bIsAPatch);
+		TSharedRef<BodyData> BodyRef = MakeShared<BodyData>(LayerName, bIsAPatch);
 		ShellToProcess.Add(SetId, BodyRef);
-		BodyRef->ShellSet.Reserve(MaxSize);
+		BodyRef->Reserve(MaxSize);
 		Body = BodyRef;
 	}
-	Body->ShellSet.Add(DagNode);
+
+	const FColor& Color = ShaderNameToColor[ShaderName];
+	Body->Add(DagNode, Color, ShaderName);
 }
 
 bool FWireTranslatorImpl::RecurseDagForLeaves(const TSharedPtr<AlDagNode>& FirstDagNode, const FDagNodeInfo& ParentInfo)
 {
 	if (TessellationOptions.StitchingTechnique != EDatasmithCADStitchingTechnique::StitchingSew)
 	{
-		return RecurseDagForLeavesNoMerge(FirstDagNode, ParentInfo);
+		RecurseDagForLeavesNoMerge(FirstDagNode, ParentInfo);
+		return true;
 	}
 
 	TSharedPtr<AlDagNode> DagNode = FirstDagNode;
@@ -1670,15 +1745,12 @@ bool FWireTranslatorImpl::RecurseDagForLeaves(const TSharedPtr<AlDagNode>& First
 
 	while (DagNode)
 	{
-		// Filter visible nodes.
-		if (!IsHidden(DagNode))
-		{
-			AlObjectType objectType = DagNode->type();
+		AlObjectType objectType = DagNode->type();
 
-			// Process the current node.
-			switch (objectType)
-			{
-				// Push all leaf nodes into 'leaves'
+		// Process the current node.
+		switch (objectType)
+		{
+			// Push all leaf nodes into 'leaves'
 			case kShellNodeType:
 			{
 				AlShellNode* ShellNode = DagNode->asShellNodePtr();
@@ -1712,6 +1784,7 @@ bool FWireTranslatorImpl::RecurseDagForLeaves(const TSharedPtr<AlDagNode>& First
 				}
 				break;
 			}
+
 			case kSurfaceNodeType:
 			{
 				AlSurfaceNode* SurfaceNode = DagNode->asSurfaceNodePtr();
@@ -1750,12 +1823,11 @@ bool FWireTranslatorImpl::RecurseDagForLeaves(const TSharedPtr<AlDagNode>& First
 				ProcessAlGroupNode(*DagNode, ParentInfo);
 				break;
 			}
+
 			default:
 			{
 				break;
 			}
-
-			}
 		}
 
 		DagNode = GetNextNode(DagNode);
@@ -1763,95 +1835,91 @@ bool FWireTranslatorImpl::RecurseDagForLeaves(const TSharedPtr<AlDagNode>& First
 
 	for (const TPair<uint32, TSharedPtr<BodyData>>& Body : ShellToProcess)
 	{
-		ProcessBodyNode(Body.Value.ToSharedRef(), ParentInfo);
+		ProcessBodyNode(Body.Value, ParentInfo);
 	}
 	return true;
 }
 
-bool FWireTranslatorImpl::RecurseDagForLeavesNoMerge(const TSharedPtr<AlDagNode>& FirstDagNode, const FDagNodeInfo& ParentInfo)
+void FWireTranslatorImpl::RecurseDagForLeavesNoMerge(const TSharedPtr<AlDagNode>& FirstDagNode, const FDagNodeInfo& ParentInfo)
 {
 	TSharedPtr<AlDagNode> DagNode = FirstDagNode;
+	while (DagNode)
+	{
+		DagForLeavesNoMerge(DagNode, ParentInfo);
+		DagNode = GetNextNode(DagNode);
+	}
+}
 
-	TMap<uint32, BodyData> ShellToProcess;
-
+void FWireTranslatorImpl::DagForLeavesNoMerge(const TSharedPtr<AlDagNode>& DagNode, const FDagNodeInfo& ParentInfo)
+{
 	const char* ShaderName = nullptr;
 
-	DagNode = FirstDagNode;
-	while (DagNode)
+	// Process the current node.
+	AlObjectType objectType = DagNode->type();
+	switch (objectType)
 	{
-		if (!IsHidden(DagNode))
+		// Push all leaf nodes into 'leaves'
+	case kShellNodeType:
+	{
+		AlShellNode* ShellNode = DagNode->asShellNodePtr();
+		TUniquePtr<AlShell> Shell(ShellNode->shell());
+		if (Shell.IsValid())
 		{
-			// Process the current node.
-			AlObjectType objectType = DagNode->type();
-			switch (objectType)
-			{
-				// Push all leaf nodes into 'leaves'
-			case kShellNodeType:
+			TUniquePtr<AlShader> Shader(Shell->firstShader());
+			if (Shader.IsValid())
 			{
-				AlShellNode* ShellNode = DagNode->asShellNodePtr();
-				TUniquePtr<AlShell> Shell(ShellNode->shell());
-				if (Shell.IsValid())
-				{
-					TUniquePtr<AlShader> Shader(Shell->firstShader());
-					if (Shader.IsValid())
-					{
-						ShaderName = Shader->name();
-					}
-
-					ProcessAlShellNode(DagNode, ParentInfo, ShaderName);
-				}
-				break;
-			}
-			case kSurfaceNodeType:
-			{
-				AlSurfaceNode* SurfaceNode = DagNode->asSurfaceNodePtr();
-				TUniquePtr<AlSurface> Surface(SurfaceNode->surface());
-				if (Surface.IsValid())
-				{
-					TUniquePtr<AlShader> Shader(Surface->firstShader());
-					if (Shader.IsValid())
-					{
-						ShaderName = Shader->name();
-					}
-				}
-				ProcessAlShellNode(DagNode, ParentInfo, ShaderName);
-				break;
+				ShaderName = Shader->name();
 			}
 
-			case kMeshNodeType:
+			ProcessAlShellNode(DagNode, ParentInfo, ShaderName);
+		}
+		break;
+	}
+	case kSurfaceNodeType:
+	{
+		AlSurfaceNode* SurfaceNode = DagNode->asSurfaceNodePtr();
+		TUniquePtr<AlSurface> Surface(SurfaceNode->surface());
+		if (Surface.IsValid())
+		{
+			TUniquePtr<AlShader> Shader(Surface->firstShader());
+			if (Shader.IsValid())
 			{
-				AlMeshNode* MeshNode = DagNode->asMeshNodePtr();
-				TUniquePtr<AlMesh> Mesh(MeshNode->mesh());
-				if (Mesh.IsValid())
-				{
-					TUniquePtr<AlShader> Shader(Mesh->firstShader());
-					if (Shader.IsValid())
-					{
-						ShaderName = Shader->name();
-					}
-				}
-				ProcessAlShellNode(DagNode, ParentInfo, ShaderName);
-				break;
+				ShaderName = Shader->name();
 			}
+		}
+		ProcessAlShellNode(DagNode, ParentInfo, ShaderName);
+		break;
+	}
 
-			// Traverse down through groups
-			case kGroupNodeType:
+	case kMeshNodeType:
+	{
+		AlMeshNode* MeshNode = DagNode->asMeshNodePtr();
+		TUniquePtr<AlMesh> Mesh(MeshNode->mesh());
+		if (Mesh.IsValid())
+		{
+			TUniquePtr<AlShader> Shader(Mesh->firstShader());
+			if (Shader.IsValid())
 			{
-				ProcessAlGroupNode(*DagNode, ParentInfo);
-				break;
+				ShaderName = Shader->name();
 			}
+		}
+		ProcessAlShellNode(DagNode, ParentInfo, ShaderName);
+		break;
+	}
 
-			default:
-			{
-				break;
-			}
+	// Traverse down through groups
+	case kGroupNodeType:
+	{
+		ProcessAlGroupNode(*DagNode, ParentInfo);
+		break;
+	}
 
-			}
+	default:
+	{
+		break;
+	}
 
-			DagNode = GetNextNode(DagNode);
-		}
 	}
-	return true;
 }
 
 TOptional<FMeshDescription> FWireTranslatorImpl::MeshDagNodeWithExternalMesher(AlDagNode& DagNode, TSharedRef<IDatasmithMeshElement> MeshElement, CADLibrary::FMeshParameters& MeshParameters)
@@ -1867,7 +1935,18 @@ TOptional<FMeshDescription> FWireTranslatorImpl::MeshDagNodeWithExternalMesher(A
 		ObjectReference = EAliasObjectReference::WorldReference;
 	}
 
-	AliasBRepConverter->AddBRep(DagNode, ObjectReference);
+	const FColor* ColorPtr = nullptr;
+	const FString* ShaderName = MeshElementToShaderName.Find(&MeshElement.Get());
+	if(ShaderName)
+	{
+		ColorPtr = ShaderNameToColor.Find(*ShaderName);
+	}
+	if (!ColorPtr) // Should never happen
+	{
+		ColorPtr = &DefaultColor;
+	}
+
+	AliasBRepConverter->AddBRep(DagNode, *ColorPtr, ObjectReference);
 
 	CADModelConverter->RepairTopology();
 
@@ -1898,9 +1977,9 @@ TOptional<FMeshDescription> FWireTranslatorImpl::MeshDagNodeWithExternalMesher(T
 		ObjectReference = EAliasObjectReference::ParentReference;
 	}
 
-	for (const TSharedPtr<AlDagNode>& DagNode : Body->ShellSet)
+	for (const TPair<TSharedPtr<AlDagNode>, FColor>& DagNode : Body->GetShells())
 	{
-		AliasBRepConverter->AddBRep(*DagNode, ObjectReference);
+		AliasBRepConverter->AddBRep(*DagNode.Key, DagNode.Value, ObjectReference);
 	}
 
 	CADModelConverter->RepairTopology();
@@ -1932,7 +2011,7 @@ TOptional<FMeshDescription> FWireTranslatorImpl::GetMeshOfShellNode(AlDagNode& D
 		{
 			// Get the meshes from the dag nodes. Note that removing the mesh's DAG.
 			// will also removes the meshes, so we have to do it later.
-			TOptional<FMeshDescription> UEMesh = GetMeshOfNodeMesh(*TesselatedNode, MeshParameters, &AlMatrix);
+			TOptional<FMeshDescription> UEMesh = GetMeshOfNodeMesh(*TesselatedNode, MeshElement, MeshParameters, &AlMatrix);
 			return UEMesh;
 		}
 	}
@@ -1950,11 +2029,10 @@ TOptional<FMeshDescription> FWireTranslatorImpl::GetMeshOfMeshBody(TSharedRef<Bo
 	FMeshDescription MeshDescription;
 	DatasmithMeshHelper::PrepareAttributeForStaticMesh(MeshDescription);
 	MeshDescription.Empty();
-	bool True = true;
 
-	for (const TSharedPtr<AlDagNode>& DagNode : Body->ShellSet)
+	for (const TPair<TSharedPtr<AlDagNode>, FColor>& DagNode : Body->GetShells())
 	{
-		AlMeshNode* MeshNode = DagNode->asMeshNodePtr();
+		AlMeshNode* MeshNode = DagNode.Key->asMeshNodePtr();
 		if (!AlIsValid(MeshNode))
 		{
 			continue;
@@ -1965,15 +2043,25 @@ TOptional<FMeshDescription> FWireTranslatorImpl::GetMeshOfMeshBody(TSharedRef<Bo
 		{
 			continue;
 		}
-
 		TUniquePtr<AlMesh> Mesh(MeshPtr);
-		OpenModelUtils::TransferAlMeshToMeshDescription(*MeshPtr, MeshDescription, MeshParameters, True, true);
+
+		AlMatrix4x4 AlMatrix;
+		DagNode.Key->localTransformationMatrix(AlMatrix);
+
+		Mesh->transform(AlMatrix);
+
+		const FColor& ShaderColor = DagNode.Value;
+		const int32 ShaderId = CreateShaderId(ShaderColor);
+		const FString SlotMaterialName = FString::FromInt(ShaderId);
+
+		const bool bMerge = true;
+		OpenModelUtils::TransferAlMeshToMeshDescription(*MeshPtr, *SlotMaterialName, MeshDescription, MeshParameters, bMerge);
 	}
 
 	return MoveTemp(MeshDescription);
 }
 
-TOptional<FMeshDescription> FWireTranslatorImpl::GetMeshOfNodeMesh(AlDagNode& TesselatedNode, CADLibrary::FMeshParameters& MeshParameters, AlMatrix4x4* AlMeshInvGlobalMatrix)
+TOptional<FMeshDescription> FWireTranslatorImpl::GetMeshOfNodeMesh(AlDagNode& TesselatedNode, TSharedRef<IDatasmithMeshElement> MeshElement, CADLibrary::FMeshParameters& MeshParameters, AlMatrix4x4* AlMeshInvGlobalMatrix)
 {
 	AlMeshNode* MeshNode = TesselatedNode.asMeshNodePtr();
 	if (!AlIsValid(MeshNode))
@@ -1993,17 +2081,18 @@ TOptional<FMeshDescription> FWireTranslatorImpl::GetMeshOfNodeMesh(AlDagNode& Te
 		SharedMesh->transform(*AlMeshInvGlobalMatrix);
 	}
 
-	return ImportMesh(*Mesh, MeshParameters);
+	return ImportMesh(*Mesh, MeshElement, MeshParameters);
 }
 
 TOptional<FMeshDescription> FWireTranslatorImpl::GetMeshDescription(TSharedRef<IDatasmithMeshElement> MeshElement, CADLibrary::FMeshParameters& MeshParameters, TSharedRef<BodyData> Body)
 {
-	if (Body->ShellSet.Num() == 0)
+	if (Body->GetShells().Num() == 0)
 	{
 		return TOptional<FMeshDescription>();
 	}
 
-	TSharedPtr<AlDagNode> DagNode = Body->ShellSet[0];
+	const TArray<TPair<TSharedPtr<AlDagNode>, FColor>>& Shells = Body->GetShells();
+	TSharedPtr<AlDagNode> DagNode = Shells[0].Key;
 	AlLayer* LayerPtr = DagNode->layer();
 	if (AlIsValid(LayerPtr))
 	{
@@ -2024,7 +2113,7 @@ TOptional<FMeshDescription> FWireTranslatorImpl::GetMeshDescription(TSharedRef<I
 		}
 	}
 
-	if (Body->bCadData)
+	if (Body->IsCadData())
 	{
 		return GetMeshOfShellBody(Body, MeshElement, MeshParameters);
 	}
@@ -2088,7 +2177,7 @@ TOptional<FMeshDescription> FWireTranslatorImpl::GetMeshDescription(TSharedRef<I
 
 	case kMeshNodeType:
 	{
-		return GetMeshOfNodeMesh(DagNode, MeshParameters);
+		return GetMeshOfNodeMesh(DagNode, MeshElement, MeshParameters);
 		break;
 	}
 
@@ -2101,13 +2190,26 @@ TOptional<FMeshDescription> FWireTranslatorImpl::GetMeshDescription(TSharedRef<I
 
 
 // Note that Alias file unit is cm like UE
-TOptional<FMeshDescription> FWireTranslatorImpl::ImportMesh(AlMesh& InMesh, CADLibrary::FMeshParameters& InMeshParameters)
+TOptional<FMeshDescription> FWireTranslatorImpl::ImportMesh(AlMesh& InMesh, TSharedRef<IDatasmithMeshElement> MeshElement, CADLibrary::FMeshParameters& InMeshParameters)
 {
-
 	FMeshDescription MeshDescription;
 	DatasmithMeshHelper::PrepareAttributeForStaticMesh(MeshDescription);
-	bool bMerge = true;
-	OpenModelUtils::TransferAlMeshToMeshDescription(InMesh, MeshDescription, InMeshParameters, bMerge);
+
+	const FColor* ShaderColorPtr = nullptr;
+	const FString* ShaderName = MeshElementToShaderName.Find(&MeshElement.Get());
+	if (ShaderName)
+	{
+		ShaderColorPtr = ShaderNameToColor.Find(*ShaderName);
+	}
+	if (!ShaderColorPtr) // Should never happen
+	{
+		ShaderColorPtr = &DefaultColor;
+	}
+	int32 ShaderId = CreateShaderId(*ShaderColorPtr);
+	FString SlotMaterialName = FString::FromInt(ShaderId);
+
+	const bool bMerge = false;
+	OpenModelUtils::TransferAlMeshToMeshDescription(InMesh, *SlotMaterialName, MeshDescription, InMeshParameters, bMerge);
 
 	return MoveTemp(MeshDescription);
 }
diff --git a/Engine/Plugins/Enterprise/DatasmithCADImporter/Source/DatasmithWireTranslator/Private/OpenModelUtils.cpp b/Engine/Plugins/Enterprise/DatasmithCADImporter/Source/DatasmithWireTranslator/Private/OpenModelUtils.cpp
--- a/Engine/Plugins/Enterprise/DatasmithCADImporter/Source/DatasmithWireTranslator/Private/OpenModelUtils.cpp
+++ b/Engine/Plugins/Enterprise/DatasmithCADImporter/Source/DatasmithWireTranslator/Private/OpenModelUtils.cpp
@@ -1611,7 +1611,7 @@ bool OpenModelUtils::IsValidActor(const TSharedPtr<IDatasmithActorElement>& Acto
 	return false;
 }
 
-bool OpenModelUtils::TransferAlMeshToMeshDescription(const AlMesh& AliasMesh, FMeshDescription& MeshDescription, CADLibrary::FMeshParameters& MeshParameters, bool& bHasNormal, bool bMerge)
+bool OpenModelUtils::TransferAlMeshToMeshDescription(const AlMesh& AliasMesh, const TCHAR* SlotMaterialId, FMeshDescription& MeshDescription, CADLibrary::FMeshParameters& MeshParameters, const bool bMerge)
 {
 	if (AliasMesh.numberOfVertices() == 0 || AliasMesh.numberOfTriangles() == 0)
 	{
@@ -1664,7 +1664,7 @@ bool OpenModelUtils::TransferAlMeshToMeshDescription(const AlMesh& AliasMesh, FM
 	// Assume one material per mesh, no partitioning
 	MeshDescription.ReserveNewPolygonGroups(1);
 	FPolygonGroupID PolyGroupId = MeshDescription.CreatePolygonGroup();
-	FName ImportedSlotName = TEXT("0"); // No access to DatasmithMeshHelper::DefaultSlotName
+	FName ImportedSlotName = SlotMaterialId;
 	PolygonGroupImportedMaterialSlotNames[PolyGroupId] = ImportedSlotName;
 
 	// At least one UV set must exist.
diff --git a/Engine/Plugins/Enterprise/DatasmithCADImporter/Source/DatasmithWireTranslator/Private/OpenModelUtils.h b/Engine/Plugins/Enterprise/DatasmithCADImporter/Source/DatasmithWireTranslator/Private/OpenModelUtils.h
--- a/Engine/Plugins/Enterprise/DatasmithCADImporter/Source/DatasmithWireTranslator/Private/OpenModelUtils.h
+++ b/Engine/Plugins/Enterprise/DatasmithCADImporter/Source/DatasmithWireTranslator/Private/OpenModelUtils.h
@@ -65,20 +65,19 @@ namespace OpenModelUtils
 		return HashCombine(IdA, HashCombine(IdB, HashCombine(IdC, IdD)));
 	}
 
-	inline uint32 GetAlDagNodeUuid(const TSharedPtr<AlDagNode>& DagNode)
+	inline uint32 GetAlDagNodeUuid(AlDagNode& DagNode)
 	{
-		if (DagNode->hasPersistentID() == sSuccess)
+		if (DagNode.hasPersistentID() == sSuccess)
 		{
 			AlPersistentID* PersistentID;
-			DagNode->persistentID(PersistentID);
+			DagNode.persistentID(PersistentID);
 			return GetTypeHash(*PersistentID);
 		}
-		FString Label = UTF8_TO_TCHAR(DagNode->name());
+		FString Label = UTF8_TO_TCHAR(DagNode.name());
 		return GetTypeHash(Label);
 	}
 
-	// Note that Alias file unit is cm like UE
-	bool TransferAlMeshToMeshDescription(const AlMesh& Mesh, FMeshDescription& MeshDescription, CADLibrary::FMeshParameters& SymmetricParameters, bool& bHasNormal, bool bMerge = false);
+	bool TransferAlMeshToMeshDescription(const AlMesh& Mesh, const TCHAR* SlotMaterialName, FMeshDescription& MeshDescription, CADLibrary::FMeshParameters& SymmetricParameters, const bool bMerge = false);
 
 	TSharedPtr<AlDagNode> TesselateDagLeaf(const AlDagNode& DagLeaf, ETesselatorType TessType, double Tolerance);
 }
diff --git a/Engine/Plugins/Enterprise/DatasmithCADImporter/Source/ParametricSurface/Private/TechSoft/CADModelToTechSoftConverterBase.cpp b/Engine/Plugins/Enterprise/DatasmithCADImporter/Source/ParametricSurface/Private/TechSoft/CADModelToTechSoftConverterBase.cpp
--- a/Engine/Plugins/Enterprise/DatasmithCADImporter/Source/ParametricSurface/Private/TechSoft/CADModelToTechSoftConverterBase.cpp
+++ b/Engine/Plugins/Enterprise/DatasmithCADImporter/Source/ParametricSurface/Private/TechSoft/CADModelToTechSoftConverterBase.cpp
@@ -55,8 +55,6 @@ bool FCADModelToTechSoftConverterBase::SaveModel(const TCHAR* InFolderPath, TSha
 		// This will be used when the file is reloaded
 		TSharedPtr<FJsonObject> JsonObject = MakeShared<FJsonObject>();
 		JsonObject->SetNumberField(JSON_ENTRY_BODY_UNIT, 0.1); // ALL BRep from Alias or Rhino are defined in mm
-		JsonObject->SetNumberField(JSON_ENTRY_COLOR_NAME, 0);
-		JsonObject->SetNumberField(JSON_ENTRY_MATERIAL_NAME, 0);
 
 		TSharedRef< TJsonWriter< TCHAR, TPrettyJsonPrintPolicy<TCHAR> > > JsonWriter = TJsonWriterFactory< TCHAR, TPrettyJsonPrintPolicy<TCHAR> >::Create(&JsonString);
 
@@ -83,6 +81,9 @@ bool FCADModelToTechSoftConverterBase::Tessellate(const CADLibrary::FMeshParamet
 	}
 #endif
 
+	TSharedPtr<FJsonObject> JsonObject = MakeShared<FJsonObject>();
+	CADLibrary::TechSoftUtils::RestoreMaterials(JsonObject, BodyMesh);
+
 	if (BodyMesh.Faces.Num() == 0)
 	{
 		return false;
