The NURBS surface was not parsed in the good reference frame with Sew option.
The methode TransferAlMeshToMeshDescription called for each sub mesh call FStaticMeshOperations::DetermineEdgeHardnessesFromVertexInstanceNormals. This methode is moved outside TransferAlMeshToMeshDescription to be called only once.

diff --git a/Engine/Plugins/Enterprise/DatasmithCADImporter/Source/DatasmithWireTranslator/Private/DatasmithWireTranslator.cpp b/Engine/Plugins/Enterprise/DatasmithCADImporter/Source/DatasmithWireTranslator/Private/DatasmithWireTranslator.cpp
--- a/Engine/Plugins/Enterprise/DatasmithCADImporter/Source/DatasmithWireTranslator/Private/DatasmithWireTranslator.cpp
+++ b/Engine/Plugins/Enterprise/DatasmithCADImporter/Source/DatasmithWireTranslator/Private/DatasmithWireTranslator.cpp
@@ -1971,7 +1971,7 @@ TOptional<FMeshDescription> FWireTranslatorImpl::MeshDagNodeWithExternalMesher(T
 		// All actors of a Alias symmetric layer are defined in the world Reference i.e. they have identity transform. So Mesh actor has to be defined in the world reference.
 		ObjectReference = EAliasObjectReference::WorldReference;
 	}
-	else if (TessellationOptions.StitchingTechnique != EDatasmithCADStitchingTechnique::StitchingSew)
+	else if (TessellationOptions.StitchingTechnique == EDatasmithCADStitchingTechnique::StitchingSew)
 	{
 		// In the case of StitchingSew, AlDagNode children of a GroupNode are merged together. To be merged, they have to be defined in the reference of parent GroupNode.
 		ObjectReference = EAliasObjectReference::ParentReference;
@@ -2057,6 +2057,9 @@ TOptional<FMeshDescription> FWireTranslatorImpl::GetMeshOfMeshBody(TSharedRef<Bo
 		const bool bMerge = true;
 		OpenModelUtils::TransferAlMeshToMeshDescription(*MeshPtr, *SlotMaterialName, MeshDescription, MeshParameters, bMerge);
 	}
+	
+	// Build edge meta data
+	FStaticMeshOperations::DetermineEdgeHardnessesFromVertexInstanceNormals(MeshDescription);
 
 	return MoveTemp(MeshDescription);
 }
@@ -2211,6 +2214,9 @@ TOptional<FMeshDescription> FWireTranslatorImpl::ImportMesh(AlMesh& InMesh, TSha
 	const bool bMerge = false;
 	OpenModelUtils::TransferAlMeshToMeshDescription(InMesh, *SlotMaterialName, MeshDescription, InMeshParameters, bMerge);
 
+	// Build edge meta data
+	FStaticMeshOperations::DetermineEdgeHardnessesFromVertexInstanceNormals(MeshDescription);
+
 	return MoveTemp(MeshDescription);
 }
 
diff --git a/Engine/Plugins/Enterprise/DatasmithCADImporter/Source/DatasmithWireTranslator/Private/OpenModelUtils.cpp b/Engine/Plugins/Enterprise/DatasmithCADImporter/Source/DatasmithWireTranslator/Private/OpenModelUtils.cpp
--- a/Engine/Plugins/Enterprise/DatasmithCADImporter/Source/DatasmithWireTranslator/Private/OpenModelUtils.cpp
+++ b/Engine/Plugins/Enterprise/DatasmithCADImporter/Source/DatasmithWireTranslator/Private/OpenModelUtils.cpp
@@ -1805,9 +1805,6 @@ bool OpenModelUtils::TransferAlMeshToMeshDescription(const AlMesh& AliasMesh, co
 		}
 	}
 
-	// Build edge meta data
-	FStaticMeshOperations::DetermineEdgeHardnessesFromVertexInstanceNormals(MeshDescription);
-
 	return true;
 }
 
