Added CAD importer to Linux build

diff --git a/Engine/Plugins/Enterprise/DatasmithCADImporter/DatasmithCADImporter.uplugin b/Engine/Plugins/Enterprise/DatasmithCADImporter/DatasmithCADImporter.uplugin
--- a/Engine/Plugins/Enterprise/DatasmithCADImporter/DatasmithCADImporter.uplugin
+++ b/Engine/Plugins/Enterprise/DatasmithCADImporter/DatasmithCADImporter.uplugin
@@ -23,14 +23,15 @@
 			"Name": "CADLibrary",
 			"Type": "Runtime",
 			"LoadingPhase": "Default",
-			"PlatformAllowList": [ "Win64" ]
+			"PlatformAllowList": [ "Win64", "Linux" ]
 		},
 		{
 			"Name": "DatasmithCADTranslator",
 			"Type": "Runtime",
 			"LoadingPhase": "Default",
 			"PlatformAllowList": [
-				"Win64"
+				"Win64",
+				"Linux"
 			]
 		},
 		{
@@ -46,7 +47,8 @@
 			"Type": "Runtime",
 			"LoadingPhase": "Default",
 			"PlatformAllowList": [
-				"Win64"
+				"Win64",
+				"Linux"
 			]
 		},
 		{
@@ -54,7 +56,8 @@
 			"Type": "Runtime",
 			"LoadingPhase": "Default",
 			"PlatformAllowList": [
-				"Win64"
+				"Win64",
+				"Linux"
 			]
 		},
 		{
@@ -62,7 +65,8 @@
 			"Type": "Runtime",
 			"LoadingPhase": "Default",
 			"PlatformAllowList": [
-				"Win64"
+				"Win64",
+				"Linux"
 			]
 		},
 		{
@@ -110,7 +114,8 @@
 			"Type": "Editor",
 			"LoadingPhase": "Default",
 			"PlatformAllowList": [
-				"Win64"
+				"Win64",
+				"Linux"
 			]
 		},
 		{
@@ -118,7 +123,8 @@
 			"Type": "Runtime",
 			"LoadingPhase": "Default",
 			"PlatformAllowList": [
-				"Win64"
+				"Win64",
+				"Linux"
 			]
 		},
 		{
@@ -126,7 +132,8 @@
 			"Type": "Runtime",
 			"LoadingPhase": "Default",
 			"PlatformAllowList": [
-				"Win64"
+				"Win64",
+				"Linux"
 			]
 		},
 		{
diff --git a/Engine/Plugins/Enterprise/DatasmithCADImporter/Source/CADInterfaces/CADInterfaces.Build.cs b/Engine/Plugins/Enterprise/DatasmithCADImporter/Source/CADInterfaces/CADInterfaces.Build.cs
--- a/Engine/Plugins/Enterprise/DatasmithCADImporter/Source/CADInterfaces/CADInterfaces.Build.cs
+++ b/Engine/Plugins/Enterprise/DatasmithCADImporter/Source/CADInterfaces/CADInterfaces.Build.cs
@@ -27,15 +27,12 @@ namespace UnrealBuildTool.Rules
 			bool bHasCoretech = System.Type.GetType("CoreTech") != null;
 			bool bHasTechSoft = System.Type.GetType("TechSoft") != null;
 
-			// Support for Windows only
-			bool bIsPlateformSupported = Target.Platform == UnrealTargetPlatform.Win64;
-
-			if (bIsPlateformSupported && bHasCoretech)
+			if (Target.Platform == UnrealTargetPlatform.Win64 && bHasCoretech)
 			{
 				PublicDependencyModuleNames.Add("CoreTech");
 			}
 
-			if (bIsPlateformSupported && bHasTechSoft)
+			if ((Target.Platform == UnrealTargetPlatform.Win64 || Target.Platform == UnrealTargetPlatform.Linux) && bHasTechSoft)
 			{
 				PublicDependencyModuleNames.Add("TechSoft");
 			}
