CAD import fixes (manual merge from 5.0.3):

FTechSoftFileParser::TraverseBRepModel: Default material was not applied from parent

Static mesh from Rhino has its unique material at slot 0 while other CAD import (from techsoft or Alias) static mesh can have many material and the slot of the material is defined by the material/color components.
To respect the mecanisme of other CAD import, BRep from Rhino is added of a color that has a slotId = 0. This color is (0, 0, 0, 0)
The Rhino import will be improve like Alias import to allow static mesh with multiple of materials (jira UETOOL-5196)

Revert to revision 2 + adapt value to TechSoft tessellation and TechSoft material

diff --git a/Engine/Plugins/Enterprise/DatasmithCADImporter/Source/CADInterfaces/Private/TechSoftFileParser.cpp b/Engine/Plugins/Enterprise/DatasmithCADImporter/Source/CADInterfaces/Private/TechSoftFileParser.cpp
--- a/Engine/Plugins/Enterprise/DatasmithCADImporter/Source/CADInterfaces/Private/TechSoftFileParser.cpp
+++ b/Engine/Plugins/Enterprise/DatasmithCADImporter/Source/CADInterfaces/Private/TechSoftFileParser.cpp
@@ -8,6 +8,7 @@
 #include "TechSoftUtilsPrivate.h"
 #include "TUniqueTechSoftObj.h"
 
+#include "HAL/FileManager.h"
 #include "Serialization/JsonSerializer.h"
 #include "Serialization/JsonWriter.h"
 #include "Templates/UnrealTemplate.h"
@@ -19,6 +20,64 @@ namespace CADLibrary
 
 namespace TechSoftFileParserImpl
 {
+
+// This code is a duplication code of CADFileReader::FindFile 
+// This is done in 5.0.3 to avoid public header modification. 
+// However this need to be rewrite in the next version. (Jira UE-152626)
+bool UpdateFileDescriptor(FFileDescriptor& File)
+{
+	const FString& FileName = File.GetFileName();
+
+	FString FilePath = FPaths::GetPath(File.GetSourcePath());
+	FString RootFilePath = File.GetRootFolder();
+
+	// Basic case: File exists at the initial path
+	if (IFileManager::Get().FileExists(*File.GetSourcePath()))
+	{
+		return true;
+	}
+
+	// Advance case: end of FilePath is in a upper-folder of RootFilePath
+	// e.g.
+	// FilePath = D:\\data temp\\Unstructured project\\Folder2\\Added_Object.SLDPRT
+	//                                                 ----------------------------
+	// RootFilePath = D:\\data\\CAD Files\\SolidWorks\\p033 - Unstructured project\\Folder1
+	//                ------------------------------------------------------------
+	// NewPath = D:\\data\\CAD Files\\SolidWorks\\p033 - Unstructured project\\Folder2\\Added_Object.SLDPRT
+	TArray<FString> RootPaths;
+	RootPaths.Reserve(30);
+	do
+	{
+		RootFilePath = FPaths::GetPath(RootFilePath);
+		RootPaths.Emplace(RootFilePath);
+	} while (!FPaths::IsDrive(RootFilePath) && !RootFilePath.IsEmpty());
+
+	TArray<FString> FilePaths;
+	FilePaths.Reserve(30);
+	FilePaths.Emplace(FileName);
+	while (!FPaths::IsDrive(FilePath) && !FilePath.IsEmpty())
+	{
+		FString FolderName = FPaths::GetCleanFilename(FilePath);
+		FilePath = FPaths::GetPath(FilePath);
+		FilePaths.Emplace(FPaths::Combine(FolderName, FilePaths.Last()));
+	};
+
+	for (int32 IndexFolderPath = 0; IndexFolderPath < RootPaths.Num(); IndexFolderPath++)
+	{
+		for (int32 IndexFilePath = 0; IndexFilePath < FilePaths.Num(); IndexFilePath++)
+		{
+			FString NewFilePath = FPaths::Combine(RootPaths[IndexFolderPath], FilePaths[IndexFilePath]);
+			if (IFileManager::Get().FileExists(*NewFilePath))
+			{
+				File.SetSourceFilePath(NewFilePath);
+				return true;
+			};
+		}
+	}
+
+	return false;
+}
+
 // Functions to clean metadata
 
 inline void RemoveUnwantedChar(FString& StringToClean, TCHAR UnwantedChar)
@@ -295,12 +354,6 @@ ECADParsingResult FTechSoftFileParser::Process()
 		FileUnit = TechSoftInterface::GetModelFileUnit(ModelFile.Get());
 	}
 
-	// Adapt BRep to CADKernel
-	if (AdaptBRepModel() != A3D_SUCCESS)
-	{
-		return ECADParsingResult::ProcessFailed;
-	}
-
 	// save the file for the next load
 	if (CADFileData.IsCacheDefined())
 	{
@@ -311,6 +364,12 @@ ECADParsingResult FTechSoftFileParser::Process()
 		}
 	}
 
+	// Adapt BRep to CADKernel
+	if (AdaptBRepModel() != A3D_SUCCESS)
+	{
+		return ECADParsingResult::ProcessFailed;
+	}
+
 	// Some formats (like IGES) require a sew all the time. In this case, bForceSew = true
 	if (bForceSew || CADFileData.GetImportParameters().GetStitchingTechnique() == StitchingSew)
 	{
@@ -410,7 +469,8 @@ void FTechSoftFileParser::GenerateBodyMesh(A3DRiRepresentationItem* Representati
 		uint32 CachedStyleIndex = Tessellation.MaterialName;
 		Tessellation.MaterialName = 0;
 
-		if (CachedStyleIndex != FTechSoftInterface::InvalidScriptIndex)
+		constexpr uint32 GraphStyleDataDefaultValue =  65535;
+		if (CachedStyleIndex != GraphStyleDataDefaultValue)
 		{
 			ExtractGraphStyleProperties(CachedStyleIndex, ColorName, MaterialName);
 		}
@@ -506,6 +566,7 @@ ECADParsingResult FTechSoftFileParser::TraverseModel()
 	FEntityMetaData MetaData;
 	ExtractMetaData(ModelFile.Get(), MetaData);
 	ExtractSpecificMetaData(ModelFile.Get(), MetaData);
+	BuildReferenceName(MetaData);
 
 	for (uint32 Index = 0; Index < ModelFileData->m_uiPOccurrencesSize; ++Index)
 	{
@@ -515,7 +576,8 @@ ECADParsingResult FTechSoftFileParser::TraverseModel()
 		}
 		else
 		{
-			TraverseReference(ModelFileData->m_ppPOccurrences[Index], FMatrix::Identity, FileUnit);
+			FString* ModelNamePtr = MetaData.MetaData.Find(TEXT("Name"));
+			TraverseReference(ModelFileData->m_ppPOccurrences[Index], ModelNamePtr ? *ModelNamePtr : FString(), FMatrix::Identity, FileUnit);
 		}
 	}
 
@@ -533,6 +595,8 @@ void FTechSoftFileParser::TraverseConfigurationSet(const A3DAsmProductOccurrence
 	FEntityMetaData MetaData;
 	ExtractMetaData(ConfigurationSetPtr, MetaData);
 	ExtractSpecificMetaData(ConfigurationSetPtr, MetaData);
+	BuildReferenceName(MetaData);
+	FString ConfigurationName = MetaData.MetaData.FindOrAdd(TEXT("Name"));
 
 	const FString& ConfigurationToLoad = CADFileData.GetCADFileDescription().GetConfiguration();
 
@@ -560,10 +624,10 @@ void FTechSoftFileParser::TraverseConfigurationSet(const A3DAsmProductOccurrence
 			{
 				FEntityMetaData ConfigurationMetaData;
 				ExtractMetaData(ConfigurationSetData->m_ppPOccurrences[Index], ConfigurationMetaData);
-				const FString* ConfigurationName = ConfigurationMetaData.MetaData.Find(TEXT("SDKName"));
-				if (ConfigurationName)
+				const FString* Configuration = ConfigurationMetaData.MetaData.Find(TEXT("SDKName"));
+				if (Configuration)
 				{
-					bIsConfigurationToLoad = (ConfigurationName->Equals(ConfigurationToLoad));
+					bIsConfigurationToLoad = (Configuration->Equals(ConfigurationToLoad));
 				}
 			}
 			else
@@ -573,7 +637,7 @@ void FTechSoftFileParser::TraverseConfigurationSet(const A3DAsmProductOccurrence
 
 			if (bIsConfigurationToLoad)
 			{
-				TraverseReference(ConfigurationSetData->m_ppPOccurrences[Index], TransformMatrix, OccurenceUnit);
+				TraverseReference(ConfigurationSetData->m_ppPOccurrences[Index], ConfigurationName, TransformMatrix, OccurenceUnit);
 				return;
 			}
 		}
@@ -592,7 +656,7 @@ void FTechSoftFileParser::TraverseConfigurationSet(const A3DAsmProductOccurrence
 
 			if (ConfigurationData->m_uiProductFlags & A3D_PRODUCT_FLAG_CONFIG)
 			{
-				TraverseReference(ConfigurationSetData->m_ppPOccurrences[Index], TransformMatrix, OccurenceUnit);
+				TraverseReference(ConfigurationSetData->m_ppPOccurrences[Index], ConfigurationName, TransformMatrix, OccurenceUnit);
 			}
 		}
 	}
@@ -662,7 +726,7 @@ void FTechSoftFileParser::CountUnderConfigurationSet(const A3DAsmProductOccurren
 	}
 }
 
-void FTechSoftFileParser::TraverseReference(const A3DAsmProductOccurrence* ReferencePtr, const FMatrix& ParentMatrix, double ParentUnit)
+void FTechSoftFileParser::TraverseReference(const A3DAsmProductOccurrence* ReferencePtr, const FString& RootName, const FMatrix& ParentMatrix, double ParentUnit)
 {
 	FEntityMetaData MetaData;
 	ExtractMetaData(ReferencePtr, MetaData);
@@ -673,7 +737,15 @@ void FTechSoftFileParser::TraverseReference(const A3DAsmProductOccurrence* Refer
 	}
 
 	ExtractSpecificMetaData(ReferencePtr, MetaData);
-	BuildReferenceName(MetaData);
+
+	if (!RootName.IsEmpty())
+	{
+		MetaData.MetaData.FindOrAdd(TEXT("Name"), RootName);
+	}
+	else
+	{
+		BuildReferenceName(MetaData);
+	}
 
 	FArchiveInstance EmptyInstance;
 	FArchiveComponent& Component = FTechSoftFileParser::AddComponent(MetaData, EmptyInstance);
@@ -695,9 +767,10 @@ void FTechSoftFileParser::TraverseReference(const A3DAsmProductOccurrence* Refer
 
 	Component.TransformMatrix = ParentMatrix * ReferenceMatrix;
 
+	const FString& OccurenceNameBase = Component.MetaData.FindOrAdd(TEXT("Name"));
 	for (uint32 OccurenceIndex = 0; OccurenceIndex < ReferenceData->m_uiPOccurrencesSize; ++OccurenceIndex)
 	{
-		int32 ChildrenId = TraverseOccurrence(ReferenceData->m_ppPOccurrences[OccurenceIndex], ReferenceUnit);
+		int32 ChildrenId = TraverseOccurrence(ReferenceData->m_ppPOccurrences[OccurenceIndex], OccurenceNameBase + TEXT("_") + FString::FromInt(OccurenceIndex+1), ReferenceUnit);
 		Component.Children.Add(ChildrenId);
 	}
 
@@ -738,6 +811,9 @@ FArchiveUnloadedComponent& FTechSoftFileParser::AddUnloadedComponent(FEntityMeta
 	Instance.bIsExternalReference = true;
 	Instance.ReferenceNodeId = ComponentId;
 
+	// Make sure that the ExternalFile path is right otherwise try to find the file and update.
+	TechSoftFileParserImpl::UpdateFileDescriptor(ComponentMetaData.ExternalFile);
+
 	Instance.ExternalReference = ComponentMetaData.ExternalFile;
 
 	if (Format == ECADFormat::SOLIDWORKS /*|| CADFileData.FileFormat() == ECADFormat::JT*/)
@@ -790,7 +866,7 @@ int32 FTechSoftFileParser::AddBody(FEntityMetaData& BodyMetaData, const FMatrix&
 	return BodyIndex;
 }
 
-FCadId FTechSoftFileParser::TraverseOccurrence(const A3DAsmProductOccurrence* OccurrencePtr, double ParentUnit)
+FCadId FTechSoftFileParser::TraverseOccurrence(const A3DAsmProductOccurrence* OccurrencePtr, const FString& DefaultOccurrenceName, double ParentUnit)
 {
 	// first product occurence with m_pPart != nullptr || m_uiPOccurrencesSize > 0
 	const A3DAsmProductOccurrence* CachedOccurrencePtr = OccurrencePtr;
@@ -809,6 +885,9 @@ FCadId FTechSoftFileParser::TraverseOccurrence(const A3DAsmProductOccurrence* Oc
 		return 0;
 	}
 
+	// Todo
+	// the life time of FEntityMetaData end at the creation of FArchiveInstance.
+	// This variable should not be accessible after
 	FEntityMetaData InstanceMetaData;
 	ExtractMetaData(OccurrencePtr, InstanceMetaData);
 
@@ -818,12 +897,15 @@ FCadId FTechSoftFileParser::TraverseOccurrence(const A3DAsmProductOccurrence* Oc
 	}
 
 	ExtractSpecificMetaData(OccurrencePtr, InstanceMetaData);
-	BuildInstanceName(InstanceMetaData);
+	BuildInstanceName(InstanceMetaData, DefaultOccurrenceName);
 
 	FArchiveInstance& Instance = AddInstance(InstanceMetaData);
 
 	A3DMiscTransformation* Location = OccurrenceData->m_pLocation;
 
+	// Todo
+	// the life time of FEntityMetaData end at the creation of FArchive(Unloaded)Component.
+	// This variable should not be accessible after
 	FEntityMetaData PrototypeMetaData;
 	if (OccurrenceData->m_pPrototype)
 	{
@@ -842,7 +924,17 @@ FCadId FTechSoftFileParser::TraverseOccurrence(const A3DAsmProductOccurrence* Oc
 		return Instance.ObjectId;
 	}
 	
-	FArchiveComponent& Component = AddComponent(InstanceMetaData, Instance);
+	// If the prototype hasn't name, set its name with the name of the instance 
+	{
+		FString& PrototypeName = PrototypeMetaData.MetaData.FindOrAdd(TEXT("Name"));
+		if (PrototypeName.IsEmpty())
+		{
+			const FString& InstanceName = Instance.MetaData.FindOrAdd(TEXT("Name"));
+			PrototypeName = InstanceName;
+		}
+	}
+
+	FArchiveComponent& Component = AddComponent(PrototypeMetaData, Instance);
 
 	while (OccurrenceData->m_pPrototype != nullptr && OccurrenceData->m_pPart == nullptr && OccurrenceData->m_uiPOccurrencesSize == 0 && OccurrenceData->m_pExternalData == nullptr)
 	{
@@ -875,9 +967,11 @@ FCadId FTechSoftFileParser::TraverseOccurrence(const A3DAsmProductOccurrence* Oc
 
 	uint32 ChildrenCount = OccurrenceData->m_uiPOccurrencesSize;
 	A3DAsmProductOccurrence** Children = OccurrenceData->m_ppPOccurrences;
+	const FString& InstanceName = Instance.MetaData.FindOrAdd(TEXT("Name"));
 	for (uint32 Index = 0; Index < ChildrenCount; ++Index)
 	{
-		int32 ChildrenId = TraverseOccurrence(Children[Index], OccurrenceUnit);
+		const FString DefaultChildName = InstanceName + TEXT("_") + FString::FromInt(Index + 1);
+		int32 ChildrenId = TraverseOccurrence(Children[Index], DefaultChildName, OccurrenceUnit);
 		Component.Children.Add(ChildrenId);
 	}
 
@@ -900,8 +994,9 @@ FCadId FTechSoftFileParser::TraverseOccurrence(const A3DAsmProductOccurrence* Oc
 		A3DAsmProductOccurrence** ExternalChildren = ExternalData->m_ppPOccurrences;
 		for (uint32 Index = 0; Index < ExternalChildrenCount; ++Index)
 		{
-			int32 ChildrenId = TraverseOccurrence(ExternalChildren[Index], OccurrenceUnit);
-			Component.Children.Add(ChildrenId);
+			const FString DefaultChildName = InstanceName + TEXT("_") + FString::FromInt(Index + 1);
+			int32 ChildId = TraverseOccurrence(ExternalChildren[Index], DefaultChildName, OccurrenceUnit);
+			Component.Children.Add(ChildId);
 		}
 	}
 
@@ -1037,7 +1132,9 @@ void FTechSoftFileParser::ProcessPrototype(const A3DAsmProductOccurrence* InProt
 			}
 		}
 		
-		if (PrototypeData->m_pPart != nullptr|| PrototypeData->m_uiPOccurrencesSize != 0 || PrototypeData->m_pExternalData != nullptr )
+		// Case of DWG file, external reference is in fact internal data as a part or a component
+		// Case of CATProduct referencing CGR files: unloaded references are saved in ExternalData. In this case non empty ExternalFile means unloaded reference
+		if (PrototypeData->m_pPart != nullptr|| PrototypeData->m_uiPOccurrencesSize != 0 || (PrototypeData->m_pExternalData != nullptr && OutPrototypeMetaData.ExternalFile.IsEmpty()))
 		{
 			OutPrototypeMetaData.bUnloaded = false;
 			PrototypePtr = nullptr;
@@ -1099,14 +1196,32 @@ void FTechSoftFileParser::TraversePartDefinition(const A3DAsmPartDefinition* Par
 	}
 
 	ExtractSpecificMetaData(PartDefinitionPtr, PartMetaData);
-	BuildPartName(PartMetaData);
+
+	// A top/down propagation of the color and material is needed. This is done in 5.0.3 with minimal code modification. 
+	// However the model parsing need to be rewrite in the next version. (Jira UE-152691)
+	const FString* Material = Part.MetaData.Find(TEXT("MaterialName"));
+	const FString* Color = Part.MetaData.Find(TEXT("ColorName"));
+	if (Material && !Material->IsEmpty() && PartMetaData.MaterialName == 0)
+	{
+		uint32 MaterialName = (uint32) FCString::Atoi64(**Material);
+		PartMetaData.MetaData.Add(TEXT("Material"), *Material);
+		PartMetaData.MaterialName = MaterialName;
+	}
+	if (Color && !Color->IsEmpty() && PartMetaData.ColorName == 0)
+	{
+		uint32 ColorName = (uint32)FCString::Atoi64(**Color);
+		PartMetaData.MetaData.Add(TEXT("Color"), *Color);
+		PartMetaData.ColorName = ColorName;
+	}
+
+	BuildPartName(PartMetaData, Part);
 
 	TUniqueTSObj<A3DAsmPartDefinitionData> PartData(PartDefinitionPtr);
 	if (PartData.IsValid())
 	{
 		for (unsigned int Index = 0; Index < PartData->m_uiRepItemsSize; ++Index)
 		{
-			int32 ChildId = TraverseRepresentationItem(PartData->m_ppRepItems[Index], PartMetaData, Part.ObjectId, ParentUnit);
+			int32 ChildId = TraverseRepresentationItem(PartData->m_ppRepItems[Index], PartMetaData, Part.ObjectId, ParentUnit, Index);
 			Part.Children.Add(ChildId);
 		}
 	}
@@ -1127,7 +1242,7 @@ void FTechSoftFileParser::CountUnderPartDefinition(const A3DAsmPartDefinition* P
 	}
 }
 
-FCadId FTechSoftFileParser::TraverseRepresentationItem(A3DRiRepresentationItem* RepresentationItem, const FEntityMetaData& PartMetaData, const FCadId ParentId, double ParentUnit)
+FCadId FTechSoftFileParser::TraverseRepresentationItem(A3DRiRepresentationItem* RepresentationItem, const FEntityMetaData& PartMetaData, const FCadId ParentId, double ParentUnit, int32 ItemIndex)
 {
 	A3DEEntityType Type;
 	A3DEntityGetType(RepresentationItem, &Type);
@@ -1137,9 +1252,9 @@ FCadId FTechSoftFileParser::TraverseRepresentationItem(A3DRiRepresentationItem*
 	case kA3DTypeRiSet:
 		return TraverseRepresentationSet(RepresentationItem, PartMetaData, ParentUnit);
 	case kA3DTypeRiBrepModel:
-		return TraverseBRepModel(RepresentationItem, PartMetaData, ParentId, ParentUnit);
+		return TraverseBRepModel(RepresentationItem, PartMetaData, ParentId, ParentUnit, ItemIndex);
 	case kA3DTypeRiPolyBrepModel:
-		return TraversePolyBRepModel(RepresentationItem, PartMetaData, ParentId, ParentUnit);
+		return TraversePolyBRepModel(RepresentationItem, PartMetaData, ParentId, ParentUnit, ItemIndex);
 	default:
 		break;
 	}
@@ -1187,7 +1302,7 @@ FCadId FTechSoftFileParser::TraverseRepresentationSet(const A3DRiSet* Representa
 
 	for (A3DUns32 Index = 0; Index < RepresentationSetData->m_uiRepItemsSize; ++Index)
 	{
-		int32 ChildId = TraverseRepresentationItem(RepresentationSetData->m_ppRepItems[Index], RepresentationSetMetaData, RepresentationSet.ObjectId, ParentUnit);
+		int32 ChildId = TraverseRepresentationItem(RepresentationSetData->m_ppRepItems[Index], RepresentationSetMetaData, RepresentationSet.ObjectId, ParentUnit, Index);
 		RepresentationSet.Children.Add(ChildId);
 	}
 	return RepresentationSetId;
@@ -1208,7 +1323,7 @@ void FTechSoftFileParser::CountUnderRepresentationSet(const A3DRiSet* Representa
 	}
 }
 
-FCadId FTechSoftFileParser::TraverseBRepModel(A3DRiBrepModel* BRepModelPtr, const FEntityMetaData& PartMetaData, const FCadId ParentId, double ParentUnit)
+FCadId FTechSoftFileParser::TraverseBRepModel(A3DRiBrepModel* BRepModelPtr, const FEntityMetaData& PartMetaData, const FCadId ParentId, double ParentUnit, int32 ItemIndex)
 {
 	if (!BRepModelPtr)
 	{
@@ -1224,14 +1339,16 @@ FCadId FTechSoftFileParser::TraverseBRepModel(A3DRiBrepModel* BRepModelPtr, cons
 	}
 
 	TUniqueTSObj<A3DRiBrepModelData> BRepModelData(BRepModelPtr);
-	BuildBodyName(BRepMetaData, PartMetaData, (bool) BRepModelData->m_bSolid);
+	BuildBodyName(BRepMetaData, PartMetaData, ItemIndex, (bool) BRepModelData->m_bSolid);
 
 	if (int32* BodyIndexPtr = RepresentationItemsCache.Find(BRepModelPtr))
 	{
 		return CADFileData.GetBodyAt(*BodyIndexPtr).ObjectId;
 	}
 
-	// if BRep model has not material or color, add part one
+	// if BRep model has not material or color, use the ones from the part
+	// A top/down propagation of the color and material is needed. This is done in 5.0.3 with minimal code modification. 
+	// However the model parsing need to be rewrite in the next version. (Jira UE-152691)
 	if (BRepMetaData.MaterialName == 0 && BRepMetaData.ColorName == 0)
 	{
 		if (PartMetaData.MaterialName)
@@ -1264,7 +1381,7 @@ FCadId FTechSoftFileParser::TraverseBRepModel(A3DRiBrepModel* BRepModelPtr, cons
 	return Body.ObjectId;
 }
 
-FCadId FTechSoftFileParser::TraversePolyBRepModel(A3DRiPolyBrepModel* PolygonalPtr, const FEntityMetaData& PartMetaData, const FCadId ParentId, double ParentUnit)
+FCadId FTechSoftFileParser::TraversePolyBRepModel(A3DRiPolyBrepModel* PolygonalPtr, const FEntityMetaData& PartMetaData, const FCadId ParentId, double ParentUnit, int32 ItemIndex)
 {
 	if (!PolygonalPtr)
 	{
@@ -1280,13 +1397,28 @@ FCadId FTechSoftFileParser::TraversePolyBRepModel(A3DRiPolyBrepModel* PolygonalP
 	}
 
 	TUniqueTSObj<A3DRiPolyBrepModelData> BRepModelData(PolygonalPtr);
-	BuildBodyName(BRepMetaData, PartMetaData, (bool) BRepModelData->m_bIsClosed);
+	BuildBodyName(BRepMetaData, PartMetaData, ItemIndex, (bool) BRepModelData->m_bIsClosed);
 
 	if (int32* BodyIndexPtr = RepresentationItemsCache.Find(PolygonalPtr))
 	{
 		return CADFileData.GetBodyAt(*BodyIndexPtr).ObjectId;
 	}
 
+	// if BRep model has not material or color, add part one
+	if (BRepMetaData.MaterialName == 0 && BRepMetaData.ColorName == 0)
+	{
+		if (PartMetaData.MaterialName)
+		{
+			BRepMetaData.MaterialName = PartMetaData.MaterialName;
+			BRepMetaData.MetaData.Add(TEXT("MaterialName"), FString::Printf(TEXT("%u"), PartMetaData.MaterialName));
+		}
+		if (PartMetaData.ColorName)
+		{
+			BRepMetaData.ColorName = PartMetaData.ColorName;
+			BRepMetaData.MetaData.Add(TEXT("ColorName"), FString::Printf(TEXT("%u"), PartMetaData.ColorName));
+		}
+	}
+
 	ExtractSpecificMetaData(PolygonalPtr, BRepMetaData);
 
 	FMatrix Matrix = FMatrix::Identity;
@@ -1313,8 +1445,11 @@ void FTechSoftFileParser::ExtractMetaData(const A3DEntity* Entity, FEntityMetaDa
 		if (MetaData->m_pcName && MetaData->m_pcName[0] != '\0')
 		{
 			FString SDKName = UTF8_TO_TCHAR(MetaData->m_pcName);
-			SDKName = TechSoftUtils::CleanSdkName(SDKName);
-			OutMetaData.MetaData.Emplace(TEXT("SDKName"), SDKName);
+			if(SDKName != TEXT("unnamed"))  // "unnamed" is create by Techsoft. This name is ignored 
+			{
+				SDKName = TechSoftUtils::CleanSdkName(SDKName);
+				OutMetaData.MetaData.Emplace(TEXT("SDKName"), SDKName);
+			}
 		}
 
 		TUniqueTSObj<A3DMiscAttributeData> AttributeData;
@@ -1398,13 +1533,12 @@ void FTechSoftFileParser::BuildReferenceName(FEntityMetaData& ReferenceData)
 			break;
 		}
 
-		FString& Name = MetaData.FindOrAdd(TEXT("Name"));
-		Name = SdkName;
+		MetaData.FindOrAdd(TEXT("Name"), SdkName);
 		return;
 	}
 }
 
-void FTechSoftFileParser::BuildInstanceName(FEntityMetaData& InstanceData)
+void FTechSoftFileParser::BuildInstanceName(FEntityMetaData& InstanceData, const FString& DefaultName)
 {
 	TMap<FString, FString>& MetaData = InstanceData.MetaData;
 
@@ -1426,9 +1560,12 @@ void FTechSoftFileParser::BuildInstanceName(FEntityMetaData& InstanceData)
 		switch (Format)
 		{
 		case ECADFormat::CATIA:
-		case ECADFormat::CATIA_3DXML:
 			SdkName = TechSoftUtils::CleanCatiaInstanceSdkName(SdkName);
 			break;
+		
+		case ECADFormat::CATIA_3DXML:
+			SdkName = TechSoftUtils::Clean3dxmlInstanceSdkName(SdkName);
+			break;
 
 		case ECADFormat::SOLIDWORKS:
 			SdkName = TechSoftUtils::CleanSwInstanceSdkName(SdkName);
@@ -1438,8 +1575,7 @@ void FTechSoftFileParser::BuildInstanceName(FEntityMetaData& InstanceData)
 			break;
 		}
 
-		FString& Name = MetaData.FindOrAdd(TEXT("Name"));
-		Name = SdkName;
+		MetaData.FindOrAdd(TEXT("Name"), SdkName);
 		return;
 	}
 
@@ -1451,35 +1587,41 @@ void FTechSoftFileParser::BuildInstanceName(FEntityMetaData& InstanceData)
 	FString& Name = MetaData.FindOrAdd(TEXT("Name"));
 	if (Name.IsEmpty())
 	{
-		Name = TEXT("Product");
+		Name = DefaultName;
 	}
 }
 
-void FTechSoftFileParser::BuildPartName(FEntityMetaData& PartData)
+void FTechSoftFileParser::BuildPartName(FEntityMetaData& PartData, const FArchiveComponent& Component)
 {
 	TMap<FString, FString>& MetaData = PartData.MetaData;
-	if (MetaData.IsEmpty())
+	if (!MetaData.IsEmpty())
 	{
-		return;
-	}
+		if (TechSoftFileParserImpl::CheckIfNameExists(MetaData))
+		{
+			return;
+		}
 
-	if (TechSoftFileParserImpl::CheckIfNameExists(MetaData))
-	{
-		return;
-	}
+		if (TechSoftFileParserImpl::ReplaceOrAddNameValue(MetaData, TEXT("PartNumber")))
+		{
+			return;
+		}
 
-	if (TechSoftFileParserImpl::ReplaceOrAddNameValue(MetaData, TEXT("PartNumber")))
-	{
-		return;
+		if (TechSoftFileParserImpl::ReplaceOrAddNameValue(MetaData, TEXT("SDKName")))
+		{
+			return;
+		}
 	}
 
-	if (TechSoftFileParserImpl::ReplaceOrAddNameValue(MetaData, TEXT("SDKName")))
+	// If a name hasn't been found for the Part, the name of its parent (component) is used
+	const FString* ComponentName = Component.MetaData.Find(TEXT("Name"));
+	if (ComponentName && !ComponentName->IsEmpty())
 	{
+		MetaData.Add(TEXT("Name"), *ComponentName);
 		return;
 	}
 }
 
-void FTechSoftFileParser::BuildBodyName(FEntityMetaData& BodyData, const FEntityMetaData& PartMetaData, bool bIsSolid)
+void FTechSoftFileParser::BuildBodyName(FEntityMetaData& BodyData, const FEntityMetaData& PartMetaData, int32 ItemIndex, bool bIsSolid)
 {
 	TMap<FString, FString>& MetaData = BodyData.MetaData;
 
@@ -1497,8 +1639,7 @@ void FTechSoftFileParser::BuildBodyName(FEntityMetaData& BodyData, const FEntity
 			SdkName = TechSoftUtils::CleanNameByRemoving_prt(SdkName);
 		}
 
-		FString& Name = MetaData.FindOrAdd(TEXT("Name"));
-		Name = SdkName;
+		MetaData.FindOrAdd(TEXT("Name"), SdkName);
 		return;
 	}
 
@@ -1507,21 +1648,24 @@ void FTechSoftFileParser::BuildBodyName(FEntityMetaData& BodyData, const FEntity
 		NamePtr = MetaData.Find(TEXT("BodyID"));
 		if (NamePtr != nullptr && !NamePtr->IsEmpty())
 		{
-			FString& Name = MetaData.FindOrAdd(TEXT("Name"));
-			Name = *NamePtr;
+			MetaData.FindOrAdd(TEXT("Name"), *NamePtr);
 			return;
 		}
 	}
 
-	const FString* PartName = PartMetaData.MetaData.Find(TEXT("Name"));
-
 	FString& Name = MetaData.FindOrAdd(TEXT("Name"));
 	if (Name.IsEmpty())
 	{
+		const FString* PartName = PartMetaData.MetaData.Find(TEXT("Name"));
 		if (PartName != nullptr && !PartName->IsEmpty())
 		{
 			Name = *PartName;
 		}
+
+		if (!Name.IsEmpty())
+		{
+			Name += TEXT("_body");
+		}
 		else
 		{
 			if (bIsSolid)
@@ -1533,6 +1677,7 @@ void FTechSoftFileParser::BuildBodyName(FEntityMetaData& BodyData, const FEntity
 				Name = TEXT("Shell");
 			}
 		}
+		Name += FString::FromInt(ItemIndex + 1);
 	}
 }
 
@@ -1781,7 +1926,7 @@ FMatrix FTechSoftFileParser::ExtractTransformation3D(const A3DMiscTransformation
 
 		InOutUnit *= UniformScale;
 
-		FMatrix Matrix(XVector, YVector, ZVector, FVector::Zero());
+		FMatrix Matrix(XVector, YVector, ZVector, FVector::ZeroVector);
 
 		if (CartesianTransformationData->m_ucBehaviour & kA3DTransformationMirror)
 		{
diff --git a/Engine/Plugins/Enterprise/DatasmithCADImporter/Source/CADInterfaces/Private/TechSoftFileParser.h b/Engine/Plugins/Enterprise/DatasmithCADImporter/Source/CADInterfaces/Private/TechSoftFileParser.h
--- a/Engine/Plugins/Enterprise/DatasmithCADImporter/Source/CADInterfaces/Private/TechSoftFileParser.h
+++ b/Engine/Plugins/Enterprise/DatasmithCADImporter/Source/CADInterfaces/Private/TechSoftFileParser.h
@@ -33,7 +33,7 @@ enum EComponentType : uint32
 	LastType
 };
 
-class CADINTERFACES_API FTechSoftFileParser : public ICADFileParser
+class FTechSoftFileParser : public ICADFileParser
 {
 public:
 
@@ -97,24 +97,30 @@ private:
 
 	// Traverse ASM tree by starting from the model
 	ECADParsingResult TraverseModel();
-	void TraverseReference(const A3DAsmProductOccurrence* Reference, const FMatrix& ParentMatrix, double ParentUnit);
+
+	// Note:
+	// Due to dataprep purpose, node without name cannot have a generic name (e.g. Product, Part, Body, Shell, etc...) but must have a name based on its parent's name.
+	// To be able to build a name, the name of the parent has to be known.
+	// The implementation of the naming policy is done in 5.0.3 with minimal code modification. However the model parsing need to be rewrite in the next version. (Jira UE-152624)
+
+	void TraverseReference(const A3DAsmProductOccurrence* Reference, const FString& RootName, const FMatrix& ParentMatrix, double ParentUnit);
 	bool IsConfigurationSet(const A3DAsmProductOccurrence* Occurrence);
 	void TraverseConfigurationSet(const A3DAsmProductOccurrence* ConfigurationSet, double ParentUnit);
-	FCadId TraverseOccurrence(const A3DAsmProductOccurrence* Occurrence, double ParentUnit);
+	FCadId TraverseOccurrence(const A3DAsmProductOccurrence* Occurrence, const FString& DefaultOccurrenceName, double ParentUnit);
 	void ProcessPrototype(const A3DAsmProductOccurrence* InPrototype, FEntityMetaData& OutMetaData, A3DMiscTransformation** OutLocation);
 	void TraversePartDefinition(const A3DAsmPartDefinition* PartDefinition, FArchiveComponent& Component, double ParentUnit);
 	FCadId TraverseRepresentationSet(const A3DRiSet* pSet, const FEntityMetaData& PartMetaData, double ParentUnit);
-	FCadId TraverseRepresentationItem(A3DRiRepresentationItem* RepresentationItem, const FEntityMetaData& PartMetaData, const FCadId ParentId, double ParentUnit);
-	FCadId TraverseBRepModel(A3DRiBrepModel* BrepModel, const FEntityMetaData& PartMetaData, const FCadId ParentId, double ParentUnit);
-	FCadId TraversePolyBRepModel(A3DRiPolyBrepModel* PolygonalBrepModel, const FEntityMetaData& PartMetaData, const FCadId ParentId, double ParentUnit);
+	FCadId TraverseRepresentationItem(A3DRiRepresentationItem* RepresentationItem, const FEntityMetaData& PartMetaData, const FCadId ParentId, double ParentUnit, int32 ItemIndex);
+	FCadId TraverseBRepModel(A3DRiBrepModel* BrepModel, const FEntityMetaData& PartMetaData, const FCadId ParentId, double ParentUnit, int32 ItemIndex);
+	FCadId TraversePolyBRepModel(A3DRiPolyBrepModel* PolygonalBrepModel, const FEntityMetaData& PartMetaData, const FCadId ParentId, double ParentUnit, int32 ItemIndex);
 
 	// MetaData
 	void ExtractSpecificMetaData(const A3DAsmProductOccurrence* Occurrence, FEntityMetaData& OutMetaData);
 
-	void BuildInstanceName(FEntityMetaData& MetaData);
+	void BuildInstanceName(FEntityMetaData& MetaData, const FString& DefaultInstanceName);
 	void BuildReferenceName(FEntityMetaData& MetaData);
-	void BuildPartName(FEntityMetaData& MetaData);
-	void BuildBodyName(FEntityMetaData& MetaData, const FEntityMetaData& PartMetaData, bool bIsSolid);
+	void BuildPartName(FEntityMetaData& MetaData, const FArchiveComponent& Component);
+	void BuildBodyName(FEntityMetaData& MetaData, const FEntityMetaData& PartMetaData, int32 ItemIndex, bool bIsSolid);
 
 	// Graphic properties
 	void ExtractGraphicProperties(const A3DGraphics* Graphics, FEntityMetaData& OutMetaData);
diff --git a/Engine/Plugins/Enterprise/DatasmithCADImporter/Source/CADInterfaces/Private/TechSoftUtils.cpp b/Engine/Plugins/Enterprise/DatasmithCADImporter/Source/CADInterfaces/Private/TechSoftUtils.cpp
--- a/Engine/Plugins/Enterprise/DatasmithCADImporter/Source/CADInterfaces/Private/TechSoftUtils.cpp
+++ b/Engine/Plugins/Enterprise/DatasmithCADImporter/Source/CADInterfaces/Private/TechSoftUtils.cpp
@@ -282,7 +282,8 @@ void GetMaterialValues(uint32 StyleIndex, FCADUUID& OutColorName, FCADUUID& OutM
 			TUniqueTSObjFromIndex<A3DGraphRgbColorData> ColorData(GraphStyleData->m_uiRgbColorIndex);
 			if (ColorData.IsValid())
 			{
-				const uint8 Alpha = GraphStyleData->m_bIsTransparencyDefined ? GraphStyleData->m_ucTransparency : 255;
+				// Alpha == Opacity == (255 - Transparency)
+				const uint8 Alpha = GraphStyleData->m_bIsTransparencyDefined ? (255 - GraphStyleData->m_ucTransparency) : 255;
 				const FColor ColorValue((uint8)(ColorData->m_dRed * 255), (uint8)(ColorData->m_dGreen * 255), (uint8)(ColorData->m_dBlue * 255), Alpha);
 
 				OutColorName = BuildColorName(ColorValue);
@@ -687,6 +688,19 @@ FString CleanCatiaInstanceSdkName(const FString& Name)
 	return Name;
 }
 
+FString Clean3dxmlInstanceSdkName(const FString& Name)
+{
+	FString NewName = CleanCatiaInstanceSdkName(Name);
+
+	int32 Index = NewName.Find(TEXT("_InstanceRep"));
+	if (Index != INDEX_NONE)
+	{
+		NewName = NewName.Left(Index);
+	}
+
+	return NewName;
+}
+
 FString Clean3dxmlReferenceSdkName(const FString& Name)
 {
 	int32 Index;
@@ -695,6 +709,14 @@ FString Clean3dxmlReferenceSdkName(const FString& Name)
 		FString NewName = Name.Left(Index);
 		return NewName;
 	}
+
+	Index = Name.Find(TEXT("_InstanceRep"));
+	if (Index != INDEX_NONE)
+	{
+		FString NewName = Name.Left(Index);
+		return NewName;
+	}
+
 	return Name;
 }
 
diff --git a/Engine/Plugins/Enterprise/DatasmithCADImporter/Source/CADInterfaces/Private/TechSoftUtilsPrivate.h b/Engine/Plugins/Enterprise/DatasmithCADImporter/Source/CADInterfaces/Private/TechSoftUtilsPrivate.h
--- a/Engine/Plugins/Enterprise/DatasmithCADImporter/Source/CADInterfaces/Private/TechSoftUtilsPrivate.h
+++ b/Engine/Plugins/Enterprise/DatasmithCADImporter/Source/CADInterfaces/Private/TechSoftUtilsPrivate.h
@@ -17,6 +17,7 @@ void ExtractAttribute(const A3DMiscAttributeData& AttributeData, TMap<FString, F
 
 FString CleanSdkName(const FString& Name);
 FString CleanCatiaInstanceSdkName(const FString& Name);
+FString Clean3dxmlInstanceSdkName(const FString& Name);
 FString CleanCatiaReferenceName(const FString& Name);
 FString Clean3dxmlReferenceSdkName(const FString& Name);
 FString CleanSwInstanceSdkName(const FString& Name);
diff --git a/Engine/Plugins/Enterprise/DatasmithCADImporter/Source/DatasmithCADTranslator/Private/DatasmithCADTranslator.cpp b/Engine/Plugins/Enterprise/DatasmithCADImporter/Source/DatasmithCADTranslator/Private/DatasmithCADTranslator.cpp
--- a/Engine/Plugins/Enterprise/DatasmithCADImporter/Source/DatasmithCADTranslator/Private/DatasmithCADTranslator.cpp
+++ b/Engine/Plugins/Enterprise/DatasmithCADImporter/Source/DatasmithCADTranslator/Private/DatasmithCADTranslator.cpp
@@ -122,6 +122,12 @@ bool FDatasmithCADTranslator::LoadScene(TSharedRef<IDatasmithScene> DatasmithSce
 	}
 
 	case CADLibrary::ECADFormat::SOLIDWORKS:
+	{
+		ImportParameters.SetModelCoordinateSystem(FDatasmithUtils::EModelCoordSystem::YUp_RightHanded);
+		ImportParameters.SetDisplayPreference(CADLibrary::EDisplayPreference::ColorOnly);
+		break;
+	}
+
 	case CADLibrary::ECADFormat::INVENTOR:
 	case CADLibrary::ECADFormat::CREO:
 	{
diff --git a/Engine/Plugins/Enterprise/DatasmithCADImporter/Source/DatasmithCADTranslator/Private/DatasmithSceneGraphBuilder.cpp b/Engine/Plugins/Enterprise/DatasmithCADImporter/Source/DatasmithCADTranslator/Private/DatasmithSceneGraphBuilder.cpp
--- a/Engine/Plugins/Enterprise/DatasmithCADImporter/Source/DatasmithCADTranslator/Private/DatasmithSceneGraphBuilder.cpp
+++ b/Engine/Plugins/Enterprise/DatasmithCADImporter/Source/DatasmithCADTranslator/Private/DatasmithSceneGraphBuilder.cpp
@@ -86,8 +86,9 @@ bool FDatasmithSceneGraphBuilder::Build()
 {
 	LoadSceneGraphDescriptionFiles();
 
-	uint32 RootHash = 0;
-	SceneGraph = FindSceneGraphArchive(RootFileDescription, RootHash);
+	uint32 RootHash = RootFileDescription.GetDescriptorHash();
+	SceneGraph = CADFileToSceneGraphArchive.FindRef(RootHash);
+
 	if (!SceneGraph)
 	{
 		return false;
@@ -125,17 +126,13 @@ void FDatasmithSceneGraphBuilder::LoadSceneGraphDescriptionFiles()
 	}
 }
 
-CADLibrary::FArchiveSceneGraph* FDatasmithSceneBaseGraphBuilder::FindSceneGraphArchive(const CADLibrary::FFileDescriptor& File, uint32& FileHash) const
-{
-	FileHash = GetTypeHash(File.GetFileName());
-	return CADFileToSceneGraphArchive.FindRef(FileHash);
-}
-
 void FDatasmithSceneGraphBuilder::FillAnchorActor(const TSharedRef<IDatasmithActorElement>& ActorElement, const FString& CleanFilenameOfCADFile)
 {
 	CADLibrary::FFileDescriptor AnchorDescription(*CleanFilenameOfCADFile);
-	uint32 AnchorHash = 0;
-	SceneGraph = FindSceneGraphArchive(AnchorDescription, AnchorHash);
+
+	uint32 AnchorHash = AnchorDescription.GetDescriptorHash();
+	SceneGraph = CADFileToSceneGraphArchive.FindRef(AnchorHash);
+
 	if (!SceneGraph)
 	{
 		return;
@@ -300,8 +297,9 @@ TSharedPtr< IDatasmithActorElement >  FDatasmithSceneBaseGraphBuilder::BuildInst
 	{
 		if (!Instance.ExternalReference.GetSourcePath().IsEmpty())
 		{
-			uint32 InstanceSceneGraphHash = 0;
-			SceneGraph = FindSceneGraphArchive(Instance.ExternalReference, InstanceSceneGraphHash);
+			uint32 InstanceSceneGraphHash = Instance.ExternalReference.GetDescriptorHash();
+			SceneGraph = CADFileToSceneGraphArchive.FindRef(InstanceSceneGraphHash);
+
 			if (SceneGraph)
 			{
 				if (AncestorSceneGraphHash.Find(InstanceSceneGraphHash) == INDEX_NONE)
@@ -641,6 +639,8 @@ void FDatasmithSceneBaseGraphBuilder::AddMetaData(TSharedPtr< IDatasmithActorEle
 	{
 		TSet<FString> UnwantedAttributes;
 
+		UnwantedAttributes.Add(TEXT("SDKName"));
+
 		// CoreTech
 		UnwantedAttributes.Add(TEXT("CTName"));
 		UnwantedAttributes.Add(TEXT("LayerId"));
diff --git a/Engine/Plugins/Enterprise/DatasmithCADImporter/Source/DatasmithDispatcher/Private/DatasmithDispatcher.cpp b/Engine/Plugins/Enterprise/DatasmithCADImporter/Source/DatasmithDispatcher/Private/DatasmithDispatcher.cpp
--- a/Engine/Plugins/Enterprise/DatasmithCADImporter/Source/DatasmithDispatcher/Private/DatasmithDispatcher.cpp
+++ b/Engine/Plugins/Enterprise/DatasmithCADImporter/Source/DatasmithDispatcher/Private/DatasmithDispatcher.cpp
@@ -203,7 +203,7 @@ void FDatasmithDispatcher::LinkCTFileToUnrealCacheFile(const CADLibrary::FFileDe
 {
 	FScopeLock Lock(&TaskPoolCriticalSection);
 
-	uint32 FileHash = GetTypeHash(CTFileDescription.GetFileName());
+	uint32 FileHash = CTFileDescription.GetDescriptorHash();
 	if (!UnrealSceneGraphFile.IsEmpty())
 	{
 		CADFileToUnrealFileMap.Add(FileHash, UnrealSceneGraphFile);
diff --git a/Engine/Plugins/Enterprise/DatasmithCADImporter/Source/DatasmithOpenNurbsTranslator/Private/DatasmithOpenNurbsTranslator.cpp b/Engine/Plugins/Enterprise/DatasmithCADImporter/Source/DatasmithOpenNurbsTranslator/Private/DatasmithOpenNurbsTranslator.cpp
--- a/Engine/Plugins/Enterprise/DatasmithCADImporter/Source/DatasmithOpenNurbsTranslator/Private/DatasmithOpenNurbsTranslator.cpp
+++ b/Engine/Plugins/Enterprise/DatasmithCADImporter/Source/DatasmithOpenNurbsTranslator/Private/DatasmithOpenNurbsTranslator.cpp
@@ -3002,7 +3002,13 @@ bool FOpenNurbsTranslatorImpl::TranslateBRep(ON_Brep* Brep, const ON_3dmObjectAt
 		CADModelConverter->SaveModel(*OutputPath, MeshElement);
 
 		CADLibrary::FMeshParameters MeshParameters;
-		return CADModelConverter->Tessellate(MeshParameters, OutMesh);
+		bool bRet = CADModelConverter->Tessellate(MeshParameters, OutMesh);
+		if (!bRet)
+		{
+			FString StaticMeshLable = MeshElement->GetLabel();
+			UE_LOG(LogDatasmithOpenNurbsTranslator, Warning, TEXT("Failed to generate the mesh of \"%s\" (%s) StaticMesh."), *StaticMeshLable, *Name);
+		}
+		return bRet;
 	}
 	else
 	{
diff --git a/Engine/Plugins/Enterprise/DatasmithCADImporter/Source/DatasmithOpenNurbsTranslator/Private/OpenNurbsBRepToTechSoftConverter.cpp b/Engine/Plugins/Enterprise/DatasmithCADImporter/Source/DatasmithOpenNurbsTranslator/Private/OpenNurbsBRepToTechSoftConverter.cpp
--- a/Engine/Plugins/Enterprise/DatasmithCADImporter/Source/DatasmithOpenNurbsTranslator/Private/OpenNurbsBRepToTechSoftConverter.cpp
+++ b/Engine/Plugins/Enterprise/DatasmithCADImporter/Source/DatasmithOpenNurbsTranslator/Private/OpenNurbsBRepToTechSoftConverter.cpp
@@ -383,8 +383,10 @@ A3DTopoFace* FOpenNurbsBRepToTechSoftConverter::CreateTopoFace(const ON_BrepFace
 	Face->m_dTolerance = 0.01; //mm
 
 	A3DTopoFace* FacePtr = CADLibrary::TechSoftInterface::CreateTopoFace(*Face);
-
-	CADLibrary::TechSoftUtils::SetEntityGraphicsColor(FacePtr, FColor(200, 200, 200));
+	
+	// Color with color name = 0. c.f. CADLibrary::BuildColorName 
+	constexpr FColor DefaultColor(0, 0, 0, 0);
+	CADLibrary::TechSoftUtils::SetEntityGraphicsColor(FacePtr, DefaultColor);
 
 	return FacePtr;
 }
diff --git a/Engine/Plugins/Enterprise/DatasmithCADImporter/Source/DatasmithWireTranslator/Private/DatasmithWireTranslator.cpp b/Engine/Plugins/Enterprise/DatasmithCADImporter/Source/DatasmithWireTranslator/Private/DatasmithWireTranslator.cpp
--- a/Engine/Plugins/Enterprise/DatasmithCADImporter/Source/DatasmithWireTranslator/Private/DatasmithWireTranslator.cpp
+++ b/Engine/Plugins/Enterprise/DatasmithCADImporter/Source/DatasmithWireTranslator/Private/DatasmithWireTranslator.cpp
@@ -72,13 +72,6 @@ namespace UE_DATASMITHWIRETRANSLATOR_NAMESPACE
 {
 
 static bool bGAliasSewByColor = false;
-FAutoConsoleVariableRef GAliasSewByColor(
-	TEXT("ds.CADTranslator.Alias.SewByColor"),
-	bGAliasSewByColor,
-	TEXT("Enable Sew action merges BReps according to their material i.e. only BReps associated with same material can be merged together.\
-Default is disable\n"),
-	ECVF_Default);
-
 
 static const FColor DefaultColor = FColor(200, 200, 200);
 
@@ -1428,7 +1421,7 @@ bool FWireTranslatorImpl::ProcessAlGroupNode(AlDagNode& GroupNode, const FDagNod
 	TSharedPtr<AlDagNode> ChildNode(AlChildPtr);
 
 	FDagNodeInfo ThisGroupNodeInfo;
-	GetDagNodeInfo(*ChildNode, ParentInfo, ThisGroupNodeInfo);
+	GetDagNodeInfo(GroupNode, ParentInfo, ThisGroupNodeInfo);
 
 	ThisGroupNodeInfo.ActorElement = FDatasmithSceneFactory::CreateActor(*OpenModelUtils::UuidToString(ThisGroupNodeInfo.Uuid));
 	ThisGroupNodeInfo.ActorElement->SetLabel(*ThisGroupNodeInfo.Label);
@@ -1955,7 +1948,13 @@ TOptional<FMeshDescription> FWireTranslatorImpl::MeshDagNodeWithExternalMesher(A
 	FMeshDescription MeshDescription;
 	DatasmithMeshHelper::PrepareAttributeForStaticMesh(MeshDescription);
 
-	CADModelConverter->Tessellate(MeshParameters, MeshDescription);
+	bool bRet = CADModelConverter->Tessellate(MeshParameters, MeshDescription);
+	if (!bRet)
+	{
+		FString StaticMeshLable = MeshElement->GetLabel();
+		FString StaticMeshName = MeshElement->GetName();
+		UE_LOG(LogDatasmithWireTranslator, Warning, TEXT("Failed to generate the mesh of \"%s\" (%s) StaticMesh."), *StaticMeshLable, *StaticMeshName);
+	}
 
 	return MoveTemp(MeshDescription);
 }
diff --git a/Engine/Plugins/Enterprise/DatasmithCADImporter/Source/ParametricSurface/Private/TechSoft/CADModelToTechSoftConverterBase.cpp b/Engine/Plugins/Enterprise/DatasmithCADImporter/Source/ParametricSurface/Private/TechSoft/CADModelToTechSoftConverterBase.cpp
--- a/Engine/Plugins/Enterprise/DatasmithCADImporter/Source/ParametricSurface/Private/TechSoft/CADModelToTechSoftConverterBase.cpp
+++ b/Engine/Plugins/Enterprise/DatasmithCADImporter/Source/ParametricSurface/Private/TechSoft/CADModelToTechSoftConverterBase.cpp
@@ -89,12 +89,7 @@ bool FCADModelToTechSoftConverterBase::Tessellate(const CADLibrary::FMeshParamet
 		return false;
 	}
 
-	if (!CADLibrary::ConvertBodyMeshToMeshDescription(ImportParameters, InMeshParameters, BodyMesh, OutMeshDescription))
-	{
-		ensureMsgf(false, TEXT("Error during mesh conversion"));
-		return false;
-	}
-	return true;
+	return CADLibrary::ConvertBodyMeshToMeshDescription(ImportParameters, InMeshParameters, BodyMesh, OutMeshDescription);
 }
 
 void FCADModelToTechSoftConverterBase::AddSurfaceDataForMesh(const TCHAR* InFilePath, const CADLibrary::FMeshParameters& InMeshParameters, const FDatasmithTessellationOptions& InTessellationOptions, FDatasmithMeshElementPayload& OutMeshPayload) const
