Remove debug compilation definitions

diff --git a/Engine/Plugins/Enterprise/DatasmithCADImporter/Source/CoreTechSurface/CoreTechSurface.Build.cs b/Engine/Plugins/Enterprise/DatasmithCADImporter/Source/CoreTechSurface/CoreTechSurface.Build.cs
--- a/Engine/Plugins/Enterprise/DatasmithCADImporter/Source/CoreTechSurface/CoreTechSurface.Build.cs
+++ b/Engine/Plugins/Enterprise/DatasmithCADImporter/Source/CoreTechSurface/CoreTechSurface.Build.cs
@@ -8,11 +8,6 @@ namespace UnrealBuildTool.Rules
 	{
 		public CoreTechSurface(ReadOnlyTargetRules Target) : base(Target)
 		{
-			// option de compil pour dev
-			OptimizeCode = CodeOptimization.Never;
-			PCHUsage = PCHUsageMode.NoPCHs;
-			// fin option de compil pour dev
-
 			PublicDependencyModuleNames.AddRange(
 				new string[]
 				{
