Enable CADkernel

diff --git a/Engine/Plugins/Enterprise/DatasmithCADImporter/DatasmithCADImporter.uplugin b/Engine/Plugins/Enterprise/DatasmithCADImporter/DatasmithCADImporter.uplugin
--- a/Engine/Plugins/Enterprise/DatasmithCADImporter/DatasmithCADImporter.uplugin
+++ b/Engine/Plugins/Enterprise/DatasmithCADImporter/DatasmithCADImporter.uplugin
@@ -109,6 +109,14 @@
 				"Win64"
 			]
 		},
+		{
+			"Name": "DatasmithWireTranslator2023_0",
+			"Type": "Editor",
+			"LoadingPhase": "Default",
+			"PlatformAllowList": [
+				"Win64"
+			]
+		},
 		{
 			"Name": "ParametricSurfaceExtension",
 			"Type": "Editor",
diff --git a/Engine/Plugins/Enterprise/DatasmithCADImporter/Source/CADTools/Private/CADOption.cpp b/Engine/Plugins/Enterprise/DatasmithCADImporter/Source/CADTools/Private/CADOption.cpp
--- a/Engine/Plugins/Enterprise/DatasmithCADImporter/Source/CADTools/Private/CADOption.cpp
+++ b/Engine/Plugins/Enterprise/DatasmithCADImporter/Source/CADTools/Private/CADOption.cpp
@@ -27,13 +27,11 @@ Default is MaxImportThreads = 0\n\
 	ECVF_Default);
 
 bool FImportParameters::bGDisableCADKernelTessellation = true;
-#ifdef ENABLE_IN_5_0_1
 FAutoConsoleVariableRef GCADTranslatorDisableCADKernelTessellation(
 	TEXT("ds.CADTranslator.DisableCADKernelTessellation"),
 	FImportParameters::bGDisableCADKernelTessellation,
 	TEXT("Disable to use CAD import library tessellator.\n"),
 	ECVF_Default);
-#endif
 
 bool FImportParameters::bGEnableCADCache = true;
 FAutoConsoleVariableRef GCADTranslatorEnableCADCache(
@@ -74,6 +72,14 @@ FAutoConsoleVariableRef GCADTranslatorStitchingTolerance(
 Default value of StitchingTolerance is 0.001 cm\n"),
 ECVF_Default);
 
+static bool bGAliasSewByColor = false;
+FAutoConsoleVariableRef GAliasSewByColor(
+	TEXT("ds.CADTranslator.Alias.SewByColor"),
+	bGAliasSewByColor,
+	TEXT("Enable Sew action merges BReps according to their material i.e. only BReps associated with same material can be merged together.\
+Default is disable\n"),
+ECVF_Default);
+
 uint32 GetTypeHash(const FImportParameters& ImportParameters)
 {
 	uint32 ParametersHash = ::GetTypeHash(ImportParameters.bGDisableCADKernelTessellation);
diff --git a/Engine/Plugins/Enterprise/DatasmithCADImporter/Source/DatasmithWireTranslator/Private/DatasmithWireTranslator.cpp b/Engine/Plugins/Enterprise/DatasmithCADImporter/Source/DatasmithWireTranslator/Private/DatasmithWireTranslator.cpp
--- a/Engine/Plugins/Enterprise/DatasmithCADImporter/Source/DatasmithWireTranslator/Private/DatasmithWireTranslator.cpp
+++ b/Engine/Plugins/Enterprise/DatasmithCADImporter/Source/DatasmithWireTranslator/Private/DatasmithWireTranslator.cpp
@@ -77,7 +77,8 @@ static const FColor DefaultColor = FColor(200, 200, 200);
 
 #ifdef USE_OPENMODEL
 
-const uint64 LibAliasNext_Version = 17881307937833405;
+const uint64 LibAliasNext_Version = 0xffffffffffffffffull;
+const uint64 LibAlias2023_0_0_Version = 8162774324609149;
 const uint64 LibAlias2022_2_0_Version = 7881307937833405;
 const uint64 LibAlias2022_1_0_Version = 7881303642865885;
 const uint64 LibAlias2022_0_1_Version = 7881299347964005;
@@ -106,8 +107,12 @@ const uint64 LibAliasVersionMax = LibAlias2022_2_0_Version;
 const FString AliasVersionChar = TEXT("AliasStudio 2022.1, Model files");
 #elif defined(OPEN_MODEL_2022_2)
 const uint64 LibAliasVersionMin = LibAlias2022_2_0_Version;
-const uint64 LibAliasVersionMax = LibAliasNext_Version;
+const uint64 LibAliasVersionMax = LibAlias2023_0_0_Version;
 const FString AliasVersionChar = TEXT("AliasStudio 2022.2, Model files");
+#elif defined(OPEN_MODEL_2023_0)
+const uint64 LibAliasVersionMin = LibAlias2023_0_0_Version;
+const uint64 LibAliasVersionMax = LibAliasNext_Version;
+const FString AliasVersionChar = TEXT("AliasStudio 2023.0, Model files");
 #endif
 
 // Alias material management (to allow sew of BReps of different materials):
@@ -326,6 +331,12 @@ public:
 			CADModelConverter = AliasToCADKernelConverter;
 			AliasBRepConverter = AliasToCADKernelConverter;
 		}
+
+		IConsoleVariable* CVarDSPlmXmlTranslatorEnabled = IConsoleManager::Get().FindConsoleVariable(TEXT("ds.CADTranslator.Alias.SewByColor"));
+		if (CVarDSPlmXmlTranslatorEnabled)
+		{
+			bGAliasSewByColor = CVarDSPlmXmlTranslatorEnabled->GetBool();
+		}
 	}
 
 	~FWireTranslatorImpl()
@@ -2257,7 +2268,11 @@ void FDatasmithWireTranslator::Initialize(FDatasmithTranslatorCapabilities& OutC
 #ifdef OPEN_MODEL_2020
 			if (LibAlias2020_Version < FileVersion && FileVersion < LibAlias2021_Version)
 			{
-				UE_LOG(LogDatasmithWireTranslator, Warning, TEXT(WRONG_VERSION_TEXT));
+				TFunction<bool()> DisplayMessage = []() -> bool 
+				{ 
+					UE_LOG(LogDatasmithWireTranslator, Warning, TEXT(WRONG_VERSION_TEXT)); return true; 
+				};
+				static const bool bIsDispaly = DisplayMessage();
 				OutCapabilities.bIsEnabled = false;
 				return;
 			}
