A CVar (ds.IFC.EnableNativeTranslator) allows to enable UE native IFC translator.

diff --git a/Engine/Plugins/Enterprise/DatasmithCADImporter/Source/DatasmithCADTranslator/Private/DatasmithCADTranslator.cpp b/Engine/Plugins/Enterprise/DatasmithCADImporter/Source/DatasmithCADTranslator/Private/DatasmithCADTranslator.cpp
--- a/Engine/Plugins/Enterprise/DatasmithCADImporter/Source/DatasmithCADTranslator/Private/DatasmithCADTranslator.cpp
+++ b/Engine/Plugins/Enterprise/DatasmithCADImporter/Source/DatasmithCADTranslator/Private/DatasmithCADTranslator.cpp
@@ -12,19 +12,42 @@
 #include "DatasmithSceneGraphBuilder.h"
 #include "DatasmithTranslator.h"
 #include "DatasmithUtils.h"
+#include "HAL/IConsoleManager.h"
 #include "IDatasmithSceneElements.h"
 
 
 DEFINE_LOG_CATEGORY(LogCADTranslator);
 
+namespace DatasmithCADTranslatorImpl
+{
+static bool bGEnableNativeIFCTranslator = false;
+FAutoConsoleVariableRef GCADTranslatorEnableNativeIFCTranslator(
+	TEXT("ds.IFC.EnableNativeTranslator"),
+	bGEnableNativeIFCTranslator,
+	TEXT("\
+Enable/disable UE native IFC translator. If native translator is disabled, TechSoft is used.\n\
+Default is disable\n"),
+	ECVF_Default);
+}
 
 void FDatasmithCADTranslator::Initialize(FDatasmithTranslatorCapabilities& OutCapabilities)
 {
-	if (ICADInterfacesModule::GetAvailability() == ECADInterfaceAvailability::Unavailable)
+	TFunction<bool()> GetCADInterfaceAvailability = []() -> bool
+	{
+		if (ICADInterfacesModule::GetAvailability() == ECADInterfaceAvailability::Unavailable)
+		{
+			return false;
+		}
+		return true;
+	};
+
+	static bool bIsCADInterfaceAvailable = GetCADInterfaceAvailability();
+	if (!bIsCADInterfaceAvailable)
 	{
 		OutCapabilities.bIsEnabled = false;
 		return;
 	}
+	OutCapabilities.bIsEnabled = true;
 
 #ifndef CAD_TRANSLATOR_DEBUG
 	OutCapabilities.bParallelLoadStaticMeshSupported = true;
@@ -74,6 +97,11 @@ void FDatasmithCADTranslator::Initialize(FDatasmithTranslatorCapabilities& OutCa
 	OutCapabilities.SupportedFileFormats.Add(FFileFormatInfo{ TEXT("dwg"), TEXT("AutoCAD, Model files") });
 	OutCapabilities.SupportedFileFormats.Add(FFileFormatInfo{ TEXT("dgn"), TEXT("MicroStation files") });
 
+	if (!DatasmithCADTranslatorImpl::bGEnableNativeIFCTranslator)
+	{
+		OutCapabilities.SupportedFileFormats.Add(FFileFormatInfo{ TEXT("ifc"), TEXT("IFC (Industry Foundation Classes)") });
+	}
+
 	if (CADLibrary::FImportParameters::GCADLibrary.Equals(TEXT("TechSoft")))
 	{
 		OutCapabilities.SupportedFileFormats.Add(FFileFormatInfo{ TEXT("hsf"), TEXT("HOOPS stream files") });
diff --git a/Engine/Plugins/Enterprise/DatasmithCADImporter/Source/DatasmithPLMXMLTranslator/Private/DatasmithPlmXmlTranslator.cpp b/Engine/Plugins/Enterprise/DatasmithCADImporter/Source/DatasmithPLMXMLTranslator/Private/DatasmithPlmXmlTranslator.cpp
--- a/Engine/Plugins/Enterprise/DatasmithCADImporter/Source/DatasmithPLMXMLTranslator/Private/DatasmithPlmXmlTranslator.cpp
+++ b/Engine/Plugins/Enterprise/DatasmithCADImporter/Source/DatasmithPLMXMLTranslator/Private/DatasmithPlmXmlTranslator.cpp
@@ -18,16 +18,21 @@
 
 DEFINE_LOG_CATEGORY_STATIC(LogDatasmithXMLPLMTranslator, Log, All)
 
-
 void FDatasmithPlmXmlTranslator::Initialize(FDatasmithTranslatorCapabilities& OutCapabilities)
 {
 #if WITH_EDITOR
 	if (GIsEditor && !GEditor->PlayWorld && !GIsPlayInEditorWorld)
 	{
-		if (ICADInterfacesModule::GetAvailability() == ECADInterfaceAvailability::Unavailable)
+		TFunction<bool()> GetCADInterfaceAvailability = []() -> bool
 		{
-			UE_LOG(LogDatasmithXMLPLMTranslator, Warning, TEXT("CAD Interface module is unavailable. Most of CAD formats (except to Rhino and Alias formats) cannot be imported."));
-		}
+			if (ICADInterfacesModule::GetAvailability() == ECADInterfaceAvailability::Unavailable)
+			{
+				UE_LOG(LogDatasmithXMLPLMTranslator, Warning, TEXT("CAD Interface module is unavailable. Most of CAD formats (except to Rhino and Alias formats) cannot be imported."));
+				return false;
+			}
+			return true;
+		};
+		static bool bIsCADInterfaceAvailable = GetCADInterfaceAvailability();
 
 		OutCapabilities.bIsEnabled = true;
 		OutCapabilities.bParallelLoadStaticMeshSupported = true;
diff --git a/Engine/Plugins/Enterprise/DatasmithIFCImporter/Source/DatasmithIFCTranslator/Private/DatasmithIFCTranslator.cpp b/Engine/Plugins/Enterprise/DatasmithIFCImporter/Source/DatasmithIFCTranslator/Private/DatasmithIFCTranslator.cpp
--- a/Engine/Plugins/Enterprise/DatasmithIFCImporter/Source/DatasmithIFCTranslator/Private/DatasmithIFCTranslator.cpp
+++ b/Engine/Plugins/Enterprise/DatasmithIFCImporter/Source/DatasmithIFCTranslator/Private/DatasmithIFCTranslator.cpp
@@ -9,6 +9,7 @@
 
 #include "Algo/Count.h"
 #include "GenericPlatform/GenericPlatformProcess.h"
+#include "HAL/IConsoleManager.h"
 #include "Misc/FileHelper.h"
 #include "Misc/Paths.h"
 #include "Logging/LogMacros.h"
@@ -61,11 +62,24 @@ void ShowLogMessages(const TArray<IFC::FLogMessage>& Errors)
 
 void FDatasmithIFCTranslator::Initialize(FDatasmithTranslatorCapabilities& OutCapabilities)
 {
-	OutCapabilities.bIsEnabled = true;
 	OutCapabilities.bParallelLoadStaticMeshSupported = true;
 
+	IConsoleVariable* CVarIfcTranslatorEnabled = IConsoleManager::Get().FindConsoleVariable(TEXT("ds.IFC.EnableNativeTranslator"));
+	if (CVarIfcTranslatorEnabled)
+	{
+		OutCapabilities.bIsEnabled = CVarIfcTranslatorEnabled->GetBool();
+		if (!OutCapabilities.bIsEnabled)
+		{
+			return;
+		}
+	}
+	else
+	{
+		OutCapabilities.bIsEnabled = true;
+	}
+
 	TArray<FFileFormatInfo>& Formats = OutCapabilities.SupportedFileFormats;
-    Formats.Emplace(TEXT("ifc"), TEXT("IFC (Industry Foundation Classes)"));
+	Formats.Emplace(TEXT("ifc"), TEXT("IFC (Industry Foundation Classes)"));
 }
 
 bool FDatasmithIFCTranslator::LoadScene(TSharedRef<IDatasmithScene> OutScene)
