[NO NON-PII COMMIT MESSAGE]

diff --git a/Engine/Plugins/Enterprise/DatasmithCADImporter/Source/DatasmithCADTranslator/Private/DatasmithCADTranslator.cpp b/Engine/Plugins/Enterprise/DatasmithCADImporter/Source/DatasmithCADTranslator/Private/DatasmithCADTranslator.cpp
--- a/Engine/Plugins/Enterprise/DatasmithCADImporter/Source/DatasmithCADTranslator/Private/DatasmithCADTranslator.cpp
+++ b/Engine/Plugins/Enterprise/DatasmithCADImporter/Source/DatasmithCADTranslator/Private/DatasmithCADTranslator.cpp
@@ -129,12 +129,12 @@ bool FDatasmithCADTranslator::LoadScene(TSharedRef<IDatasmithScene> DatasmithSce
 	CADLibrary::FFileDescriptor FileDescriptor(*FPaths::ConvertRelativePathToFull(GetSource().GetSourceFile()));
 
 	UE_LOG(LogCADTranslator, Display, TEXT("CAD translation [%s]."), *FileDescriptor.GetSourcePath());
-	UE_LOG(LogCADTranslator, Display, TEXT(" - Parsing Library:     %s"), *CADLibrary::FImportParameters::GCADLibrary);
-	UE_LOG(LogCADTranslator, Display, TEXT(" - Tesselation Library: %s")
+	UE_LOG(LogCADTranslator, Display, TEXT(" - Parsing Library:      %s"), *CADLibrary::FImportParameters::GCADLibrary);
+	UE_LOG(LogCADTranslator, Display, TEXT(" - Tessellation Library: %s")
 		, CADLibrary::FImportParameters::bGDisableCADKernelTessellation ? *CADLibrary::FImportParameters::GCADLibrary : TEXT("CADKernel"));
-	UE_LOG(LogCADTranslator, Display, TEXT(" - Cache mode:          %s")
+	UE_LOG(LogCADTranslator, Display, TEXT(" - Cache mode:           %s")
 		, CADLibrary::FImportParameters::bGEnableCADCache ? (CADLibrary::FImportParameters::bGOverwriteCache ? TEXT("Override") : TEXT("Enabled")) : TEXT("Disabled"));
-	UE_LOG(LogCADTranslator, Display, TEXT(" - Processing:          %s")
+	UE_LOG(LogCADTranslator, Display, TEXT(" - Processing:           %s")
 		, CADLibrary::FImportParameters::bGEnableCADCache ? (CADLibrary::GMaxImportThreads == 1 ? TEXT("Sequencial") : TEXT("Parallel")) : TEXT("Sequencial"));
 
 	ImportParameters.SetTesselationParameters(TesselationOptions.ChordTolerance, TesselationOptions.MaxEdgeLength, TesselationOptions.NormalTolerance, (CADLibrary::EStitchingTechnique)TesselationOptions.StitchingTechnique);
diff --git a/Engine/Plugins/Enterprise/DatasmithCADImporter/Source/DatasmithOpenNurbsTranslator/Private/DatasmithOpenNurbsTranslator.cpp b/Engine/Plugins/Enterprise/DatasmithCADImporter/Source/DatasmithOpenNurbsTranslator/Private/DatasmithOpenNurbsTranslator.cpp
--- a/Engine/Plugins/Enterprise/DatasmithCADImporter/Source/DatasmithOpenNurbsTranslator/Private/DatasmithOpenNurbsTranslator.cpp
+++ b/Engine/Plugins/Enterprise/DatasmithCADImporter/Source/DatasmithOpenNurbsTranslator/Private/DatasmithOpenNurbsTranslator.cpp
@@ -3273,6 +3273,26 @@ bool FDatasmithOpenNurbsTranslator::LoadScene(TSharedRef<IDatasmithScene> OutSce
 		return false;
 	}
 
+	{
+		FString TesselationLibrary;
+		if (OpenNurbsOptions.Geometry == EDatasmithOpenNurbsBrepTessellatedSource::UseRenderMeshes)
+		{
+			TesselationLibrary = TEXT("Rhino");
+		}
+		else if (CADLibrary::FImportParameters::bGDisableCADKernelTessellation)
+		{
+			TesselationLibrary = *CADLibrary::FImportParameters::GCADLibrary;
+		}
+		else
+		{
+			TesselationLibrary = TEXT("CADKernel");
+		}
+
+		UE_LOG(LogDatasmithOpenNurbsTranslator, Display, TEXT("CAD translation [%s]."), *Filename);
+		UE_LOG(LogDatasmithOpenNurbsTranslator, Display, TEXT(" - Parsing Library:      %s"), TEXT("OpenNurbs"));
+		UE_LOG(LogDatasmithOpenNurbsTranslator, Display, TEXT(" - Tessellation Library: %s"), *TesselationLibrary);
+	}
+
 	Translator = MakeShared<FOpenNurbsTranslatorImpl>(GetSource().GetSceneName(), OutScene, FPaths::GetPath(Filename), nullptr);
 	if (!Translator)
 	{
diff --git a/Engine/Plugins/Enterprise/DatasmithCADImporter/Source/DatasmithWireTranslator/Private/DatasmithWireTranslator.cpp b/Engine/Plugins/Enterprise/DatasmithCADImporter/Source/DatasmithWireTranslator/Private/DatasmithWireTranslator.cpp
--- a/Engine/Plugins/Enterprise/DatasmithCADImporter/Source/DatasmithWireTranslator/Private/DatasmithWireTranslator.cpp
+++ b/Engine/Plugins/Enterprise/DatasmithCADImporter/Source/DatasmithWireTranslator/Private/DatasmithWireTranslator.cpp
@@ -2315,6 +2315,10 @@ bool FDatasmithWireTranslator::LoadScene(TSharedRef<IDatasmithScene> OutScene)
 		return false;
 	}
 
+	UE_LOG(LogDatasmithWireTranslator, Display, TEXT("CAD translation [%s]."), *Filename);
+	UE_LOG(LogDatasmithWireTranslator, Display, TEXT(" - Parsing Library:      %s"), TEXT("Alias"));
+	UE_LOG(LogDatasmithWireTranslator, Display, TEXT(" - Tessellation Library: %s") , CADLibrary::FImportParameters::bGDisableCADKernelTessellation ? *CADLibrary::FImportParameters::GCADLibrary : TEXT("CADKernel"));
+
 	FString OutputPath = FPaths::ConvertRelativePathToFull(FPaths::Combine(FDatasmithWireTranslatorModule::Get().GetTempDir(), TEXT("Cache"), GetSource().GetSceneName()));
 	IFileManager::Get().MakeDirectory(*OutputPath, true);
 	Translator->SetOutputPath(OutputPath);
