Remove KernelIO SDK: step 1 - remove references

diff --git a/Engine/Plugins/Enterprise/DatasmithCADImporter/Source/CADInterfaces/Private/CADFileReader.cpp b/Engine/Plugins/Enterprise/DatasmithCADImporter/Source/CADInterfaces/Private/CADFileReader.cpp
--- a/Engine/Plugins/Enterprise/DatasmithCADImporter/Source/CADInterfaces/Private/CADFileReader.cpp
+++ b/Engine/Plugins/Enterprise/DatasmithCADImporter/Source/CADInterfaces/Private/CADFileReader.cpp
@@ -7,10 +7,6 @@
 #include "HAL/FileManager.h"
 #include "Templates/TypeHash.h"
 
-#ifdef USE_KERNEL_IO_SDK
-#include "CoreTechFileParser.h"
-#endif
-
 #ifdef USE_TECHSOFT_SDK
 #include "TechSoftFileParser.h"
 #include "TechSoftFileParserCADKernelTessellator.h"
@@ -21,26 +17,14 @@ namespace CADLibrary
 	FCADFileReader::FCADFileReader(const FImportParameters& ImportParams, FFileDescriptor& InFile, const FString& EnginePluginsPath, const FString& InCachePath)
 		: CADFileData(ImportParams, InFile, InCachePath)
 	{
-#ifdef USE_KERNEL_IO_SDK
-		if (FImportParameters::GCADLibrary == TEXT("KernelIO"))
+#ifdef USE_TECHSOFT_SDK
+		if (FImportParameters::bGDisableCADKernelTessellation)
 		{
-			CADParser = MakeUnique<FCoreTechFileParser>(CADFileData, EnginePluginsPath);
+			CADParser = MakeUnique<FTechSoftFileParser>(CADFileData, EnginePluginsPath);
 		}
-#endif
-#if defined(USE_KERNEL_IO_SDK) && defined(USE_TECHSOFT_SDK)
 		else
-#endif
-#ifdef USE_TECHSOFT_SDK
-		if (FImportParameters::GCADLibrary == TEXT("TechSoft"))
 		{
-			if (FImportParameters::bGDisableCADKernelTessellation)
-			{
-				CADParser = MakeUnique<FTechSoftFileParser>(CADFileData, EnginePluginsPath);
-			}
-			else
-			{
-				CADParser = MakeUnique<FTechSoftFileParserCADKernelTessellator>(CADFileData, EnginePluginsPath);
-			}
+			CADParser = MakeUnique<FTechSoftFileParserCADKernelTessellator>(CADFileData, EnginePluginsPath);
 		}
 #endif
 	}
diff --git a/Engine/Plugins/Enterprise/DatasmithCADImporter/Source/CADInterfaces/Private/CADInterfacesModule.cpp b/Engine/Plugins/Enterprise/DatasmithCADImporter/Source/CADInterfaces/Private/CADInterfacesModule.cpp
--- a/Engine/Plugins/Enterprise/DatasmithCADImporter/Source/CADInterfaces/Private/CADInterfacesModule.cpp
+++ b/Engine/Plugins/Enterprise/DatasmithCADImporter/Source/CADInterfaces/Private/CADInterfacesModule.cpp
@@ -2,7 +2,6 @@
 #include "CADInterfacesModule.h"
 
 #include "CADOptions.h"
-#include "CoreTechTypes.h"
 #include "HAL/PlatformProcess.h"
 #include "Misc/Paths.h"
 #include "Modules/ModuleManager.h"
@@ -22,7 +21,6 @@ private:
 	static void* TechSoftLibHandle;
 };
 
-void* FCADInterfacesModule::KernelIOLibHandle = nullptr;
 void* FCADInterfacesModule::TechSoftLibHandle = nullptr;
 
 ICADInterfacesModule& ICADInterfacesModule::Get()
@@ -36,19 +34,9 @@ ECADInterfaceAvailability ICADInterfacesModule::GetAvailability()
 {
 	if (FModuleManager::Get().IsModuleLoaded(CADINTERFACES_MODULE_NAME))
 	{
-		if (CADLibrary::FImportParameters::GCADLibrary == TEXT("KernelIO"))
+		if (CADLibrary::TechSoftInterface::TECHSOFT_InitializeKernel())
 		{
-			if (CADLibrary::CTKIO_InitializeKernel())
-			{
-				return ECADInterfaceAvailability::Available;
-			}
-		}
-		else if (CADLibrary::FImportParameters::GCADLibrary == TEXT("TechSoft"))
-		{
-			if (CADLibrary::TechSoftInterface::TECHSOFT_InitializeKernel())
-			{
-				return ECADInterfaceAvailability::Available;
-			}
+			return ECADInterfaceAvailability::Available;
 		}
 	}
 
@@ -62,32 +50,6 @@ void FCADInterfacesModule::StartupModule()
 	FString CADImporterDllPath = FPaths::Combine(FPaths::EnginePluginsDir(), TEXT("Enterprise/DatasmithCADImporter"), TEXT("Binaries"), FPlatformProcess::GetBinariesSubdirectory());
 	FPlatformProcess::PushDllDirectory(*CADImporterDllPath);
 
-#if WITH_EDITOR & defined(USE_KERNEL_IO_SDK)
-	check(KernelIOLibHandle == nullptr);
-
-	FString KernelIODll = TEXT("kernel_io.dll");
-	KernelIODll = FPaths::Combine(CADImporterDllPath, KernelIODll);
-
-	if (!FPaths::FileExists(KernelIODll))
-	{
-		UE_LOG(LogCADInterfaces, Warning, TEXT("CoreTech module is missing. Plug-in will not be functional."));
-	}
-	else
-	{
-		KernelIOLibHandle = FPlatformProcess::GetDllHandle(*KernelIODll);
-		if (KernelIOLibHandle == nullptr)
-		{
-			UE_LOG(LogCADInterfaces, Warning, TEXT("Failed to load required library %s. Plug-in will not be functional."), *KernelIODll);
-		}
-		else
-		{
-			CADLibrary::InitializeCoreTechInterface();
-		}
-
-		FPlatformProcess::PopDllDirectory(*CADImporterDllPath);
-	}
-#endif
-
 #if WITH_EDITOR & defined(USE_TECHSOFT_SDK)
 	check(TechSoftLibHandle == nullptr);
 
@@ -123,15 +85,6 @@ void FCADInterfacesModule::StartupModule()
 
 void FCADInterfacesModule::ShutdownModule()
 {
-	if (KernelIOLibHandle != nullptr)
-	{
-#if WITH_EDITOR && defined(USE_KERNEL_IO_SDK)
-		// Reset the CoreTechInterface object if compiling for the editor and CoreTech sdk is available
-		CADLibrary::SetCoreTechInterface(TSharedPtr<CADLibrary::ICoreTechInterface>());
-#endif
-		FPlatformProcess::FreeDllHandle(KernelIOLibHandle);
-		KernelIOLibHandle = nullptr;
-	}
 }
 
 IMPLEMENT_MODULE(FCADInterfacesModule, CADInterfaces);
diff --git a/Engine/Plugins/Enterprise/DatasmithCADImporter/Source/CADInterfaces/Public/CADFileData.h b/Engine/Plugins/Enterprise/DatasmithCADImporter/Source/CADInterfaces/Public/CADFileData.h
--- a/Engine/Plugins/Enterprise/DatasmithCADImporter/Source/CADInterfaces/Public/CADFileData.h
+++ b/Engine/Plugins/Enterprise/DatasmithCADImporter/Source/CADInterfaces/Public/CADFileData.h
@@ -28,7 +28,6 @@ public:
 		if (!SceneFileHash)
 		{
 			SceneFileHash = HashCombine(FileDescription.GetDescriptorHash(), ::GetTypeHash(ImportParameters.GetStitchingTechnique()));
-			SceneFileHash = HashCombine(SceneFileHash, ::GetTypeHash(CADLibrary::FImportParameters::GCADLibrary));
 		}
 		return SceneFileHash;
 	}
diff --git a/Engine/Plugins/Enterprise/DatasmithCADImporter/Source/CADTools/Private/CADOption.cpp b/Engine/Plugins/Enterprise/DatasmithCADImporter/Source/CADTools/Private/CADOption.cpp
--- a/Engine/Plugins/Enterprise/DatasmithCADImporter/Source/CADTools/Private/CADOption.cpp
+++ b/Engine/Plugins/Enterprise/DatasmithCADImporter/Source/CADTools/Private/CADOption.cpp
@@ -4,16 +4,6 @@
 namespace CADLibrary
 {
 
-FString FImportParameters::GCADLibrary = TEXT("TechSoft");
-
-#ifdef USE_KERNEL_IO_SDK 
-FAutoConsoleVariableRef GCADTranslatorLibrary(
-	TEXT("ds.CADTranslator.CADLibrary"),
-	FImportParameters::GCADLibrary,
-	TEXT("Define the used library to import CAD file i.e. \"TechSoft\" or \"KernelIO\".\n"),
-	ECVF_Default);
-#endif
-
 int32 GMaxImportThreads = 0;
 FAutoConsoleVariableRef GCADTranslatorMaxImportThreads(
 	TEXT("ds.CADTranslator.MaxImportThreads"),
@@ -83,7 +73,6 @@ ECVF_Default);
 uint32 GetTypeHash(const FImportParameters& ImportParameters)
 {
 	uint32 ParametersHash = ::GetTypeHash(ImportParameters.bGDisableCADKernelTessellation);
-	ParametersHash = HashCombine(ParametersHash, ::GetTypeHash(ImportParameters.GCADLibrary));
 	ParametersHash = HashCombine(ParametersHash, ::GetTypeHash(ImportParameters.ChordTolerance));
 	ParametersHash = HashCombine(ParametersHash, ::GetTypeHash(ImportParameters.MaxEdgeLength));
 	ParametersHash = HashCombine(ParametersHash, ::GetTypeHash(ImportParameters.MaxNormalAngle));
diff --git a/Engine/Plugins/Enterprise/DatasmithCADImporter/Source/CADTools/Private/CADToolsModule.cpp b/Engine/Plugins/Enterprise/DatasmithCADImporter/Source/CADTools/Private/CADToolsModule.cpp
--- a/Engine/Plugins/Enterprise/DatasmithCADImporter/Source/CADTools/Private/CADToolsModule.cpp
+++ b/Engine/Plugins/Enterprise/DatasmithCADImporter/Source/CADTools/Private/CADToolsModule.cpp
@@ -18,7 +18,7 @@ bool FCADToolsModule::IsAvailable()
 
 uint32 FCADToolsModule::GetCacheVersion()
 {
-	return 501;
+	return 502;
 }
 
 IMPLEMENT_MODULE(FCADToolsModule, CADTools);
diff --git a/Engine/Plugins/Enterprise/DatasmithCADImporter/Source/CADTools/Public/CADOptions.h b/Engine/Plugins/Enterprise/DatasmithCADImporter/Source/CADTools/Public/CADOptions.h
--- a/Engine/Plugins/Enterprise/DatasmithCADImporter/Source/CADTools/Public/CADOptions.h
+++ b/Engine/Plugins/Enterprise/DatasmithCADImporter/Source/CADTools/Public/CADOptions.h
@@ -55,7 +55,6 @@ namespace CADLibrary
 		CADTOOLS_API static bool bGOverwriteCache;
 		CADTOOLS_API static bool bGPreferJtFileEmbeddedTessellation;
 		CADTOOLS_API static float GStitchingTolerance;
-		CADTOOLS_API static FString GCADLibrary;
 
 	public:
 		FImportParameters(double InMetricUnit = 0.001, double InScaleFactor = 1., FDatasmithUtils::EModelCoordSystem NewCoordinateSystem = FDatasmithUtils::EModelCoordSystem::ZUp_RightHanded)
@@ -130,7 +129,6 @@ namespace CADLibrary
 			Ar << ImportParameters.bGEnableCADCache;
 			Ar << ImportParameters.bGPreferJtFileEmbeddedTessellation;
 			Ar << ImportParameters.GStitchingTolerance;
-			Ar << ImportParameters.GCADLibrary;
 			return Ar;
 		}
 
@@ -215,7 +213,7 @@ namespace CADLibrary
 
 		if (FImportParameters::bGDisableCADKernelTessellation)
 		{
-			OutFileName += FImportParameters::GCADLibrary.Equals("TechSoft") ? TEXT(".prc") : TEXT(".ct");
+			OutFileName += TEXT(".prc");
 		}
 		else
 		{
diff --git a/Engine/Plugins/Enterprise/DatasmithCADImporter/Source/DatasmithCADTranslator/Private/DatasmithCADTranslator.cpp b/Engine/Plugins/Enterprise/DatasmithCADImporter/Source/DatasmithCADTranslator/Private/DatasmithCADTranslator.cpp
--- a/Engine/Plugins/Enterprise/DatasmithCADImporter/Source/DatasmithCADTranslator/Private/DatasmithCADTranslator.cpp
+++ b/Engine/Plugins/Enterprise/DatasmithCADImporter/Source/DatasmithCADTranslator/Private/DatasmithCADTranslator.cpp
@@ -102,15 +102,8 @@ void FDatasmithCADTranslator::Initialize(FDatasmithTranslatorCapabilities& OutCa
 		OutCapabilities.SupportedFileFormats.Add(FFileFormatInfo{ TEXT("ifc"), TEXT("IFC (Industry Foundation Classes)") });
 	}
 
-	if (CADLibrary::FImportParameters::GCADLibrary.Equals(TEXT("TechSoft")))
-	{
-		OutCapabilities.SupportedFileFormats.Add(FFileFormatInfo{ TEXT("hsf"), TEXT("HOOPS stream files") });
-		OutCapabilities.SupportedFileFormats.Add(FFileFormatInfo{ TEXT("prc"), TEXT("HOOPS stream files") });
-	}
-	else
-	{
-		OutCapabilities.SupportedFileFormats.Add(FFileFormatInfo{ TEXT("ct"), TEXT("Kernel_IO files") });
-	}
+	OutCapabilities.SupportedFileFormats.Add(FFileFormatInfo{ TEXT("hsf"), TEXT("HOOPS stream files") });
+	OutCapabilities.SupportedFileFormats.Add(FFileFormatInfo{ TEXT("prc"), TEXT("HOOPS stream files") });
 }
 
 bool FDatasmithCADTranslator::IsSourceSupported(const FDatasmithSceneSource& Source)
@@ -129,9 +122,9 @@ bool FDatasmithCADTranslator::LoadScene(TSharedRef<IDatasmithScene> DatasmithSce
 	CADLibrary::FFileDescriptor FileDescriptor(*FPaths::ConvertRelativePathToFull(GetSource().GetSourceFile()));
 
 	UE_LOG(LogCADTranslator, Display, TEXT("CAD translation [%s]."), *FileDescriptor.GetSourcePath());
-	UE_LOG(LogCADTranslator, Display, TEXT(" - Parsing Library:      %s"), *CADLibrary::FImportParameters::GCADLibrary);
+	UE_LOG(LogCADTranslator, Display, TEXT(" - Parsing Library:      %s"), TEXT("TechSoft"));
 	UE_LOG(LogCADTranslator, Display, TEXT(" - Tessellation Library: %s")
-		, CADLibrary::FImportParameters::bGDisableCADKernelTessellation ? *CADLibrary::FImportParameters::GCADLibrary : TEXT("CADKernel"));
+		, CADLibrary::FImportParameters::bGDisableCADKernelTessellation ? TEXT("TechSoft") : TEXT("CADKernel"));
 	UE_LOG(LogCADTranslator, Display, TEXT(" - Cache mode:           %s")
 		, CADLibrary::FImportParameters::bGEnableCADCache ? (CADLibrary::FImportParameters::bGOverwriteCache ? TEXT("Override") : TEXT("Enabled")) : TEXT("Disabled"));
 	UE_LOG(LogCADTranslator, Display, TEXT(" - Processing:           %s")
diff --git a/Engine/Plugins/Enterprise/DatasmithCADImporter/Source/DatasmithCADTranslator/Public/DatasmithSceneGraphBuilder.h b/Engine/Plugins/Enterprise/DatasmithCADImporter/Source/DatasmithCADTranslator/Public/DatasmithSceneGraphBuilder.h
--- a/Engine/Plugins/Enterprise/DatasmithCADImporter/Source/DatasmithCADTranslator/Public/DatasmithSceneGraphBuilder.h
+++ b/Engine/Plugins/Enterprise/DatasmithCADImporter/Source/DatasmithCADTranslator/Public/DatasmithSceneGraphBuilder.h
@@ -5,8 +5,6 @@
 
 #include "CADData.h"
 #include "CADSceneGraph.h"
-#include "CoreTechTypes.h"
-#include "CTSession.h"
 #include "DatasmithImportOptions.h"
 #include "MeshDescriptionHelper.h"
 
diff --git a/Engine/Plugins/Enterprise/DatasmithCADImporter/Source/DatasmithOpenNurbsTranslator/Private/DatasmithOpenNurbsTranslator.cpp b/Engine/Plugins/Enterprise/DatasmithCADImporter/Source/DatasmithOpenNurbsTranslator/Private/DatasmithOpenNurbsTranslator.cpp
--- a/Engine/Plugins/Enterprise/DatasmithCADImporter/Source/DatasmithOpenNurbsTranslator/Private/DatasmithOpenNurbsTranslator.cpp
+++ b/Engine/Plugins/Enterprise/DatasmithCADImporter/Source/DatasmithOpenNurbsTranslator/Private/DatasmithOpenNurbsTranslator.cpp
@@ -9,8 +9,6 @@
 #include "CADKernelSurfaceExtension.h"
 #include "CADModelConverter.h"
 #include "CADOptions.h"
-#include "CoreTechSurfaceExtension.h"
-#include "CoreTechSurfaceHelper.h"
 #include "DatasmithImportOptions.h"
 #include "DatasmithMaterialElements.h"
 #include "DatasmithMaterialsUtils.h"
@@ -20,7 +18,6 @@
 #include "DatasmithUtils.h"
 #include "OpenNurbsBRepConverter.h"
 #include "OpenNurbsBRepToCADKernelConverter.h"
-#include "OpenNurbsBRepToCoretechConverter.h" // requires CoreTech as public dependency
 #include "OpenNurbsBRepToTechSoftConverter.h" // requires Techsoft as public dependency
 #include "Utility/DatasmithMeshHelper.h"
 
@@ -653,18 +650,9 @@ public:
 
 		if (CADLibrary::FImportParameters::bGDisableCADKernelTessellation)
 		{
-			if (CADLibrary::FImportParameters::GCADLibrary == TEXT("TechSoft"))
-			{
-				TSharedRef<FOpenNurbsBRepToTechSoftConverter> OpenNurbsBRepToCoretechConverter = MakeShared<FOpenNurbsBRepToTechSoftConverter>(ImportParameters);
-				CADModelConverter = OpenNurbsBRepToCoretechConverter;
-				OpenNurbsBRepConverter = OpenNurbsBRepToCoretechConverter;
-			}
-			else if (CADLibrary::FImportParameters::GCADLibrary == TEXT("KernelIO"))
-			{
-				TSharedRef<FOpenNurbsBRepToCoretechConverter> OpenNurbsBRepToCoretechConverter = MakeShared<FOpenNurbsBRepToCoretechConverter>(TEXT("Al2CTSharedSession"), ImportParameters);
-				CADModelConverter = OpenNurbsBRepToCoretechConverter;
-				OpenNurbsBRepConverter = OpenNurbsBRepToCoretechConverter;
-			}
+			TSharedRef<FOpenNurbsBRepToTechSoftConverter> OpenNurbsBRepToTechSoftConverter = MakeShared<FOpenNurbsBRepToTechSoftConverter>(ImportParameters);
+			CADModelConverter = OpenNurbsBRepToTechSoftConverter;
+			OpenNurbsBRepConverter = OpenNurbsBRepToTechSoftConverter;
 		}
 		else
 		{
@@ -3281,7 +3269,7 @@ bool FDatasmithOpenNurbsTranslator::LoadScene(TSharedRef<IDatasmithScene> OutSce
 		}
 		else if (CADLibrary::FImportParameters::bGDisableCADKernelTessellation)
 		{
-			TesselationLibrary = *CADLibrary::FImportParameters::GCADLibrary;
+			TesselationLibrary = TEXT("TechSoft");
 		}
 		else
 		{
diff --git a/Engine/Plugins/Enterprise/DatasmithCADImporter/Source/DatasmithOpenNurbsTranslator/Private/OpenNurbsBRepToCADKernelConverter.cpp b/Engine/Plugins/Enterprise/DatasmithCADImporter/Source/DatasmithOpenNurbsTranslator/Private/OpenNurbsBRepToCADKernelConverter.cpp
--- a/Engine/Plugins/Enterprise/DatasmithCADImporter/Source/DatasmithOpenNurbsTranslator/Private/OpenNurbsBRepToCADKernelConverter.cpp
+++ b/Engine/Plugins/Enterprise/DatasmithCADImporter/Source/DatasmithOpenNurbsTranslator/Private/OpenNurbsBRepToCADKernelConverter.cpp
@@ -3,7 +3,6 @@
 #include "OpenNurbsBRepToCADKernelConverter.h"
 
 #ifdef USE_OPENNURBS
-#include "CoreTechSurfaceHelper.h"
 
 #pragma warning(push)
 #pragma warning(disable:4265)
diff --git a/Engine/Plugins/Enterprise/DatasmithCADImporter/Source/DatasmithOpenNurbsTranslator/Private/OpenNurbsBRepToTechSoftConverter.h b/Engine/Plugins/Enterprise/DatasmithCADImporter/Source/DatasmithOpenNurbsTranslator/Private/OpenNurbsBRepToTechSoftConverter.h
--- a/Engine/Plugins/Enterprise/DatasmithCADImporter/Source/DatasmithOpenNurbsTranslator/Private/OpenNurbsBRepToTechSoftConverter.h
+++ b/Engine/Plugins/Enterprise/DatasmithCADImporter/Source/DatasmithOpenNurbsTranslator/Private/OpenNurbsBRepToTechSoftConverter.h
@@ -8,7 +8,6 @@
 
 #include "OpenNurbsBRepConverter.h"
 #include "CADModelToTechSoftConverterBase.h"
-#include "CTSession.h"
 
 #ifdef USE_TECHSOFT_SDK
 #include "TechSoftInterface.h"
diff --git a/Engine/Plugins/Enterprise/DatasmithCADImporter/Source/DatasmithPLMXMLTranslator/Private/DatasmithPlmXmlImporter.cpp b/Engine/Plugins/Enterprise/DatasmithCADImporter/Source/DatasmithPLMXMLTranslator/Private/DatasmithPlmXmlImporter.cpp
--- a/Engine/Plugins/Enterprise/DatasmithCADImporter/Source/DatasmithPLMXMLTranslator/Private/DatasmithPlmXmlImporter.cpp
+++ b/Engine/Plugins/Enterprise/DatasmithCADImporter/Source/DatasmithPLMXMLTranslator/Private/DatasmithPlmXmlImporter.cpp
@@ -20,8 +20,6 @@
 #include "CADData.h"
 #include "CADToolsModule.h"
 #include "CADKernelSurfaceExtension.h"
-#include "CoreTechSurfaceExtension.h"
-#include "CoreTechSurfaceHelper.h"
 #include "DatasmithAdditionalData.h"
 #include "DatasmithDispatcher.h"
 #include "DatasmithMeshBuilder.h"
diff --git a/Engine/Plugins/Enterprise/DatasmithCADImporter/Source/DatasmithWireTranslator/Private/DatasmithWireTranslator.cpp b/Engine/Plugins/Enterprise/DatasmithCADImporter/Source/DatasmithWireTranslator/Private/DatasmithWireTranslator.cpp
--- a/Engine/Plugins/Enterprise/DatasmithCADImporter/Source/DatasmithWireTranslator/Private/DatasmithWireTranslator.cpp
+++ b/Engine/Plugins/Enterprise/DatasmithCADImporter/Source/DatasmithWireTranslator/Private/DatasmithWireTranslator.cpp
@@ -8,12 +8,13 @@
 #include "DatasmithSceneFactory.h"
 #include "DatasmithTranslator.h"
 #include "DatasmithUtils.h"
-#include "HAL/ConsoleManager.h"
 #include "DatasmithWireTranslatorModule.h"
+#include "HAL/ConsoleManager.h"
 #include "IDatasmithSceneElements.h"
 #include "Misc/FileHelper.h"
 #include "Misc/Paths.h"
 #include "OpenModelUtils.h"
+#include "Utility/DatasmithMeshHelper.h"
 
 #include "StaticMeshDescription.h"
 #include "StaticMeshOperations.h"
@@ -28,11 +29,8 @@
 
 #include "AliasBrepConverter.h"
 #include "AliasModelToCADKernelConverter.h"
-#include "AliasModelToCoretechConverter.h" // requires CoreTech as public dependency
 #include "AliasModelToTechSoftConverter.h" // requires Techsoft as public dependency
 #include "CADInterfacesModule.h"
-#include "CoreTechSurfaceExtension.h"
-#include "CoreTechSurfaceHelper.h"
 
 #if PLATFORM_WINDOWS
 #include "Windows/AllowWindowsPlatformTypes.h"
@@ -312,18 +310,9 @@ public:
 		CADLibrary::FImportParameters ImportParameters(MetricUnit, ScaleFactor);
 		if(CADLibrary::FImportParameters::bGDisableCADKernelTessellation)
 		{
-			if (CADLibrary::FImportParameters::GCADLibrary == TEXT("KernelIO"))
-			{
-				TSharedRef<FAliasModelToCoretechConverter> AliasToCoretechConverter = MakeShared<FAliasModelToCoretechConverter>(TEXT("Al2CTSharedSession"), ImportParameters);
-				CADModelConverter = AliasToCoretechConverter;
-				AliasBRepConverter = AliasToCoretechConverter;
-			}
-			else if (CADLibrary::FImportParameters::GCADLibrary == TEXT("TechSoft"))
-			{
-				TSharedRef<FAliasModelToTechSoftConverter> AliasToTechSoftConverter = MakeShared<FAliasModelToTechSoftConverter>(ImportParameters);
-				CADModelConverter = AliasToTechSoftConverter;
-				AliasBRepConverter = AliasToTechSoftConverter;
-			}
+			TSharedRef<FAliasModelToTechSoftConverter> AliasToTechSoftConverter = MakeShared<FAliasModelToTechSoftConverter>(ImportParameters);
+			CADModelConverter = AliasToTechSoftConverter;
+			AliasBRepConverter = AliasToTechSoftConverter;
 		}
 		else
 		{
@@ -2317,7 +2306,7 @@ bool FDatasmithWireTranslator::LoadScene(TSharedRef<IDatasmithScene> OutScene)
 
 	UE_LOG(LogDatasmithWireTranslator, Display, TEXT("CAD translation [%s]."), *Filename);
 	UE_LOG(LogDatasmithWireTranslator, Display, TEXT(" - Parsing Library:      %s"), TEXT("Alias"));
-	UE_LOG(LogDatasmithWireTranslator, Display, TEXT(" - Tessellation Library: %s") , CADLibrary::FImportParameters::bGDisableCADKernelTessellation ? *CADLibrary::FImportParameters::GCADLibrary : TEXT("CADKernel"));
+	UE_LOG(LogDatasmithWireTranslator, Display, TEXT(" - Tessellation Library: %s") , CADLibrary::FImportParameters::bGDisableCADKernelTessellation ? TEXT("TechSoft") : TEXT("CADKernel"));
 
 	FString OutputPath = FPaths::ConvertRelativePathToFull(FPaths::Combine(FDatasmithWireTranslatorModule::Get().GetTempDir(), TEXT("Cache"), GetSource().GetSceneName()));
 	IFileManager::Get().MakeDirectory(*OutputPath, true);
diff --git a/Engine/Plugins/Enterprise/DatasmithCADImporter/Source/ParametricSurface/Private/ParametricSurfaceModule.cpp b/Engine/Plugins/Enterprise/DatasmithCADImporter/Source/ParametricSurface/Private/ParametricSurfaceModule.cpp
--- a/Engine/Plugins/Enterprise/DatasmithCADImporter/Source/ParametricSurface/Private/ParametricSurfaceModule.cpp
+++ b/Engine/Plugins/Enterprise/DatasmithCADImporter/Source/ParametricSurface/Private/ParametricSurfaceModule.cpp
@@ -2,7 +2,6 @@
 
 #include "ParametricSurfaceModule.h"
 
-#include "CoreTech/CoreTechSurfaceData.h"
 #include "TechSoft/TechSoftParametricSurface.h"
 
 #include "Modules/ModuleInterface.h"
@@ -34,19 +33,9 @@ bool FParametricSurfaceModule::IsAvailable()
 	return FModuleManager::Get().IsModuleLoaded(PARAMETRICSURFACE_MODULE_NAME);
 }
 
-UParametricSurfaceData* FParametricSurfaceModule::CreateParametricSurface(const TCHAR* CADLibraryName)
+UParametricSurfaceData* FParametricSurfaceModule::CreateParametricSurface()
 {
-	if (!FCString::Strcmp(TEXT("TechSoft"), CADLibraryName))
-	{
-		return Datasmith::MakeAdditionalData<UTechSoftParametricSurfaceData>();
-	}
-
-	if (!FCString::Strcmp(TEXT("KernelIO"), CADLibraryName))
-	{
-		return Datasmith::MakeAdditionalData<UCoreTechParametricSurfaceData>();
-	}
-
-	return nullptr;
+	return Datasmith::MakeAdditionalData<UTechSoftParametricSurfaceData>();
 }
 
 IMPLEMENT_MODULE(FParametricSurfaceModule, ParametricSurface)
diff --git a/Engine/Plugins/Enterprise/DatasmithCADImporter/Source/ParametricSurface/Private/ParametricSurfaceTranslator.cpp b/Engine/Plugins/Enterprise/DatasmithCADImporter/Source/ParametricSurface/Private/ParametricSurfaceTranslator.cpp
--- a/Engine/Plugins/Enterprise/DatasmithCADImporter/Source/ParametricSurface/Private/ParametricSurfaceTranslator.cpp
+++ b/Engine/Plugins/Enterprise/DatasmithCADImporter/Source/ParametricSurface/Private/ParametricSurfaceTranslator.cpp
@@ -40,7 +40,7 @@ bool ParametricSurfaceUtils::AddSurfaceData(const TCHAR* MeshFilePath, const CAD
 {
 	if (MeshFilePath && IFileManager::Get().FileExists(MeshFilePath))
 	{
-		UParametricSurfaceData* ParametricSurfaceData = FParametricSurfaceModule::CreateParametricSurface(*CADLibrary::FImportParameters::GCADLibrary);
+		UParametricSurfaceData* ParametricSurfaceData = FParametricSurfaceModule::CreateParametricSurface();
 
 		if (!ParametricSurfaceData || !ParametricSurfaceData->SetFile(MeshFilePath))
 		{
diff --git a/Engine/Plugins/Enterprise/DatasmithCADImporter/Source/ParametricSurface/Public/ParametricSurfaceData.h b/Engine/Plugins/Enterprise/DatasmithCADImporter/Source/ParametricSurface/Public/ParametricSurfaceData.h
--- a/Engine/Plugins/Enterprise/DatasmithCADImporter/Source/ParametricSurface/Public/ParametricSurfaceData.h
+++ b/Engine/Plugins/Enterprise/DatasmithCADImporter/Source/ParametricSurface/Public/ParametricSurfaceData.h
@@ -113,24 +113,3 @@ protected:
 	TArray<uint8> RawData;
 };
 
-//void AddSurfaceDataForMesh(const TCHAR* InFilePath, const CADLibrary::FImportParameters& InSceneParameters, const CADLibrary::FMeshParameters& InMeshParameters, , FDatasmithMeshElementPayload& OutMeshPayload)
-//{
-//	if (ICADInterfacesModule::GetAvailability() == ECADInterfaceAvailability::Available)
-//	{
-//		if (FPaths::FileExists(InFilePath))
-//		{
-//			TArray<uint8> ByteArray;
-//			if (FFileHelper::LoadFileToArray(ByteArray, InFilePath))
-//			{
-//				UCoreTechParametricSurfaceData* CoreTechData = Datasmith::MakeAdditionalData<UCoreTechParametricSurfaceData>();
-//				CoreTechData->SourceFile = InFilePath;
-//				CoreTechData->RawData = MoveTemp(ByteArray);
-
-
-//				CoreTechData->
-//				OutMeshPayload.AdditionalData.Add(CoreTechData);
-//			}
-//		}
-//	}
-//}
-
diff --git a/Engine/Plugins/Enterprise/DatasmithCADImporter/Source/ParametricSurface/Public/ParametricSurfaceModule.h b/Engine/Plugins/Enterprise/DatasmithCADImporter/Source/ParametricSurface/Public/ParametricSurfaceModule.h
--- a/Engine/Plugins/Enterprise/DatasmithCADImporter/Source/ParametricSurface/Public/ParametricSurfaceModule.h
+++ b/Engine/Plugins/Enterprise/DatasmithCADImporter/Source/ParametricSurface/Public/ParametricSurfaceModule.h
@@ -17,5 +17,5 @@ public:
 	static FParametricSurfaceModule& Get();
 	static bool IsAvailable();
 
-	static PARAMETRICSURFACE_API class UParametricSurfaceData* CreateParametricSurface(const TCHAR* CADLibraryName);
+	static PARAMETRICSURFACE_API class UParametricSurfaceData* CreateParametricSurface();
 };
diff --git a/Engine/Plugins/Enterprise/DatasmithCADImporter/Source/ParametricSurfaceExtension/Private/ParametricSurfaceExtensionModule.cpp b/Engine/Plugins/Enterprise/DatasmithCADImporter/Source/ParametricSurfaceExtension/Private/ParametricSurfaceExtensionModule.cpp
--- a/Engine/Plugins/Enterprise/DatasmithCADImporter/Source/ParametricSurfaceExtension/Private/ParametricSurfaceExtensionModule.cpp
+++ b/Engine/Plugins/Enterprise/DatasmithCADImporter/Source/ParametricSurfaceExtension/Private/ParametricSurfaceExtensionModule.cpp
@@ -13,7 +13,7 @@
 #include "UObject/StrongObjectPtr.h"
 
 
-#define LOCTEXT_NAMESPACE "CoreTechExtensionModule"
+#define LOCTEXT_NAMESPACE "ParametricSurfaceExtensionModule"
 
 
 /** UI extension that displays a Retessellate action in the StaticMeshEditor */
@@ -98,5 +98,5 @@ void FParametricSurfaceExtensionModule::StartupModule()
 
 IMPLEMENT_MODULE(FParametricSurfaceExtensionModule, ParametricSurfaceExtension);
 
-#undef LOCTEXT_NAMESPACE // "DatasmithCoreTechExtensionModule"
+#undef LOCTEXT_NAMESPACE // "ParametricSurfaceExtensionModule"
 
diff --git a/Engine/Source/Programs/Enterprise/Datasmith/DatasmithCADWorker/Private/DatasmithCADWorker.cpp b/Engine/Source/Programs/Enterprise/Datasmith/DatasmithCADWorker/Private/DatasmithCADWorker.cpp
--- a/Engine/Source/Programs/Enterprise/Datasmith/DatasmithCADWorker/Private/DatasmithCADWorker.cpp
+++ b/Engine/Source/Programs/Enterprise/Datasmith/DatasmithCADWorker/Private/DatasmithCADWorker.cpp
@@ -7,9 +7,6 @@
 #include "CADInterfacesModule.h"
 #include "CADToolsModule.h"
 
-#ifdef USE_KERNEL_IO_SDK
-#include "CoreTechTypes.h"
-#endif
 #ifdef USE_TECHSOFT_SDK
 #include "TechSoftInterface.h"
 #endif
@@ -52,7 +49,7 @@ int32 Main(int32 Argc, TCHAR * Argv[])
 {
 	UE_SET_LOG_VERBOSITY(LogDatasmithCADWorker, Verbose);
 
-#if !defined(USE_KERNEL_IO_SDK) && !defined(USE_TECHSOFT_SDK)
+#if !defined(USE_TECHSOFT_SDK)
 	UE_LOG(LogDatasmithCADWorker, Error, TEXT("Missing CAD module. DatasmithCADWorker is not functional."));
 	return EXIT_MISSING_CAD_MODULE;
 #else
@@ -68,17 +65,11 @@ int32 Main(int32 Argc, TCHAR * Argv[])
 	GetParameter(Argc, Argv, "-CacheVersion", CacheVersion);
 	GetParameter(Argc, Argv, "-EnginePluginsDir", EnginePluginsPath);
 
-#ifdef USE_KERNEL_IO_SDK
-	CADLibrary::InitializeCoreTechInterface();
-#endif
-
-#ifdef USE_TECHSOFT_SDK
 	if (!CADLibrary::TechSoftInterface::TECHSOFT_InitializeKernel(*EnginePluginsPath))
 	{
 		UE_LOG(LogDatasmithCADWorker, Error, TEXT("TechSoft interface cannot be initialized. CADInterfaces module is not available."));
 		return EXIT_FAILURE;
 	}
-#endif
 
 	if (ICADInterfacesModule::Get().GetAvailability() != ECADInterfaceAvailability::Available)
 	{
@@ -98,7 +89,7 @@ int32 Main(int32 Argc, TCHAR * Argv[])
 	Worker.Run();
 
 	return EXIT_SUCCESS;
-#endif // USE_KERNEL_IO_SDK
+#endif // USE_TECHSOFT_SDK
 }
 
 int32 Filter(uint32 Code, struct _EXCEPTION_POINTERS *Ep)
