[NO NON-PII COMMIT MESSAGE]

diff --git a/Engine/Plugins/Enterprise/DatasmithCADImporter/Source/DatasmithPLMXMLTranslator/Private/DatasmithPlmXmlImporter.cpp b/Engine/Plugins/Enterprise/DatasmithCADImporter/Source/DatasmithPLMXMLTranslator/Private/DatasmithPlmXmlImporter.cpp
--- a/Engine/Plugins/Enterprise/DatasmithCADImporter/Source/DatasmithPLMXMLTranslator/Private/DatasmithPlmXmlImporter.cpp
+++ b/Engine/Plugins/Enterprise/DatasmithCADImporter/Source/DatasmithPLMXMLTranslator/Private/DatasmithPlmXmlImporter.cpp
@@ -1139,7 +1139,20 @@ namespace PlmXml
 			return false;
 		}
 
+		// PLMXMLSchema.xsd
+		// rootInstanceRef: if the graph of Instances and StructureRevisionViews
+		//    in the InstanceGraph has a single root, which is an Instance,
+		//    this may be use to indicate it.
+		// rootRefs : however, in general there may be more than one root, and the
+		//    roots may be Instances or StructureRevisionViews.The use of
+		//    this attribute to specify the root(s) is preferred.
+		// 
+		// So the search of rootRefs is done first. If it doesn't exist, then rootInstanceRef is search
 		ImportContext.InstanceGraphRootRefs = PlmXml::GetAttributeIDREFS(InstanceGraphNode, TEXT("rootRefs"));
+		if (!ImportContext.InstanceGraphRootRefs.Num())
+		{
+			ImportContext.InstanceGraphRootRefs = PlmXml::GetAttributeIDREFS(InstanceGraphNode, TEXT("rootInstanceRef"));
+		}
 
 		if (!ImportContext.InstanceGraphRootRefs.Num())
 		{
