Remove KernelIO SDK: step 2 - remove Interface sources

diff --git a/Engine/Config/ConsoleVariables.ini b/Engine/Config/ConsoleVariables.ini
--- a/Engine/Config/ConsoleVariables.ini
+++ b/Engine/Config/ConsoleVariables.ini
@@ -128,9 +128,6 @@ fx.SkipVectorVMBackendOptimizations=1
 ; Uncomment to enable/disable CADKernel tessellation. If disable, CAD SDK tessellator is used.
 ;ds.CADTranslator.DisableCADKernelTessellation=false
 
-; Uncomment to select the SDK to load CAD files: "KernelIO" or "TechSoft". Default is "KernelIO"
-;ds.CADTranslator.CADLibrary="KernelIO"
-
 ; Uncomment to force to import embedded tessellation instead of meshing BRep model. Default is false
 ;ds.CADTranslator.PreferJtFileEmbeddedTessellation=true
 
diff --git a/Engine/Plugins/Enterprise/DatasmithCADImporter/DatasmithCADImporter.uplugin b/Engine/Plugins/Enterprise/DatasmithCADImporter/DatasmithCADImporter.uplugin
--- a/Engine/Plugins/Enterprise/DatasmithCADImporter/DatasmithCADImporter.uplugin
+++ b/Engine/Plugins/Enterprise/DatasmithCADImporter/DatasmithCADImporter.uplugin
@@ -144,14 +144,6 @@
 				"Linux"
 			]
 		},
-		{
-			"Name": "CoreTechSurface",
-			"Type": "Runtime",
-			"LoadingPhase": "Default",
-			"PlatformAllowList": [
-				"Win64"
-			]
-		},
 		{
 			"Name": "DatasmithPLMXMLTranslator",
 			"Type": "Editor",
diff --git a/Engine/Plugins/Enterprise/DatasmithCADImporter/Source/CADInterfaces/CADInterfaces.Build.cs b/Engine/Plugins/Enterprise/DatasmithCADImporter/Source/CADInterfaces/CADInterfaces.Build.cs
--- a/Engine/Plugins/Enterprise/DatasmithCADImporter/Source/CADInterfaces/CADInterfaces.Build.cs
+++ b/Engine/Plugins/Enterprise/DatasmithCADImporter/Source/CADInterfaces/CADInterfaces.Build.cs
@@ -23,15 +23,9 @@ namespace UnrealBuildTool.Rules
 				}
 			);
 
-			// CAD library is only available if CoreTech is available too
-			bool bHasCoretech = System.Type.GetType("CoreTech") != null;
+			// CAD library is only available if TechSoft is available too
 			bool bHasTechSoft = System.Type.GetType("TechSoft") != null;
 
-			if (Target.Platform == UnrealTargetPlatform.Win64 && bHasCoretech)
-			{
-				PublicDependencyModuleNames.Add("CoreTech");
-			}
-
 			if ((Target.Platform == UnrealTargetPlatform.Win64 || Target.Platform == UnrealTargetPlatform.Linux) && bHasTechSoft)
 			{
 				PublicDependencyModuleNames.Add("TechSoft");
diff --git a/Engine/Plugins/Enterprise/DatasmithCADImporter/Source/CADInterfaces/Private/CTSession.cpp b/Engine/Plugins/Enterprise/DatasmithCADImporter/Source/CADInterfaces/Private/CTSession.cpp
deleted file mode 100644
diff --git a/Engine/Plugins/Enterprise/DatasmithCADImporter/Source/CADInterfaces/Private/CoreTechBridge.cpp b/Engine/Plugins/Enterprise/DatasmithCADImporter/Source/CADInterfaces/Private/CoreTechBridge.cpp
deleted file mode 100644
diff --git a/Engine/Plugins/Enterprise/DatasmithCADImporter/Source/CADInterfaces/Private/CoreTechBridge.h b/Engine/Plugins/Enterprise/DatasmithCADImporter/Source/CADInterfaces/Private/CoreTechBridge.h
deleted file mode 100644
diff --git a/Engine/Plugins/Enterprise/DatasmithCADImporter/Source/CADInterfaces/Private/CoreTechFileParser.h b/Engine/Plugins/Enterprise/DatasmithCADImporter/Source/CADInterfaces/Private/CoreTechFileParser.h
deleted file mode 100644
diff --git a/Engine/Plugins/Enterprise/DatasmithCADImporter/Source/CADInterfaces/Private/CoreTechFileParserImpl.cpp b/Engine/Plugins/Enterprise/DatasmithCADImporter/Source/CADInterfaces/Private/CoreTechFileParserImpl.cpp
deleted file mode 100644
diff --git a/Engine/Plugins/Enterprise/DatasmithCADImporter/Source/CADInterfaces/Private/CoreTechInterfaceImpl.cpp b/Engine/Plugins/Enterprise/DatasmithCADImporter/Source/CADInterfaces/Private/CoreTechInterfaceImpl.cpp
deleted file mode 100644
diff --git a/Engine/Plugins/Enterprise/DatasmithCADImporter/Source/CADInterfaces/Private/CoreTechInterfaceImpl.h b/Engine/Plugins/Enterprise/DatasmithCADImporter/Source/CADInterfaces/Private/CoreTechInterfaceImpl.h
deleted file mode 100644
diff --git a/Engine/Plugins/Enterprise/DatasmithCADImporter/Source/CADInterfaces/Private/CoreTechTypes.cpp b/Engine/Plugins/Enterprise/DatasmithCADImporter/Source/CADInterfaces/Private/CoreTechTypes.cpp
deleted file mode 100644
diff --git a/Engine/Plugins/Enterprise/DatasmithCADImporter/Source/CADInterfaces/Public/CTSession.h b/Engine/Plugins/Enterprise/DatasmithCADImporter/Source/CADInterfaces/Public/CTSession.h
deleted file mode 100644
diff --git a/Engine/Plugins/Enterprise/DatasmithCADImporter/Source/CADInterfaces/Public/CoreTechTypes.h b/Engine/Plugins/Enterprise/DatasmithCADImporter/Source/CADInterfaces/Public/CoreTechTypes.h
deleted file mode 100644
diff --git a/Engine/Plugins/Enterprise/DatasmithCADImporter/Source/CADLibrary/CADLibrary.Build.cs b/Engine/Plugins/Enterprise/DatasmithCADImporter/Source/CADLibrary/CADLibrary.Build.cs
--- a/Engine/Plugins/Enterprise/DatasmithCADImporter/Source/CADLibrary/CADLibrary.Build.cs
+++ b/Engine/Plugins/Enterprise/DatasmithCADImporter/Source/CADLibrary/CADLibrary.Build.cs
@@ -24,18 +24,5 @@ public class CADLibrary : ModuleRules
 
 		// Support for Windows only
 		bool bIsPlateformSupported = Target.Platform == UnrealTargetPlatform.Win64;
-
-		// CAD library is only available if CoreTech is available too
-		bool bHasCoretech = System.Type.GetType("CoreTech") != null;
-
-		if (bIsPlateformSupported && bHasCoretech)
-		{
-			PublicDependencyModuleNames.AddRange(
-				new string[]
-				{
-					"CoreTech",
-				}
-			);
-		}
 	}
 }
diff --git a/Engine/Plugins/Enterprise/DatasmithCADImporter/Source/CoreTechSurface/CoreTechSurface.Build.cs b/Engine/Plugins/Enterprise/DatasmithCADImporter/Source/CoreTechSurface/CoreTechSurface.Build.cs
deleted file mode 100644
diff --git a/Engine/Plugins/Enterprise/DatasmithCADImporter/Source/CoreTechSurface/Private/CoreTechSurfaceExtension.cpp b/Engine/Plugins/Enterprise/DatasmithCADImporter/Source/CoreTechSurface/Private/CoreTechSurfaceExtension.cpp
deleted file mode 100644
diff --git a/Engine/Plugins/Enterprise/DatasmithCADImporter/Source/CoreTechSurface/Private/CoreTechSurfaceHelper.cpp b/Engine/Plugins/Enterprise/DatasmithCADImporter/Source/CoreTechSurface/Private/CoreTechSurfaceHelper.cpp
deleted file mode 100644
diff --git a/Engine/Plugins/Enterprise/DatasmithCADImporter/Source/CoreTechSurface/Private/CoreTechSurfaceModule.cpp b/Engine/Plugins/Enterprise/DatasmithCADImporter/Source/CoreTechSurface/Private/CoreTechSurfaceModule.cpp
deleted file mode 100644
diff --git a/Engine/Plugins/Enterprise/DatasmithCADImporter/Source/CoreTechSurface/Public/CADModelToCoretechConverterBase.h b/Engine/Plugins/Enterprise/DatasmithCADImporter/Source/CoreTechSurface/Public/CADModelToCoretechConverterBase.h
deleted file mode 100644
diff --git a/Engine/Plugins/Enterprise/DatasmithCADImporter/Source/CoreTechSurface/Public/CoreTechSurfaceExtension.h b/Engine/Plugins/Enterprise/DatasmithCADImporter/Source/CoreTechSurface/Public/CoreTechSurfaceExtension.h
deleted file mode 100644
diff --git a/Engine/Plugins/Enterprise/DatasmithCADImporter/Source/CoreTechSurface/Public/CoreTechSurfaceHelper.h b/Engine/Plugins/Enterprise/DatasmithCADImporter/Source/CoreTechSurface/Public/CoreTechSurfaceHelper.h
deleted file mode 100644
diff --git a/Engine/Plugins/Enterprise/DatasmithCADImporter/Source/CoreTechSurface/Public/CoreTechSurfaceModule.h b/Engine/Plugins/Enterprise/DatasmithCADImporter/Source/CoreTechSurface/Public/CoreTechSurfaceModule.h
deleted file mode 100644
diff --git a/Engine/Plugins/Enterprise/DatasmithCADImporter/Source/DatasmithCADTranslator/DatasmithCADTranslator.Build.cs b/Engine/Plugins/Enterprise/DatasmithCADImporter/Source/DatasmithCADTranslator/DatasmithCADTranslator.Build.cs
--- a/Engine/Plugins/Enterprise/DatasmithCADImporter/Source/DatasmithCADTranslator/DatasmithCADTranslator.Build.cs
+++ b/Engine/Plugins/Enterprise/DatasmithCADImporter/Source/DatasmithCADTranslator/DatasmithCADTranslator.Build.cs
@@ -23,7 +23,6 @@ namespace UnrealBuildTool.Rules
 					"CADInterfaces",
 					"CADLibrary",
 					"CADTools",
-					"CoreTechSurface",
 					"DatasmithCore",
 					"DatasmithContent",
 					"DatasmithDispatcher",
@@ -34,11 +33,6 @@ namespace UnrealBuildTool.Rules
 					"Sockets",
 				}
 			);
-
-			if (System.Type.GetType("CoreTech") != null)
-			{
-				PublicDependencyModuleNames.Add("CoreTech");
-			}
 		}
 	}
 }
diff --git a/Engine/Plugins/Enterprise/DatasmithCADImporter/Source/DatasmithDispatcher/DatasmithDispatcher.Build.cs b/Engine/Plugins/Enterprise/DatasmithCADImporter/Source/DatasmithDispatcher/DatasmithDispatcher.Build.cs
--- a/Engine/Plugins/Enterprise/DatasmithCADImporter/Source/DatasmithDispatcher/DatasmithDispatcher.Build.cs
+++ b/Engine/Plugins/Enterprise/DatasmithCADImporter/Source/DatasmithDispatcher/DatasmithDispatcher.Build.cs
@@ -23,11 +23,6 @@ namespace UnrealBuildTool.Rules
 				{
                 }
 			);
-
-            if (System.Type.GetType("CoreTech") != null)
-            {
-                PublicDependencyModuleNames.Add("CoreTech");
-            }
         }
     }
 }
diff --git a/Engine/Plugins/Enterprise/DatasmithCADImporter/Source/DatasmithOpenNurbsTranslator/DatasmithOpenNurbsTranslator.Build.cs b/Engine/Plugins/Enterprise/DatasmithCADImporter/Source/DatasmithOpenNurbsTranslator/DatasmithOpenNurbsTranslator.Build.cs
--- a/Engine/Plugins/Enterprise/DatasmithCADImporter/Source/DatasmithOpenNurbsTranslator/DatasmithOpenNurbsTranslator.Build.cs
+++ b/Engine/Plugins/Enterprise/DatasmithCADImporter/Source/DatasmithOpenNurbsTranslator/DatasmithOpenNurbsTranslator.Build.cs
@@ -25,7 +25,6 @@ namespace UnrealBuildTool.Rules
 					"CADKernelSurface",
 					"CADLibrary",
 					"CADTools",
-					"CoreTechSurface",
 					"DatasmithCore",
 					"DatasmithContent",
 					"DatasmithTranslator",
@@ -41,11 +40,6 @@ namespace UnrealBuildTool.Rules
 				PrivateDependencyModuleNames.Add("OpenNurbs6");
 			}
 
-            if (System.Type.GetType("CoreTech") != null)
-            {
-				PrivateDependencyModuleNames.Add("CoreTech");
-            }
-
             if (Target.Type == TargetType.Editor)
             {
 				PrivateDependencyModuleNames.AddRange(
diff --git a/Engine/Plugins/Enterprise/DatasmithCADImporter/Source/DatasmithOpenNurbsTranslator/Private/OpenNurbsBRepToCoretechConverter.cpp b/Engine/Plugins/Enterprise/DatasmithCADImporter/Source/DatasmithOpenNurbsTranslator/Private/OpenNurbsBRepToCoretechConverter.cpp
deleted file mode 100644
diff --git a/Engine/Plugins/Enterprise/DatasmithCADImporter/Source/DatasmithOpenNurbsTranslator/Private/OpenNurbsBRepToCoretechConverter.h b/Engine/Plugins/Enterprise/DatasmithCADImporter/Source/DatasmithOpenNurbsTranslator/Private/OpenNurbsBRepToCoretechConverter.h
deleted file mode 100644
diff --git a/Engine/Plugins/Enterprise/DatasmithCADImporter/Source/DatasmithPLMXMLTranslator/DatasmithPLMXMLTranslator.Build.cs b/Engine/Plugins/Enterprise/DatasmithCADImporter/Source/DatasmithPLMXMLTranslator/DatasmithPLMXMLTranslator.Build.cs
--- a/Engine/Plugins/Enterprise/DatasmithCADImporter/Source/DatasmithPLMXMLTranslator/DatasmithPLMXMLTranslator.Build.cs
+++ b/Engine/Plugins/Enterprise/DatasmithCADImporter/Source/DatasmithPLMXMLTranslator/DatasmithPLMXMLTranslator.Build.cs
@@ -25,7 +25,6 @@ namespace UnrealBuildTool.Rules
 					"CADKernelSurface",
 					"CADLibrary",
 					"CADTools",
-					"CoreTechSurface",
 					"DatasmithCADTranslator", // for DatasmithMeshBuilder
 					"DatasmithContent",
 					"DatasmithCore",
@@ -49,11 +48,6 @@ namespace UnrealBuildTool.Rules
 					}
 				);
 			}
-
-			if (System.Type.GetType("CoreTech") != null)
-			{
-				PrivateDependencyModuleNames.Add("CoreTech");
-			}
 		}
 	}
 }
diff --git a/Engine/Plugins/Enterprise/DatasmithCADImporter/Source/DatasmithWireTranslator/DatasmithWireTranslator2020.Build.cs b/Engine/Plugins/Enterprise/DatasmithCADImporter/Source/DatasmithWireTranslator/DatasmithWireTranslator2020.Build.cs
--- a/Engine/Plugins/Enterprise/DatasmithCADImporter/Source/DatasmithWireTranslator/DatasmithWireTranslator2020.Build.cs
+++ b/Engine/Plugins/Enterprise/DatasmithCADImporter/Source/DatasmithWireTranslator/DatasmithWireTranslator2020.Build.cs
@@ -26,7 +26,6 @@ public abstract class DatasmithWireTranslatorBase : ModuleRules
 				"CADKernelSurface",
 				"CADLibrary",
 				"CADTools",
-				"CoreTechSurface",
 				"DatasmithContent",
 				"DatasmithCore",
 				"DatasmithTranslator",
@@ -56,11 +55,6 @@ public abstract class DatasmithWireTranslatorBase : ModuleRules
 		{
 			PrivateDependencyModuleNames.Add(GetAliasVersion());
 		}
-
-		if (System.Type.GetType("CoreTech") != null)
-		{
-			PrivateDependencyModuleNames.Add("CoreTech");
-		}
 	}
 
 	public abstract string GetAliasVersion();
diff --git a/Engine/Plugins/Enterprise/DatasmithCADImporter/Source/DatasmithWireTranslator/Private/AliasModelToCoretechConverter.cpp b/Engine/Plugins/Enterprise/DatasmithCADImporter/Source/DatasmithWireTranslator/Private/AliasModelToCoretechConverter.cpp
deleted file mode 100644
diff --git a/Engine/Plugins/Enterprise/DatasmithCADImporter/Source/DatasmithWireTranslator/Private/AliasModelToCoretechConverter.h b/Engine/Plugins/Enterprise/DatasmithCADImporter/Source/DatasmithWireTranslator/Private/AliasModelToCoretechConverter.h
deleted file mode 100644
diff --git a/Engine/Plugins/Enterprise/DatasmithCADImporter/Source/ParametricSurface/Private/CoreTech/CoreTechSurfaceData.cpp b/Engine/Plugins/Enterprise/DatasmithCADImporter/Source/ParametricSurface/Private/CoreTech/CoreTechSurfaceData.cpp
deleted file mode 100644
diff --git a/Engine/Plugins/Enterprise/DatasmithCADImporter/Source/ParametricSurface/Private/CoreTech/CoreTechSurfaceData.h b/Engine/Plugins/Enterprise/DatasmithCADImporter/Source/ParametricSurface/Private/CoreTech/CoreTechSurfaceData.h
deleted file mode 100644
