Improve  CADImport naming

diff --git a/Engine/Plugins/Enterprise/DatasmithCADImporter/Source/CADInterfaces/Private/CADSceneGraph.cpp b/Engine/Plugins/Enterprise/DatasmithCADImporter/Source/CADInterfaces/Private/CADSceneGraph.cpp
--- a/Engine/Plugins/Enterprise/DatasmithCADImporter/Source/CADInterfaces/Private/CADSceneGraph.cpp
+++ b/Engine/Plugins/Enterprise/DatasmithCADImporter/Source/CADInterfaces/Private/CADSceneGraph.cpp
@@ -45,7 +45,7 @@ FArchive& operator<<(FArchive& Ar, FArchiveBody& Body)
 	Ar << Body.MaterialFaceSet;
 	Ar << Body.ColorFaceSet;
 	Ar << Body.ParentId;
-	Ar << Body.MeshActorName;
+	Ar << Body.MeshActorUId;
 	Ar << Body.BodyUnit;
 
 	return Ar;
diff --git a/Engine/Plugins/Enterprise/DatasmithCADImporter/Source/CADInterfaces/Private/TechSoftFileParser.cpp b/Engine/Plugins/Enterprise/DatasmithCADImporter/Source/CADInterfaces/Private/TechSoftFileParser.cpp
--- a/Engine/Plugins/Enterprise/DatasmithCADImporter/Source/CADInterfaces/Private/TechSoftFileParser.cpp
+++ b/Engine/Plugins/Enterprise/DatasmithCADImporter/Source/CADInterfaces/Private/TechSoftFileParser.cpp
@@ -453,38 +453,38 @@ void FTechSoftFileParser::GenerateBodyMesh(A3DRiRepresentationItem* Representati
 		// the mesh of the body is empty, the body is deleted.
 		// Todo (jira UETOOL-5148): add a boolean in Body to flag that the body should not be build
 		Body.ParentId = 0;
-		Body.MeshActorName = 0;
+		Body.MeshActorUId = 0;
 	}
 
 	// Convert material
-	FCADUUID DefaultColorName = Body.ColorFaceSet.Num() > 0 ? *Body.ColorFaceSet.begin() : 0;
-	FCADUUID DefaultMaterialName = Body.MaterialFaceSet.Num() > 0 ? *Body.MaterialFaceSet.begin() : 0;
+	FCadUuid DefaultColorUId = Body.ColorFaceSet.Num() > 0 ? *Body.ColorFaceSet.begin() : 0;
+	FCadUuid DefaultMaterialUId = Body.MaterialFaceSet.Num() > 0 ? *Body.MaterialFaceSet.begin() : 0;
 
 	for (FTessellationData& Tessellation : BodyMesh.Faces)
 	{
-		FCADUUID ColorName = DefaultColorName;
-		FCADUUID MaterialName = DefaultMaterialName;
+		FCadUuid ColorUId = DefaultColorUId;
+		FCadUuid MaterialUId = DefaultMaterialUId;
 
 		// Extract proper color or material based on style index
-		uint32 CachedStyleIndex = Tessellation.MaterialName;
-		Tessellation.MaterialName = 0;
+		uint32 CachedStyleIndex = Tessellation.MaterialUId;
+		Tessellation.MaterialUId = 0;
 
 		constexpr uint32 GraphStyleDataDefaultValue =  65535;
 		if (CachedStyleIndex != GraphStyleDataDefaultValue)
 		{
-			ExtractGraphStyleProperties(CachedStyleIndex, ColorName, MaterialName);
+			ExtractGraphStyleProperties(CachedStyleIndex, ColorUId, MaterialUId);
 		}
 
-		if (ColorName)
+		if (ColorUId)
 		{
-			Tessellation.ColorName = ColorName;
-			BodyMesh.ColorSet.Add(ColorName);
+			Tessellation.ColorUId = ColorUId;
+			BodyMesh.ColorSet.Add(ColorUId);
 		}
 
-		if (MaterialName)
+		if (MaterialUId)
 		{
-			Tessellation.MaterialName = MaterialName;
-			BodyMesh.MaterialSet.Add(MaterialName);
+			Tessellation.MaterialUId = MaterialUId;
+			BodyMesh.MaterialSet.Add(MaterialUId);
 		}
 	}
 
@@ -497,7 +497,7 @@ void FTechSoftFileParser::GenerateBodyMesh(A3DRiRepresentationItem* Representati
 
 	if (Type == kA3DTypeRiBrepModel)
 	{
-		FString FilePath = CADFileData.GetBodyCachePath(Body.MeshActorName);
+		FString FilePath = CADFileData.GetBodyCachePath(Body.MeshActorUId);
 		if (!FilePath.IsEmpty())
 		{
 			TSharedPtr<FJsonObject> JsonObject = MakeShared<FJsonObject>();
@@ -505,8 +505,8 @@ void FTechSoftFileParser::GenerateBodyMesh(A3DRiRepresentationItem* Representati
 			// Save body unit and default color and material attributes in a json string
 			// This will be used when the file is reloaded
 			JsonObject->SetNumberField(JSON_ENTRY_BODY_UNIT, Body.BodyUnit);
-			JsonObject->SetNumberField(JSON_ENTRY_COLOR_NAME, DefaultColorName);
-			JsonObject->SetNumberField(JSON_ENTRY_MATERIAL_NAME, DefaultMaterialName);
+			JsonObject->SetNumberField(JSON_ENTRY_COLOR_NAME, DefaultColorUId);
+			JsonObject->SetNumberField(JSON_ENTRY_MATERIAL_NAME, DefaultMaterialUId);
 
 			FString JsonString;
 			TSharedRef< TJsonWriter< TCHAR, TPrettyJsonPrintPolicy<TCHAR> > > JsonWriter = TJsonWriterFactory< TCHAR, TPrettyJsonPrintPolicy<TCHAR> >::Create(&JsonString);
@@ -1753,7 +1753,7 @@ void FTechSoftFileParser::ExtractSpecificMetaData(const A3DAsmProductOccurrence*
 
 FArchiveColor& FTechSoftFileParser::FindOrAddColor(uint32 ColorIndex, uint8 Alpha)
 {
-	uint32 ColorHId = BuildColorId(ColorIndex, Alpha);
+	FMaterialUId ColorHId = BuildColorFastUId(ColorIndex, Alpha);
 	if (FArchiveColor* Color = CADFileData.FindColor(ColorHId))
 	{
 		return *Color;
@@ -1763,7 +1763,7 @@ FArchiveColor& FTechSoftFileParser::FindOrAddColor(uint32 ColorIndex, uint8 Alph
 	NewColor.Color = TechSoftUtils::GetColorAt(ColorIndex);
 	NewColor.Color.A = Alpha;
 	
-	NewColor.UEMaterialName = BuildColorName(NewColor.Color);
+	NewColor.UEMaterialName = BuildColorUId(NewColor.Color);
 	return NewColor;
 }
 
@@ -1787,7 +1787,7 @@ FArchiveMaterial& FTechSoftFileParser::AddMaterialAt(uint32 MaterialIndexToSave,
 		// Material.Emissive = GetColor(MaterialData->m_uiEmissive);
 		// Material.Reflexion;
 	}
-	NewMaterial.UEMaterialName = BuildMaterialName(Material);
+	NewMaterial.UEMaterialName = BuildMaterialUId(Material);
 	return NewMaterial;
 }
 
@@ -1850,8 +1850,8 @@ void FTechSoftFileParser::ExtractGraphicProperties(const A3DGraphics* Graphics,
 		return;
 	}
 
-	FCADUUID& ColorName = OutMetaData.ColorName;
-	FCADUUID& MaterialName = OutMetaData.MaterialName;
+	FCadUuid& ColorName = OutMetaData.ColorName;
+	FCadUuid& MaterialName = OutMetaData.MaterialName;
 
 	ExtractGraphStyleProperties(GraphicsData->m_uiStyleIndex, ColorName, MaterialName);
 
@@ -1877,7 +1877,7 @@ void FTechSoftFileParser::ExtractGraphicProperties(const A3DGraphics* Graphics,
 
 // Please review TechSoftUtils::GetMaterialValues if anything changes
 // in this method or the methods it calls
-void FTechSoftFileParser::ExtractGraphStyleProperties(uint32 StyleIndex, FCADUUID& OutColorName, FCADUUID& OutMaterialName)
+void FTechSoftFileParser::ExtractGraphStyleProperties(uint32 StyleIndex, FCadUuid& OutColorName, FCadUuid& OutMaterialName)
 {
 	TUniqueTSObjFromIndex<A3DGraphStyleData> GraphStyleData(StyleIndex);
 
diff --git a/Engine/Plugins/Enterprise/DatasmithCADImporter/Source/CADInterfaces/Private/TechSoftFileParser.h b/Engine/Plugins/Enterprise/DatasmithCADImporter/Source/CADInterfaces/Private/TechSoftFileParser.h
--- a/Engine/Plugins/Enterprise/DatasmithCADImporter/Source/CADInterfaces/Private/TechSoftFileParser.h
+++ b/Engine/Plugins/Enterprise/DatasmithCADImporter/Source/CADInterfaces/Private/TechSoftFileParser.h
@@ -129,7 +129,7 @@ private:
 	 * ColorName and MaterialName have to be initialized before.
 	 * This method update the value ColorName or MaterialName accordingly of the GraphStyleData type (material or color)
 	 */
-	void ExtractGraphStyleProperties(uint32 StyleIndex, FCADUUID& ColorName, FCADUUID& MaterialName);
+	void ExtractGraphStyleProperties(uint32 StyleIndex, FCadUuid& ColorName, FCadUuid& MaterialName);
 	FArchiveColor& FindOrAddColor(uint32 ColorIndex, uint8 Alpha);
 	FArchiveMaterial& FindOrAddMaterial(uint32 MaterialId, const A3DGraphStyleData& GraphStyleData);
 
diff --git a/Engine/Plugins/Enterprise/DatasmithCADImporter/Source/CADInterfaces/Private/TechSoftFileParserCADKernelTessellator.cpp b/Engine/Plugins/Enterprise/DatasmithCADImporter/Source/CADInterfaces/Private/TechSoftFileParserCADKernelTessellator.cpp
--- a/Engine/Plugins/Enterprise/DatasmithCADImporter/Source/CADInterfaces/Private/TechSoftFileParserCADKernelTessellator.cpp
+++ b/Engine/Plugins/Enterprise/DatasmithCADImporter/Source/CADInterfaces/Private/TechSoftFileParserCADKernelTessellator.cpp
@@ -181,7 +181,7 @@ void FTechSoftFileParserCADKernelTessellator::SewAndMesh(TArray<A3DRiRepresentat
 		FArchiveBody& ArchiveBody = CADFileData.GetBodyAt(ArchiveBodyIndex);
 		ArchiveBody.MetaData.Empty();
 		ArchiveBody.ParentId = 0;
-		ArchiveBody.MeshActorName = 0;
+		ArchiveBody.MeshActorUId = 0;
 	}
 }
 
@@ -218,7 +218,7 @@ void FTechSoftFileParserCADKernelTessellator::MeshAndGetTessellation(CADKernel::
 	// Save Body in CADKernelArchive file for re-tessellation
 	if (CADFileData.IsCacheDefined())
 	{
-		FString BodyFilePath = CADFileData.GetBodyCachePath(ArchiveBody.MeshActorName);
+		FString BodyFilePath = CADFileData.GetBodyCachePath(ArchiveBody.MeshActorUId);
 		CADKernelSession.SaveDatabase(*BodyFilePath);
 	}
 
diff --git a/Engine/Plugins/Enterprise/DatasmithCADImporter/Source/CADInterfaces/Private/TechSoftInterfaceUtils.inl b/Engine/Plugins/Enterprise/DatasmithCADImporter/Source/CADInterfaces/Private/TechSoftInterfaceUtils.inl
--- a/Engine/Plugins/Enterprise/DatasmithCADImporter/Source/CADInterfaces/Private/TechSoftInterfaceUtils.inl
+++ b/Engine/Plugins/Enterprise/DatasmithCADImporter/Source/CADInterfaces/Private/TechSoftInterfaceUtils.inl
@@ -169,11 +169,11 @@ private:
 			// there is a bijection between A3DTess3DData->m_psFaceTessData and A3DTopoShellData->m_ppFaces
 			Tessellation.PatchId = Index;
 
-			Tessellation.MaterialName = FTechSoftInterface::InvalidScriptIndex;
+			Tessellation.MaterialUId = FTechSoftInterface::InvalidScriptIndex;
 			if (FaceTessData.m_uiStyleIndexesSize > 0)
 			{
 				// Store the StyleIndex on the MaterialName. It will be processed after tessellation
-				Tessellation.MaterialName = FaceTessData.m_puiStyleIndexes[0];
+				Tessellation.MaterialUId = FaceTessData.m_puiStyleIndexes[0];
 			}
 
 			// Pre-allocate memory for triangles' data
diff --git a/Engine/Plugins/Enterprise/DatasmithCADImporter/Source/CADInterfaces/Private/TechSoftUtils.cpp b/Engine/Plugins/Enterprise/DatasmithCADImporter/Source/CADInterfaces/Private/TechSoftUtils.cpp
--- a/Engine/Plugins/Enterprise/DatasmithCADImporter/Source/CADInterfaces/Private/TechSoftUtils.cpp
+++ b/Engine/Plugins/Enterprise/DatasmithCADImporter/Source/CADInterfaces/Private/TechSoftUtils.cpp
@@ -262,7 +262,7 @@ void BuildCADMaterial(uint32 MaterialIndex, const A3DGraphStyleData& GraphStyleD
 #endif
 
 // Replicates the logic in FTechSoftFileParser::ExtractGraphStyleProperties
-void GetMaterialValues(uint32 StyleIndex, FCADUUID& OutColorName, FCADUUID& OutMaterialName)
+void GetMaterialValues(uint32 StyleIndex, FCadUuid& OutColorName, FCadUuid& OutMaterialName)
 {
 #if defined USE_TECHSOFT_SDK && !defined CADKERNEL_DEV
 	TUniqueTSObjFromIndex<A3DGraphStyleData> GraphStyleData(StyleIndex);
@@ -275,7 +275,7 @@ void GetMaterialValues(uint32 StyleIndex, FCADUUID& OutColorName, FCADUUID& OutM
 
 			BuildCADMaterial(GraphStyleData->m_uiRgbColorIndex, *GraphStyleData, Material);
 
-			OutMaterialName = BuildMaterialName(Material);
+			OutMaterialName = BuildMaterialUId(Material);
 		}
 		else
 		{
@@ -286,7 +286,7 @@ void GetMaterialValues(uint32 StyleIndex, FCADUUID& OutColorName, FCADUUID& OutM
 				const uint8 Alpha = GraphStyleData->m_bIsTransparencyDefined ? (255 - GraphStyleData->m_ucTransparency) : 255;
 				const FColor ColorValue((uint8)(ColorData->m_dRed * 255), (uint8)(ColorData->m_dGreen * 255), (uint8)(ColorData->m_dBlue * 255), Alpha);
 
-				OutColorName = BuildColorName(ColorValue);
+				OutColorName = BuildColorUId(ColorValue);
 			}
 		}
 	}
@@ -295,8 +295,8 @@ void GetMaterialValues(uint32 StyleIndex, FCADUUID& OutColorName, FCADUUID& OutM
 
 void RestoreMaterials(const TSharedPtr<FJsonObject>& DefaultValues, FBodyMesh& BodyMesh)
 {
-	FCADUUID DefaultColorName = 0;
-	FCADUUID DefaultMaterialName = 0;
+	FCadUuid DefaultColorName = 0;
+	FCadUuid DefaultMaterialName = 0;
 
 	DefaultValues->TryGetNumberField(JSON_ENTRY_COLOR_NAME, DefaultColorName);
 	DefaultValues->TryGetNumberField(JSON_ENTRY_MATERIAL_NAME, DefaultMaterialName);
@@ -307,23 +307,23 @@ void RestoreMaterials(const TSharedPtr<FJsonObject>& DefaultValues, FBodyMesh& B
 	for (FTessellationData& Tessellation : BodyMesh.Faces)
 	{
 		// Extract proper color or material based on style index
-		uint32 CachedStyleIndex = Tessellation.MaterialName;
-		Tessellation.MaterialName = 0;
+		uint32 CachedStyleIndex = Tessellation.MaterialUId;
+		Tessellation.MaterialUId = 0;
 
-		FCADUUID ColorName = DefaultColorName;
-		FCADUUID MaterialName = DefaultMaterialName;
+		FCadUuid ColorName = DefaultColorName;
+		FCadUuid MaterialName = DefaultMaterialName;
 
 		GetMaterialValues(CachedStyleIndex, ColorName, MaterialName);
 
 		if (ColorName)
 		{
-			Tessellation.ColorName = ColorName;
+			Tessellation.ColorUId = ColorName;
 			BodyMesh.ColorSet.Add(ColorName);
 		}
 
 		if (MaterialName)
 		{
-			Tessellation.MaterialName = MaterialName;
+			Tessellation.MaterialUId = MaterialName;
 			BodyMesh.MaterialSet.Add(MaterialName);
 		}
 	}
diff --git a/Engine/Plugins/Enterprise/DatasmithCADImporter/Source/CADInterfaces/Public/CADFileData.h b/Engine/Plugins/Enterprise/DatasmithCADImporter/Source/CADInterfaces/Public/CADFileData.h
--- a/Engine/Plugins/Enterprise/DatasmithCADImporter/Source/CADInterfaces/Public/CADFileData.h
+++ b/Engine/Plugins/Enterprise/DatasmithCADImporter/Source/CADInterfaces/Public/CADFileData.h
@@ -234,8 +234,8 @@ public:
 	FBodyMesh& AddBodyMesh(FCadId BodyId, FArchiveBody& Body)
 	{
 		FBodyMesh& BodyMesh = BodyMeshes.Emplace_GetRef(BodyId);
-		BodyMesh.MeshActorName = GetStaticMeshHash(BodyId);
-		Body.MeshActorName = BodyMesh.MeshActorName;
+		BodyMesh.MeshActorUId = GetStaticMeshHash(BodyId);
+		Body.MeshActorUId = BodyMesh.MeshActorUId;
 		return BodyMesh;
 	}
 
@@ -275,22 +275,22 @@ public:
 		return SceneGraphArchive;
 	}
 
-	FArchiveMaterial* FindMaterial(uint32 MaterialId)
+	FArchiveMaterial* FindMaterial(FMaterialUId MaterialId)
 	{
 		return SceneGraphArchive.MaterialHIdToMaterial.Find(MaterialId);
 	}
 
-	FArchiveMaterial& AddMaterial(uint32 MaterialId)
+	FArchiveMaterial& AddMaterial(FMaterialUId MaterialId)
 	{
 		return SceneGraphArchive.MaterialHIdToMaterial.Emplace(MaterialId, MaterialId);
 	}
 
-	FArchiveColor* FindColor(uint32 ColorId)
+	FArchiveColor* FindColor(FMaterialUId ColorId)
 	{
 		return SceneGraphArchive.ColorHIdToColor.Find(ColorId);
 	}
 
-	FArchiveColor& AddColor(uint32 ColorId)
+	FArchiveColor& AddColor(FMaterialUId ColorId)
 	{
 		return SceneGraphArchive.ColorHIdToColor.Emplace(ColorId, ColorId);
 	}
diff --git a/Engine/Plugins/Enterprise/DatasmithCADImporter/Source/CADInterfaces/Public/CADSceneGraph.h b/Engine/Plugins/Enterprise/DatasmithCADImporter/Source/CADInterfaces/Public/CADSceneGraph.h
--- a/Engine/Plugins/Enterprise/DatasmithCADImporter/Source/CADInterfaces/Public/CADSceneGraph.h
+++ b/Engine/Plugins/Enterprise/DatasmithCADImporter/Source/CADInterfaces/Public/CADSceneGraph.h
@@ -81,18 +81,18 @@ public:
 
 public:
 	FCadId ParentId = 0;
-	FCADUUID MeshActorName = 0;
+	FCadUuid MeshActorUId = 0;
 	double BodyUnit = 1.;
 
-	TSet<FMaterialId> MaterialFaceSet;
-	TSet<FColorId> ColorFaceSet;
+	TSet<FMaterialUId> MaterialFaceSet;
+	TSet<FMaterialUId> ColorFaceSet;
 
 };
 
 class CADINTERFACES_API FArchiveColor
 {
 public:
-	FArchiveColor(FColorId Id = 0)
+	FArchiveColor(FMaterialUId Id = 0)
 		: ObjectId(Id)
 	{
 	}
@@ -100,15 +100,15 @@ public:
 	friend FArchive& operator<<(FArchive& Ar, FArchiveColor& C);
 
 public:
-	FColorId ObjectId;
+	FMaterialUId ObjectId;
 	FColor Color;
-	FCADUUID UEMaterialName;
+	FCadUuid UEMaterialName;
 };
 
 class CADINTERFACES_API FArchiveMaterial
 {
 public:
-	FArchiveMaterial(FMaterialId Id = 0)
+	FArchiveMaterial(FMaterialUId Id = 0)
 		: ObjectId(Id)
 	{
 	}
@@ -116,8 +116,8 @@ public:
 	friend FArchive& operator<<(FArchive& Ar, FArchiveMaterial& C);
 
 public:
-	FMaterialId ObjectId;
-	FCADUUID UEMaterialName;
+	FMaterialUId ObjectId;
+	FCadUuid UEMaterialName;
 	FCADMaterial Material;
 };
 
@@ -134,8 +134,8 @@ public:
 	FString ArchiveFileName;
 	FString FullPath;
 
-	TMap<FColorId, FArchiveColor> ColorHIdToColor;
-	TMap<FMaterialId, FArchiveMaterial> MaterialHIdToMaterial;
+	TMap<FMaterialUId, FArchiveColor> ColorHIdToColor;
+	TMap<FMaterialUId, FArchiveMaterial> MaterialHIdToMaterial;
 
 	TArray<FArchiveBody> Bodies;
 	TArray<FArchiveComponent> Components;
diff --git a/Engine/Plugins/Enterprise/DatasmithCADImporter/Source/CADInterfaces/Public/TechSoftUtils.h b/Engine/Plugins/Enterprise/DatasmithCADImporter/Source/CADInterfaces/Public/TechSoftUtils.h
--- a/Engine/Plugins/Enterprise/DatasmithCADImporter/Source/CADInterfaces/Public/TechSoftUtils.h
+++ b/Engine/Plugins/Enterprise/DatasmithCADImporter/Source/CADInterfaces/Public/TechSoftUtils.h
@@ -36,8 +36,8 @@ struct FEntityMetaData
 	bool bRemoved = false;
 	bool bShow = true;
 	bool bUnloaded = false;
-	FCADUUID ColorName = 0;
-	FCADUUID MaterialName = 0;
+	FCadUuid ColorName = 0;
+	FCadUuid MaterialName = 0;
 	FFileDescriptor ExternalFile;
 };
 
diff --git a/Engine/Plugins/Enterprise/DatasmithCADImporter/Source/CADLibrary/Private/CADKernelTools.cpp b/Engine/Plugins/Enterprise/DatasmithCADImporter/Source/CADLibrary/Private/CADKernelTools.cpp
--- a/Engine/Plugins/Enterprise/DatasmithCADImporter/Source/CADLibrary/Private/CADKernelTools.cpp
+++ b/Engine/Plugins/Enterprise/DatasmithCADImporter/Source/CADLibrary/Private/CADKernelTools.cpp
@@ -315,16 +315,16 @@ namespace CADLibrary
 		Tessellation.NormalArray = MoveTemp(FaceMesh.Normals);
 		Tessellation.TexCoordArray = MoveTemp(FaceMesh.UVMap);
 
-		if (FaceMaterial.Color)
+		if (FaceMaterial.ColorUId)
 		{
-			Tessellation.ColorName = FaceMaterial.Color;
-			OutBodyMesh.ColorSet.Add(FaceMaterial.Color);
+			Tessellation.ColorUId = FaceMaterial.ColorUId;
+			OutBodyMesh.ColorSet.Add(FaceMaterial.ColorUId);
 		}
 
-		if (FaceMaterial.Material)
+		if (FaceMaterial.MaterialUId)
 		{
-			Tessellation.MaterialName = FaceMaterial.Material;
-			OutBodyMesh.MaterialSet.Add(FaceMaterial.Material);
+			Tessellation.MaterialUId = FaceMaterial.MaterialUId;
+			OutBodyMesh.MaterialSet.Add(FaceMaterial.MaterialUId);
 		}
 
 		return Tessellation.VertexIndices.Num() / 3;
@@ -334,11 +334,11 @@ namespace CADLibrary
 	{
 		if (ShapeEntity.GetColorId() != 0)
 		{
-			DisplayDataId.Color = ShapeEntity.GetColorId();
+			DisplayDataId.ColorUId = ShapeEntity.GetColorId();
 		}
 		if (ShapeEntity.GetMaterialId() != 0)
 		{
-			DisplayDataId.Material = ShapeEntity.GetMaterialId();
+			DisplayDataId.MaterialUId = ShapeEntity.GetMaterialId();
 		}
 	}
 
@@ -361,7 +361,7 @@ namespace CADLibrary
 		FObjectDisplayDataId BodyMaterial;
 
 		GetDisplayDataIds(Body, BodyMaterial);
-		BodyMaterial.DefaultMaterialName = BodyMaterial.Material != 0 ? BodyMaterial.Material : BodyMaterial.Color;
+		BodyMaterial.DefaultMaterialUId = BodyMaterial.MaterialUId != 0 ? BodyMaterial.MaterialUId : BodyMaterial.ColorUId;
 
 		// Loop through the face of bodies and collect all tessellation data
 		for (const TSharedPtr<CADKernel::FShell>& Shell : Body.GetShells())
diff --git a/Engine/Plugins/Enterprise/DatasmithCADImporter/Source/CADLibrary/Private/MeshDescriptionHelper.cpp b/Engine/Plugins/Enterprise/DatasmithCADImporter/Source/CADLibrary/Private/MeshDescriptionHelper.cpp
--- a/Engine/Plugins/Enterprise/DatasmithCADImporter/Source/CADLibrary/Private/MeshDescriptionHelper.cpp
+++ b/Engine/Plugins/Enterprise/DatasmithCADImporter/Source/CADLibrary/Private/MeshDescriptionHelper.cpp
@@ -422,7 +422,7 @@ namespace CADLibrary
 		for (const FTessellationData& FaceTessellation : BodyTessellation.Faces)
 		{
 			// material is preferred over color
-			MaterialToPolygonGroupMapping.Add(FaceTessellation.MaterialName ? FaceTessellation.MaterialName : FaceTessellation.ColorName, INDEX_NONE);
+			MaterialToPolygonGroupMapping.Add(FaceTessellation.MaterialUId ? FaceTessellation.MaterialUId : FaceTessellation.ColorUId, INDEX_NONE);
 		}
 
 		// Add to the mesh, a polygon groups per material
@@ -467,7 +467,7 @@ namespace CADLibrary
 				}
 
 				// Get the polygonGroup (material is preferred over color)
-				FCADUUID GraphicName = Tessellation.MaterialName ? Tessellation.MaterialName : Tessellation.ColorName;
+				FCadUuid GraphicName = Tessellation.MaterialUId ? Tessellation.MaterialUId : Tessellation.ColorUId;
 				const FPolygonGroupID* PolygonGroupID = MaterialToPolygonGroupMapping.Find(GraphicName);
 				if (PolygonGroupID == nullptr)
 				{
@@ -643,7 +643,7 @@ namespace CADLibrary
 
 	TSharedPtr<IDatasmithUEPbrMaterialElement> CreateUEPbrMaterialFromColor(const FColor& InColor)
 	{
-		FString Name = FString::FromInt(BuildColorName(InColor));
+		FString Name = FString::FromInt(BuildColorUId(InColor));
 		FString Label = FString::Printf(TEXT("color_%02x%02x%02x%02x"), InColor.R, InColor.G, InColor.B, InColor.A);
 
 		// Take the Material diffuse color and connect it to the BaseColor of a UEPbrMaterial
@@ -679,7 +679,7 @@ namespace CADLibrary
 
 	TSharedPtr<IDatasmithUEPbrMaterialElement> CreateUEPbrMaterialFromMaterial(FCADMaterial& InMaterial, TSharedRef<IDatasmithScene> Scene)
 	{
-		FString Name = FString::FromInt(BuildMaterialName(InMaterial));
+		FString Name = FString::FromInt(BuildMaterialUId(InMaterial));
 
 		// Take the Material diffuse color and connect it to the BaseColor of a UEPbrMaterial
 		TSharedRef<IDatasmithUEPbrMaterialElement> MaterialElement = FDatasmithSceneFactory::CreateUEPbrMaterial(*Name);
diff --git a/Engine/Plugins/Enterprise/DatasmithCADImporter/Source/CADTools/Private/CADData.cpp b/Engine/Plugins/Enterprise/DatasmithCADImporter/Source/CADTools/Private/CADData.cpp
--- a/Engine/Plugins/Enterprise/DatasmithCADImporter/Source/CADTools/Private/CADData.cpp
+++ b/Engine/Plugins/Enterprise/DatasmithCADImporter/Source/CADTools/Private/CADData.cpp
@@ -12,48 +12,43 @@ uint32 MeshArchiveMagic = 345612;
 namespace CADLibrary
 {
 
-uint32 BuildColorId(uint32 ColorId, uint8 Alpha)
+FMaterialUId BuildColorFastUId(uint32 ColorId, uint8 Alpha)
 {
 	if (Alpha == 0)
 	{
 		Alpha = 1;
 	}
-	return ColorId | Alpha << 24;
+	uint32 FastColorId =  ColorId | Alpha << 24;
+	return FMath::Abs((int32)GetTypeHash(FastColorId));
 }
 
-void GetCTColorIdAlpha(FColorId ColorId, uint32& CTColorId, uint8& Alpha)
-{
-	CTColorId = ColorId & 0x00ffffff;
-	Alpha = (uint8)((ColorId & 0xff000000) >> 24);
-}
-
-int32 BuildColorName(const FColor& Color)
+FMaterialUId BuildColorUId(const FColor& Color)
 {
 	return FMath::Abs((int32)GetTypeHash(Color));
 }
 
-int32 BuildMaterialName(const FCADMaterial& Material)
+FMaterialUId BuildMaterialUId(const FCADMaterial& Material)
 {
 	using ::GetTypeHash;
 
-	uint32 MaterialName = 0;
+	uint32 MaterialUId = 0;
 	if (!Material.MaterialName.IsEmpty())
 	{
-		MaterialName = GetTypeHash(*Material.MaterialName); // we add material name because it could be used by the end user so two material with same parameters but different name are different.
+		MaterialUId = GetTypeHash(*Material.MaterialName); // we add material name because it could be used by the end user so two material with same parameters but different name are different.
 	}
 
-	MaterialName = HashCombine(MaterialName, GetTypeHash(Material.Diffuse));
-	MaterialName = HashCombine(MaterialName, GetTypeHash(Material.Ambient));
-	MaterialName = HashCombine(MaterialName, GetTypeHash(Material.Specular));
-	MaterialName = HashCombine(MaterialName, GetTypeHash((int)(Material.Shininess * 255.0)));
-	MaterialName = HashCombine(MaterialName, GetTypeHash((int)(Material.Transparency * 255.0)));
-	MaterialName = HashCombine(MaterialName, GetTypeHash((int)(Material.Reflexion * 255.0)));
+	MaterialUId = HashCombine(MaterialUId, GetTypeHash(Material.Diffuse));
+	MaterialUId = HashCombine(MaterialUId, GetTypeHash(Material.Ambient));
+	MaterialUId = HashCombine(MaterialUId, GetTypeHash(Material.Specular));
+	MaterialUId = HashCombine(MaterialUId, GetTypeHash((int)(Material.Shininess * 255.0)));
+	MaterialUId = HashCombine(MaterialUId, GetTypeHash((int)(Material.Transparency * 255.0)));
+	MaterialUId = HashCombine(MaterialUId, GetTypeHash((int)(Material.Reflexion * 255.0)));
 
 	if (!Material.TextureName.IsEmpty())
 	{
-		MaterialName = HashCombine(MaterialName, GetTypeHash(*Material.TextureName));
+		MaterialUId = HashCombine(MaterialUId, GetTypeHash(*Material.TextureName));
 	}
-	return FMath::Abs((int32) MaterialName);
+	return FMath::Abs((int32)MaterialUId);
 }
 
 FArchive& operator<<(FArchive& Ar, FCADMaterial& Material)
@@ -90,8 +85,8 @@ FArchive& operator<<(FArchive& Ar, FTessellationData& TessellationData)
 	Ar << TessellationData.NormalArray;
 	Ar << TessellationData.TexCoordArray;
 
-	Ar << TessellationData.ColorName;
-	Ar << TessellationData.MaterialName;
+	Ar << TessellationData.ColorUId;
+	Ar << TessellationData.MaterialUId;
 
 	Ar << TessellationData.PatchId;
 
@@ -106,7 +101,7 @@ FArchive& operator<<(FArchive& Ar, FBodyMesh& BodyMesh)
 
 	Ar << BodyMesh.TriangleCount;
 	Ar << BodyMesh.BodyID;
-	Ar << BodyMesh.MeshActorName;
+	Ar << BodyMesh.MeshActorUId;
 
 	Ar << BodyMesh.MaterialSet;
 	Ar << BodyMesh.ColorSet;
diff --git a/Engine/Plugins/Enterprise/DatasmithCADImporter/Source/CADTools/Public/CADData.h b/Engine/Plugins/Enterprise/DatasmithCADImporter/Source/CADTools/Public/CADData.h
--- a/Engine/Plugins/Enterprise/DatasmithCADImporter/Source/CADTools/Public/CADData.h
+++ b/Engine/Plugins/Enterprise/DatasmithCADImporter/Source/CADTools/Public/CADData.h
@@ -10,10 +10,8 @@
 class FArchive;
 
 using FCadId = uint32; // Identifier defined in the input CAD file
-using FColorId = uint32; // Identifier defined in the input CAD file
-using FMaterialId = uint32; // Identifier defined in the input CAD file
-using FCADUUID = uint32;  // Universal unique identifier that be used for the unreal asset name (Actor, Material)
-
+using FMaterialUId = int32; // Material of Color unique Identifier defined based on the material or color properties
+using FCadUuid = uint32;  // Unique identifier that be used for the unreal asset name (Actor, Material)
 
 namespace CADLibrary
 {
@@ -78,9 +76,9 @@ public:
 
 struct CADTOOLS_API FObjectDisplayDataId
 {
-	FCADUUID DefaultMaterialName = 0;
-	FMaterialId Material = 0;
-	FColorId Color = 0; // => FastHash == ColorId+Transparency
+	FMaterialUId DefaultMaterialUId = 0;
+	FMaterialUId MaterialUId = 0;
+	FMaterialUId ColorUId = 0;
 };
 
 class CADTOOLS_API FFileDescriptor
@@ -223,7 +221,7 @@ struct CADTOOLS_API FTessellationData
 	/** Index of each vertex in FBody::VertexArray. Empty with CoreTech and filled by FillKioVertexPosition */
 	TArray<int32> PositionIndices;
 
-	/** Index of Vertices of each face in the local Vertices set (i.e. VerticesBodyIndex for CADKernel, VertexArray for Coretech) */
+	/** Index of Vertices of each face in the local Vertices set (i.e. VerticesBodyIndex for CADKernel) */
 	TArray<int32> VertexIndices;
 
 	/** Normal of each vertex */
@@ -232,8 +230,8 @@ struct CADTOOLS_API FTessellationData
 	/** UV coordinates of each vertex */
 	TArray<FVector2D> TexCoordArray;
 
-	FCADUUID ColorName = 0;
-	FCADUUID MaterialName = 0;
+	FMaterialUId ColorUId = 0;
+	FMaterialUId MaterialUId = 0;
 
 	int32 PatchId;
 };
@@ -255,20 +253,19 @@ public:
 
 	uint32 TriangleCount = 0;
 	FCadId BodyID = 0;
-	FCADUUID MeshActorName = 0;
+	FCadUuid MeshActorUId = 0;
 
 	TArray<int32> VertexIds;  // StaticMesh FVertexID NO Serialize, filled by FillKioVertexPosition or FillVertexPosition
 	TArray<int32> SymmetricVertexIds; // StaticMesh FVertexID for sym part NO Serialize, filled by FillKioVertexPosition or FillVertexPosition
 
-	TSet<uint32> MaterialSet;
-	TSet<uint32> ColorSet;
+	TSet<FMaterialUId> MaterialSet;
+	TSet<FMaterialUId> ColorSet;
 };
 
-CADTOOLS_API uint32 BuildColorId(uint32 ColorId, uint8 Alpha);
-CADTOOLS_API void GetCTColorIdAlpha(uint32 ColorHash, uint32& OutColorId, uint8& OutAlpha);
+CADTOOLS_API FMaterialUId BuildColorFastUId(uint32 ColorId, uint8 Alpha);
 
-CADTOOLS_API int32 BuildColorName(const FColor& Color);
-CADTOOLS_API int32 BuildMaterialName(const FCADMaterial& Material);
+CADTOOLS_API FMaterialUId BuildColorUId(const FColor& Color);
+CADTOOLS_API FMaterialUId BuildMaterialUId(const FCADMaterial& Material);
 
 CADTOOLS_API void SerializeBodyMeshSet(const TCHAR* Filename, TArray<FBodyMesh>& InBodySet);
 CADTOOLS_API void DeserializeBodyMeshFile(const TCHAR* Filename, TArray<FBodyMesh>& OutBodySet);
diff --git a/Engine/Plugins/Enterprise/DatasmithCADImporter/Source/DatasmithCADTranslator/Private/DatasmithMeshBuilder.cpp b/Engine/Plugins/Enterprise/DatasmithCADImporter/Source/DatasmithCADTranslator/Private/DatasmithMeshBuilder.cpp
--- a/Engine/Plugins/Enterprise/DatasmithCADImporter/Source/DatasmithCADTranslator/Private/DatasmithMeshBuilder.cpp
+++ b/Engine/Plugins/Enterprise/DatasmithCADImporter/Source/DatasmithCADTranslator/Private/DatasmithMeshBuilder.cpp
@@ -32,7 +32,7 @@ void FDatasmithMeshBuilder::LoadMeshFiles(TMap<uint32, FString>& CADFileToMeshFi
 		DeserializeBodyMeshFile(*MeshFile, BodyMeshSet);
 		for (CADLibrary::FBodyMesh& Body : BodyMeshSet)
 		{
-			MeshActorNameToBodyMesh.Emplace(Body.MeshActorName, &Body);
+			MeshActorNameToBodyMesh.Emplace(Body.MeshActorUId, &Body);
 		}
 	}
 }
@@ -40,7 +40,7 @@ void FDatasmithMeshBuilder::LoadMeshFiles(TMap<uint32, FString>& CADFileToMeshFi
 TOptional<FMeshDescription> FDatasmithMeshBuilder::GetMeshDescription(TSharedRef<IDatasmithMeshElement> OutMeshElement, CADLibrary::FMeshParameters& OutMeshParameters)
 {
 	const TCHAR* NameLabel = OutMeshElement->GetName();
-	FCADUUID BodyUuid = (FCADUUID) FCString::Atoi64(OutMeshElement->GetName() + 2);  // +2 to remove 2 first char (Ox)
+	FCadUuid BodyUuid = (FCadUuid) FCString::Atoi64(OutMeshElement->GetName() + 2);  // +2 to remove 2 first char (Ox)
 	if (BodyUuid == 0)
 	{
 		return TOptional<FMeshDescription>();
@@ -66,7 +66,7 @@ TOptional<FMeshDescription> FDatasmithMeshBuilder::GetMeshDescription(TSharedRef
 
 		for (CADLibrary::FTessellationData& Face : Body.Faces)
 		{
-			Face.ColorName = MaterialSlotId;
+			Face.ColorUId = MaterialSlotId;
 		}
 	}
 
diff --git a/Engine/Plugins/Enterprise/DatasmithCADImporter/Source/DatasmithCADTranslator/Private/DatasmithSceneGraphBuilder.cpp b/Engine/Plugins/Enterprise/DatasmithCADImporter/Source/DatasmithCADTranslator/Private/DatasmithSceneGraphBuilder.cpp
--- a/Engine/Plugins/Enterprise/DatasmithCADImporter/Source/DatasmithCADTranslator/Private/DatasmithSceneGraphBuilder.cpp
+++ b/Engine/Plugins/Enterprise/DatasmithCADImporter/Source/DatasmithCADTranslator/Private/DatasmithSceneGraphBuilder.cpp
@@ -115,12 +115,12 @@ void FDatasmithSceneGraphBuilder::LoadSceneGraphDescriptionFiles()
 
 		for(const auto& ColorPair : MockUpDescription.ColorHIdToColor)
 		{
-			ColorNameToColorArchive.Emplace(ColorPair.Value.UEMaterialName, ColorPair.Value);
+			ColorUIdToColorArchive.Emplace(ColorPair.Value.UEMaterialName, ColorPair.Value);
 		}
 
 		for (const auto& MaterialPair : MockUpDescription.MaterialHIdToMaterial)
 		{
-			MaterialNameToMaterialArchive.Emplace(MaterialPair.Value.UEMaterialName, MaterialPair.Value);
+			MaterialUIdToMaterialArchive.Emplace(MaterialPair.Value.UEMaterialName, MaterialPair.Value);
 		}
 
 	}
@@ -178,16 +178,16 @@ FDatasmithSceneBaseGraphBuilder::FDatasmithSceneBaseGraphBuilder(CADLibrary::FAr
 {
 	if (InSceneGraph)
 	{
-		ColorNameToColorArchive.Reserve(SceneGraph->ColorHIdToColor.Num());
+		ColorUIdToColorArchive.Reserve(SceneGraph->ColorHIdToColor.Num());
 		for(const auto& ColorPair : SceneGraph->ColorHIdToColor)
 		{
-			ColorNameToColorArchive.Emplace(ColorPair.Value.UEMaterialName, ColorPair.Value);
+			ColorUIdToColorArchive.Emplace(ColorPair.Value.UEMaterialName, ColorPair.Value);
 		}
 
-		MaterialNameToMaterialArchive.Reserve(SceneGraph->MaterialHIdToMaterial.Num());
+		MaterialUIdToMaterialArchive.Reserve(SceneGraph->MaterialHIdToMaterial.Num());
 		for (const auto& MaterialPair : SceneGraph->MaterialHIdToMaterial)
 		{
-			MaterialNameToMaterialArchive.Emplace(MaterialPair.Value.UEMaterialName, MaterialPair.Value);
+			MaterialUIdToMaterialArchive.Emplace(MaterialPair.Value.UEMaterialName, MaterialPair.Value);
 		}
 	}
 }
@@ -419,7 +419,7 @@ void FDatasmithSceneBaseGraphBuilder::GetNodeUUIDAndName(
 		OutName = "NoName";
 	}
 
-	FCADUUID UEUUID = HashCombine(GetTypeHash(InParentUEUUID), GetTypeHash(InComponentIndex));
+	FCadUuid UEUUID = HashCombine(GetTypeHash(InParentUEUUID), GetTypeHash(InComponentIndex));
 
 	if (InstanceUUID)
 	{
@@ -482,7 +482,7 @@ TSharedPtr<IDatasmithActorElement> FDatasmithSceneBaseGraphBuilder::BuildBody(in
 
 	CADLibrary::FArchiveBody& Body = SceneGraph->Bodies[BodyIndex];
 
-	if (Body.ParentId == 0 || Body.MeshActorName == 0)
+	if (Body.ParentId == 0 || Body.MeshActorUId == 0)
 	{
 		return TSharedPtr<IDatasmithActorElement>();
 	}
@@ -535,10 +535,10 @@ TSharedPtr<IDatasmithActorElement> FDatasmithSceneBaseGraphBuilder::BuildBody(in
 
 TSharedPtr< IDatasmithMeshElement > FDatasmithSceneBaseGraphBuilder::FindOrAddMeshElement(CADLibrary::FArchiveBody& Body, FString& BodyName)
 {
-	FString ShellUuidName = FString::Printf(TEXT("0x%012u"), Body.MeshActorName);
+	FString ShellUuidName = FString::Printf(TEXT("0x%012u"), Body.MeshActorUId);
 
 	// Look if geometry has not been already processed, return it if found
-	TSharedPtr< IDatasmithMeshElement >* MeshElementPtr = BodyUuidToMeshElement.Find(Body.MeshActorName);
+	TSharedPtr< IDatasmithMeshElement >* MeshElementPtr = BodyUuidToMeshElement.Find(Body.MeshActorUId);
 	if (MeshElementPtr != nullptr)
 	{
 		return *MeshElementPtr;
@@ -554,15 +554,15 @@ TSharedPtr< IDatasmithMeshElement > FDatasmithSceneBaseGraphBuilder::FindOrAddMe
 	// the scene graph archive name that is define by the name and the stat of the file (creation date, size)
 	MD5.Update(reinterpret_cast<const uint8*>(SceneGraph->ArchiveFileName.GetCharArray().GetData()), SceneGraph->ArchiveFileName.GetCharArray().Num());
 	// MeshActorName
-	MD5.Update(reinterpret_cast<const uint8*>(&Body.MeshActorName), sizeof Body.MeshActorName);
+	MD5.Update(reinterpret_cast<const uint8*>(&Body.MeshActorUId), sizeof Body.MeshActorUId);
 
 	FMD5Hash Hash;
 	Hash.Set(MD5);
 	MeshElement->SetFileHash(Hash);
 
-	TFunction<void(TSet<uint32>&)> SetMaterialToDatasmithMeshElement = [&](TSet<uint32>& MaterialSet)
+	TFunction<void(TSet<FMaterialUId>&)> SetMaterialToDatasmithMeshElement = [&](TSet<FMaterialUId>& MaterialSet)
 	{
-		for (uint32 MaterialSlotId : MaterialSet)
+		for (FMaterialUId MaterialSlotId : MaterialSet)
 		{
 			TSharedPtr< IDatasmithMaterialIDElement > PartMaterialIDElement;
 			PartMaterialIDElement = FindOrAddMaterial(MaterialSlotId);
@@ -575,9 +575,9 @@ TSharedPtr< IDatasmithMeshElement > FDatasmithSceneBaseGraphBuilder::FindOrAddMe
 
 	DatasmithScene->AddMesh(MeshElement);
 
-	BodyUuidToMeshElement.Add(Body.MeshActorName, MeshElement);
+	BodyUuidToMeshElement.Add(Body.MeshActorUId, MeshElement);
 
-	FString BodyCachePath = CADLibrary::BuildCacheFilePath(*CachePath, TEXT("body"), Body.MeshActorName);
+	FString BodyCachePath = CADLibrary::BuildCacheFilePath(*CachePath, TEXT("body"), Body.MeshActorUId);
 	MeshElement->SetFile(*BodyCachePath);
 
 	return MeshElement;
@@ -605,11 +605,11 @@ TSharedPtr<IDatasmithMaterialIDElement> FDatasmithSceneBaseGraphBuilder::FindOrA
 	}
 	else if (MaterialUuid > 0)
 	{
-		if (CADLibrary::FArchiveColor* Color = ColorNameToColorArchive.Find(MaterialUuid))
+		if (CADLibrary::FArchiveColor* Color = ColorUIdToColorArchive.Find(MaterialUuid))
 		{
 			MaterialElement = CADLibrary::CreateUEPbrMaterialFromColor(Color->Color);
 		}
-		else if (CADLibrary::FArchiveMaterial* Material = MaterialNameToMaterialArchive.Find(MaterialUuid))
+		else if (CADLibrary::FArchiveMaterial* Material = MaterialUIdToMaterialArchive.Find(MaterialUuid))
 		{
 			MaterialElement = CADLibrary::CreateUEPbrMaterialFromMaterial(Material->Material, DatasmithScene);
 		}
diff --git a/Engine/Plugins/Enterprise/DatasmithCADImporter/Source/DatasmithCADTranslator/Public/DatasmithMeshBuilder.h b/Engine/Plugins/Enterprise/DatasmithCADImporter/Source/DatasmithCADTranslator/Public/DatasmithMeshBuilder.h
--- a/Engine/Plugins/Enterprise/DatasmithCADImporter/Source/DatasmithCADTranslator/Public/DatasmithMeshBuilder.h
+++ b/Engine/Plugins/Enterprise/DatasmithCADImporter/Source/DatasmithCADTranslator/Public/DatasmithMeshBuilder.h
@@ -24,7 +24,7 @@ public:
 
 		for (CADLibrary::FBodyMesh& Body : BodyMeshSet)
 		{
-			MeshActorNameToBodyMesh.Emplace(Body.MeshActorName, &Body);
+			MeshActorNameToBodyMesh.Emplace(Body.MeshActorUId, &Body);
 		}
 	}
 
@@ -36,7 +36,7 @@ protected:
 	void LoadMeshFiles(TMap<uint32, FString>& CADFileToMeshFile);
 
 	TArray<TArray<CADLibrary::FBodyMesh>> BodyMeshes;
-	TMap<FCADUUID, CADLibrary::FBodyMesh*> MeshActorNameToBodyMesh;
+	TMap<FCadUuid, CADLibrary::FBodyMesh*> MeshActorNameToBodyMesh;
 
 	CADLibrary::FImportParameters ImportParameters;
 };
diff --git a/Engine/Plugins/Enterprise/DatasmithCADImporter/Source/DatasmithCADTranslator/Public/DatasmithSceneGraphBuilder.h b/Engine/Plugins/Enterprise/DatasmithCADImporter/Source/DatasmithCADTranslator/Public/DatasmithSceneGraphBuilder.h
--- a/Engine/Plugins/Enterprise/DatasmithCADImporter/Source/DatasmithCADTranslator/Public/DatasmithSceneGraphBuilder.h
+++ b/Engine/Plugins/Enterprise/DatasmithCADImporter/Source/DatasmithCADTranslator/Public/DatasmithSceneGraphBuilder.h
@@ -40,10 +40,10 @@ public:
 	const TCHAR* Uuid;
 
 	CADLibrary::FCADMaterial Material;
-	uint32 MaterialUuid;
+	FMaterialUId MaterialUuid;
 
 	FColor Color;
-	uint32 ColorUuid;
+	FMaterialUId ColorUuid;
 };
 
 
@@ -93,17 +93,16 @@ protected:
 	TArray<CADLibrary::FArchiveSceneGraph> ArchiveMockUps;
 	TMap<uint32, CADLibrary::FArchiveSceneGraph*> CADFileToSceneGraphArchive;
 
-	TMap< FCADUUID, TSharedPtr< IDatasmithMeshElement > > BodyUuidToMeshElement;
+	TMap<FCadUuid, TSharedPtr< IDatasmithMeshElement>> BodyUuidToMeshElement;
 
-	TMap< FCADUUID, TSharedPtr< IDatasmithUEPbrMaterialElement > > MaterialUuidMap;
-	TSharedPtr<IDatasmithUEPbrMaterialElement > DefaultMaterial;
+	TMap<FMaterialUId, TSharedPtr< IDatasmithUEPbrMaterialElement>> MaterialUuidMap;
+	TSharedPtr<IDatasmithUEPbrMaterialElement> DefaultMaterial;
 
-	TMap<FCADUUID, CADLibrary::FArchiveColor> ColorNameToColorArchive; 
-	TMap<FCADUUID, CADLibrary::FArchiveMaterial> MaterialNameToMaterialArchive; 
+	TMap<FMaterialUId, CADLibrary::FArchiveColor> ColorUIdToColorArchive; 
+	TMap<FMaterialUId, CADLibrary::FArchiveMaterial> MaterialUIdToMaterialArchive; 
 
 	TArray<uint32> AncestorSceneGraphHash;
 
-
 	bool bPreferMaterial;
 	bool bMaterialPropagationIsTopDown;
 };
diff --git a/Engine/Plugins/Enterprise/DatasmithCADImporter/Source/DatasmithWireTranslator/Private/AliasModelToCADKernelConverter.cpp b/Engine/Plugins/Enterprise/DatasmithCADImporter/Source/DatasmithWireTranslator/Private/AliasModelToCADKernelConverter.cpp
--- a/Engine/Plugins/Enterprise/DatasmithCADImporter/Source/DatasmithWireTranslator/Private/AliasModelToCADKernelConverter.cpp
+++ b/Engine/Plugins/Enterprise/DatasmithCADImporter/Source/DatasmithWireTranslator/Private/AliasModelToCADKernelConverter.cpp
@@ -282,7 +282,7 @@ bool FAliasModelToCADKernelConverter::AddBRep(AlDagNode& DagNode, const FColor&
 
 	TSharedRef<CADKernel::FBody> CADKernelBody = CADKernel::FEntity::MakeShared<CADKernel::FBody>();
 	TSharedRef<CADKernel::FShell> CADKernelShell = CADKernel::FEntity::MakeShared<CADKernel::FShell>();
-	uint32 ColorId = (uint32) CADLibrary::BuildColorName(Color);
+	uint32 ColorId = (uint32) CADLibrary::BuildColorUId(Color);
 	CADKernelShell->SetColorId(ColorId);
 	CADKernelBody->AddShell(CADKernelShell);
 
