Improve  CADImport naming

diff --git a/Engine/Plugins/Enterprise/DatasmithCADImporter/Source/CADInterfaces/Private/CADSceneGraph.cpp b/Engine/Plugins/Enterprise/DatasmithCADImporter/Source/CADInterfaces/Private/CADSceneGraph.cpp
--- a/Engine/Plugins/Enterprise/DatasmithCADImporter/Source/CADInterfaces/Private/CADSceneGraph.cpp
+++ b/Engine/Plugins/Enterprise/DatasmithCADImporter/Source/CADInterfaces/Private/CADSceneGraph.cpp
@@ -9,9 +9,9 @@
 namespace CADLibrary
 {
 
-FArchive& operator<<(FArchive& Ar, FCADArchiveObject& Object)
+FArchive& operator<<(FArchive& Ar, FArchiveCADObject& Object)
 {
-	Ar << Object.ObjectId;
+	Ar << Object.Id;
 	Ar << Object.MetaData;
 	Ar << Object.TransformMatrix;
 	return Ar;
@@ -19,29 +19,29 @@ FArchive& operator<<(FArchive& Ar, FCADArchiveObject& Object)
 
 FArchive& operator<<(FArchive& Ar, FArchiveInstance& Instance) 
 {
-	Ar << (FCADArchiveObject&) Instance;
+	Ar << (FArchiveCADObject&) Instance;
 	Ar << Instance.ReferenceNodeId;
 	Ar << Instance.bIsExternalReference;
 	Ar << Instance.ExternalReference;
 	return Ar;
 }
 
-FArchive& operator<<(FArchive& Ar, FArchiveComponent& Component)
+FArchive& operator<<(FArchive& Ar, FArchiveReference& Component)
 {
-	Ar << (FCADArchiveObject&) Component;
+	Ar << (FArchiveCADObject&) Component;
 	Ar << Component.Children;
 	return Ar;
 }
 
-FArchive& operator<<(FArchive& Ar, FArchiveUnloadedComponent& Unloaded) 
+FArchive& operator<<(FArchive& Ar, FArchiveUnloadedReference& Unloaded) 
 {
-	Ar << (FArchiveComponent&) Unloaded;
+	Ar << (FArchiveReference&) Unloaded;
 	return Ar;
 }
 
 FArchive& operator<<(FArchive& Ar, FArchiveBody& Body) 
 {
-	Ar << (FCADArchiveObject&) Body;
+	Ar << (FArchiveCADObject&) Body;
 	Ar << Body.MaterialFaceSet;
 	Ar << Body.ColorFaceSet;
 	Ar << Body.ParentId;
@@ -53,17 +53,17 @@ FArchive& operator<<(FArchive& Ar, FArchiveBody& Body)
 
 FArchive& operator<<(FArchive& Ar, FArchiveColor& Color) 
 {
-	Ar << Color.ObjectId;
+	Ar << Color.Id;
 	Ar << Color.Color;
-	Ar << Color.UEMaterialName;
+	Ar << Color.UEMaterialUId;
 	return Ar;
 }
 
 FArchive& operator<<(FArchive& Ar, FArchiveMaterial& Material)
 {
-	Ar << Material.ObjectId;
+	Ar << Material.Id;
 	Ar << Material.Material;
-	Ar << Material.UEMaterialName;
+	Ar << Material.UEMaterialUId;
 	return Ar;
 }
 
@@ -72,19 +72,19 @@ FArchive& operator<<(FArchive& Ar, FArchiveSceneGraph& SceneGraph)
 	Ar << SceneGraph.CADFileName;
 	Ar << SceneGraph.ArchiveFileName;
 	Ar << SceneGraph.FullPath;
-	Ar << SceneGraph.ExternalReferences;
+	Ar << SceneGraph.ExternalReferenceFiles;
 
 	Ar << SceneGraph.ColorHIdToColor;
 	Ar << SceneGraph.MaterialHIdToMaterial;
 
 	Ar << SceneGraph.Instances;
-	Ar << SceneGraph.Components;
-	Ar << SceneGraph.UnloadedComponents;
+	Ar << SceneGraph.References;
+	Ar << SceneGraph.UnloadedReferences;
 	Ar << SceneGraph.Bodies;
 
 	Ar << SceneGraph.CADIdToInstanceIndex;
-	Ar << SceneGraph.CADIdToComponentIndex;
-	Ar << SceneGraph.CADIdToUnloadedComponentIndex;
+	Ar << SceneGraph.CADIdToReferenceIndex;
+	Ar << SceneGraph.CADIdToUnloadedReferenceIndex;
 	Ar << SceneGraph.CADIdToBodyIndex;
 
 	return Ar;
diff --git a/Engine/Plugins/Enterprise/DatasmithCADImporter/Source/CADInterfaces/Private/TechSoftFileParser.cpp b/Engine/Plugins/Enterprise/DatasmithCADImporter/Source/CADInterfaces/Private/TechSoftFileParser.cpp
--- a/Engine/Plugins/Enterprise/DatasmithCADImporter/Source/CADInterfaces/Private/TechSoftFileParser.cpp
+++ b/Engine/Plugins/Enterprise/DatasmithCADImporter/Source/CADInterfaces/Private/TechSoftFileParser.cpp
@@ -388,9 +388,9 @@ ECADParsingResult FTechSoftFileParser::Process()
 	}
 
 	FString TechSoftVersion = TechSoftInterface::GetTechSoftVersion();
-	if (!TechSoftVersion.IsEmpty() && CADFileData.ComponentCount() > 0)
+	if (!TechSoftVersion.IsEmpty() && CADFileData.ReferenceCount() > 0)
 	{
-		CADFileData.GetComponentAt(0).MetaData.Add(TEXT("TechsoftVersion"), TechSoftVersion);
+		CADFileData.GetReferenceAt(0).MetaData.Add(TEXT("TechsoftVersion"), TechSoftVersion);
 	}
 
 	ModelFile.Reset();
@@ -409,7 +409,7 @@ void FTechSoftFileParser::SewModel()
 
 void FTechSoftFileParser::GenerateBodyMeshes()
 {
-	for (TPair<A3DRiRepresentationItem*, int32>& Entry : RepresentationItemsCache)
+	for (TPair<A3DRiRepresentationItem*, FCadId>& Entry : RepresentationItemsCache)
 	{
 		A3DRiRepresentationItem* RepresentationItemPtr = Entry.Key;
 		FArchiveBody& Body = CADFileData.GetBodyAt(Entry.Value);
@@ -419,7 +419,7 @@ void FTechSoftFileParser::GenerateBodyMeshes()
 
 void FTechSoftFileParser::GenerateBodyMesh(A3DRiRepresentationItem* Representation, FArchiveBody& Body)
 {
-	FBodyMesh& BodyMesh = CADFileData.AddBodyMesh(Body.ObjectId, Body);
+	FBodyMesh& BodyMesh = CADFileData.AddBodyMesh(Body.Id, Body);
 
 	uint32 NewBRepCount = 0;
 	A3DRiBrepModel** NewBReps = nullptr;
@@ -748,7 +748,7 @@ void FTechSoftFileParser::TraverseReference(const A3DAsmProductOccurrence* Refer
 	}
 
 	FArchiveInstance EmptyInstance;
-	FArchiveComponent& Component = FTechSoftFileParser::AddComponent(MetaData, EmptyInstance);
+	FArchiveReference& Reference = FTechSoftFileParser::AddReference(MetaData, EmptyInstance);
 
 	TUniqueTSObj<A3DAsmProductOccurrenceData> ReferenceData(ReferencePtr);
 	if (!ReferenceData.IsValid())
@@ -765,18 +765,18 @@ void FTechSoftFileParser::TraverseReference(const A3DAsmProductOccurrence* Refer
 		ReferenceMatrix = ExtractTransformation(Location, ReferenceUnit);
 	}
 
-	Component.TransformMatrix = ParentMatrix * ReferenceMatrix;
+	Reference.TransformMatrix = ParentMatrix * ReferenceMatrix;
 
-	const FString& OccurenceNameBase = Component.MetaData.FindOrAdd(TEXT("Name"));
+	const FString& OccurenceNameBase = Reference.MetaData.FindOrAdd(TEXT("Name"));
 	for (uint32 OccurenceIndex = 0; OccurenceIndex < ReferenceData->m_uiPOccurrencesSize; ++OccurenceIndex)
 	{
 		int32 ChildrenId = TraverseOccurrence(ReferenceData->m_ppPOccurrences[OccurenceIndex], OccurenceNameBase + TEXT("_") + FString::FromInt(OccurenceIndex+1), ReferenceUnit);
-		Component.Children.Add(ChildrenId);
+		Reference.Children.Add(ChildrenId);
 	}
 
 	if (ReferenceData->m_pPart)
 	{
-		TraversePartDefinition(ReferenceData->m_pPart, Component, ReferenceUnit);
+		TraversePartDefinition(ReferenceData->m_pPart, Reference, ReferenceUnit);
 	}
 }
 
@@ -789,36 +789,36 @@ FArchiveInstance& FTechSoftFileParser::AddInstance(FEntityMetaData& InstanceMeta
 	return Instance;
 }
 
-FArchiveComponent& FTechSoftFileParser::AddComponent(FEntityMetaData& ComponentMetaData, FArchiveInstance& Instance)
+FArchiveReference& FTechSoftFileParser::AddReference(FEntityMetaData& ReferenceMetaData, FArchiveInstance& Instance)
 {
-	FCadId ComponentId = LastEntityId++;
-	int32 ComponentIndex = CADFileData.AddComponent(ComponentId);
-	FArchiveComponent& Prototype = CADFileData.GetComponentAt(ComponentIndex);
-	Prototype.MetaData = MoveTemp(ComponentMetaData.MetaData);
+	FCadId ReferenceId = LastEntityId++;
+	int32 ReferenceIndex = CADFileData.AddReference(ReferenceId);
+	FArchiveReference& Prototype = CADFileData.GetReferenceAt(ReferenceIndex);
+	Prototype.MetaData = MoveTemp(ReferenceMetaData.MetaData);
 
-	Instance.ReferenceNodeId = ComponentId;
+	Instance.ReferenceNodeId = ReferenceId;
 
 	return Prototype;
 }
 
-FArchiveUnloadedComponent& FTechSoftFileParser::AddUnloadedComponent(FEntityMetaData& ComponentMetaData, FArchiveInstance& Instance)
+FArchiveUnloadedReference& FTechSoftFileParser::AddUnloadedReference(FEntityMetaData& ReferenceMetaData, FArchiveInstance& Instance)
 {
-	FCadId ComponentId = LastEntityId++;
-	int32 ComponentIndex = CADFileData.AddUnloadedComponent(ComponentId);
-	FArchiveUnloadedComponent& Component = CADFileData.GetUnloadedComponentAt(ComponentIndex);
-	Component.MetaData = MoveTemp(ComponentMetaData.MetaData);
+	FCadId ReferenceId = LastEntityId++;
+	int32 ReferenceIndex = CADFileData.AddUnloadedReference(ReferenceId);
+	FArchiveUnloadedReference& Reference = CADFileData.GetUnloadedReferenceAt(ReferenceIndex);
+	Reference.MetaData = MoveTemp(ReferenceMetaData.MetaData);
 
 	Instance.bIsExternalReference = true;
-	Instance.ReferenceNodeId = ComponentId;
+	Instance.ReferenceNodeId = ReferenceId;
 
 	// Make sure that the ExternalFile path is right otherwise try to find the file and update.
-	TechSoftFileParserImpl::UpdateFileDescriptor(ComponentMetaData.ExternalFile);
+	TechSoftFileParserImpl::UpdateFileDescriptor(ReferenceMetaData.ExternalFile);
 
-	Instance.ExternalReference = ComponentMetaData.ExternalFile;
+	Instance.ExternalReference = ReferenceMetaData.ExternalFile;
 
 	if (Format == ECADFormat::SOLIDWORKS /*|| CADFileData.FileFormat() == ECADFormat::JT*/)
 	{
-		if (FString* ConfigurationName = Component.MetaData.Find(TEXT("ConfigurationName")))
+		if (FString* ConfigurationName = Reference.MetaData.Find(TEXT("ConfigurationName")))
 		{
 			Instance.ExternalReference.SetConfiguration(*ConfigurationName);
 		}
@@ -826,23 +826,23 @@ FArchiveUnloadedComponent& FTechSoftFileParser::AddUnloadedComponent(FEntityMeta
 
 	CADFileData.AddExternalRef(Instance.ExternalReference);
 
-	return Component;
+	return Reference;
 }
 
-FArchiveComponent& FTechSoftFileParser::AddOccurence(FEntityMetaData& InstanceMetaData, FCadId& OutComponentId)
+FArchiveReference& FTechSoftFileParser::AddOccurence(FEntityMetaData& InstanceMetaData, FCadId& OutReferenceId)
 {
 	FArchiveInstance& Instance = AddInstance(InstanceMetaData);
-	OutComponentId = Instance.ObjectId;
+	OutReferenceId = Instance.Id;
 	FEntityMetaData ReferenceMetaData;
-	FArchiveComponent& Prototype = AddComponent(ReferenceMetaData, Instance);
+	FArchiveReference& Prototype = AddReference(ReferenceMetaData, Instance);
 	return Prototype;
 }
 
-FArchiveComponent& FTechSoftFileParser::AddOccurence(FEntityMetaData& InstanceMetaData, FEntityMetaData& ReferenceMetaData, FCadId& OutComponentId)
+FArchiveReference& FTechSoftFileParser::AddOccurence(FEntityMetaData& InstanceMetaData, FEntityMetaData& ReferenceMetaData, FCadId& OutReferenceId)
 {
 	FArchiveInstance& Instance = AddInstance(InstanceMetaData);
-	OutComponentId = Instance.ObjectId;
-	FArchiveComponent& Prototype = AddComponent(ReferenceMetaData, Instance);
+	OutReferenceId = Instance.Id;
+	FArchiveReference& Prototype = AddReference(ReferenceMetaData, Instance);
 	return Prototype;
 }
 
@@ -904,7 +904,7 @@ FCadId FTechSoftFileParser::TraverseOccurrence(const A3DAsmProductOccurrence* Oc
 	A3DMiscTransformation* Location = OccurrenceData->m_pLocation;
 
 	// Todo
-	// the life time of FEntityMetaData end at the creation of FArchive(Unloaded)Component.
+	// the life time of FEntityMetaData end at the creation of FArchive(Unloaded)Reference.
 	// This variable should not be accessible after
 	FEntityMetaData PrototypeMetaData;
 	if (OccurrenceData->m_pPrototype)
@@ -920,8 +920,8 @@ FCadId FTechSoftFileParser::TraverseOccurrence(const A3DAsmProductOccurrence* Oc
 
 	if (PrototypeMetaData.bUnloaded)
 	{
-		FArchiveUnloadedComponent& UnloadedComponent = AddUnloadedComponent(PrototypeMetaData, Instance);
-		return Instance.ObjectId;
+		FArchiveUnloadedReference& UnloadedReference = AddUnloadedReference(PrototypeMetaData, Instance);
+		return Instance.Id;
 	}
 	
 	// If the prototype hasn't name, set its name with the name of the instance 
@@ -934,7 +934,7 @@ FCadId FTechSoftFileParser::TraverseOccurrence(const A3DAsmProductOccurrence* Oc
 		}
 	}
 
-	FArchiveComponent& Component = AddComponent(PrototypeMetaData, Instance);
+	FArchiveReference& Reference = AddReference(PrototypeMetaData, Instance);
 
 	while (OccurrenceData->m_pPrototype != nullptr && OccurrenceData->m_pPart == nullptr && OccurrenceData->m_uiPOccurrencesSize == 0 && OccurrenceData->m_pExternalData == nullptr)
 	{
@@ -944,7 +944,7 @@ FCadId FTechSoftFileParser::TraverseOccurrence(const A3DAsmProductOccurrence* Oc
 
 	if(OccurrenceData->m_pPart == nullptr && OccurrenceData->m_uiPOccurrencesSize == 0 && OccurrenceData->m_pExternalData == nullptr)
 	{
-		return Instance.ObjectId;
+		return Instance.Id;
 	}
 
 	// Add part
@@ -955,7 +955,7 @@ FCadId FTechSoftFileParser::TraverseOccurrence(const A3DAsmProductOccurrence* Oc
 	if (OccurrenceData->m_pPart != nullptr)
 	{
 		A3DAsmPartDefinition* PartDefinition = OccurrenceData->m_pPart;
-		TraversePartDefinition(PartDefinition, Component, OccurrenceUnit);
+		TraversePartDefinition(PartDefinition, Reference, OccurrenceUnit);
 	}
 
 	// Add Occurrence's Children
@@ -972,7 +972,7 @@ FCadId FTechSoftFileParser::TraverseOccurrence(const A3DAsmProductOccurrence* Oc
 	{
 		const FString DefaultChildName = InstanceName + TEXT("_") + FString::FromInt(Index + 1);
 		int32 ChildrenId = TraverseOccurrence(Children[Index], DefaultChildName, OccurrenceUnit);
-		Component.Children.Add(ChildrenId);
+		Reference.Children.Add(ChildrenId);
 	}
 
 	// Add External data
@@ -987,7 +987,7 @@ FCadId FTechSoftFileParser::TraverseOccurrence(const A3DAsmProductOccurrence* Oc
 		if (ExternalData->m_pPart != nullptr)
 		{
 			A3DAsmPartDefinition* PartDefinition = ExternalData->m_pPart;
-			TraversePartDefinition(PartDefinition, Component, OccurrenceUnit);
+			TraversePartDefinition(PartDefinition, Reference, OccurrenceUnit);
 		}
 
 		uint32 ExternalChildrenCount = ExternalData->m_uiPOccurrencesSize;
@@ -996,11 +996,11 @@ FCadId FTechSoftFileParser::TraverseOccurrence(const A3DAsmProductOccurrence* Oc
 		{
 			const FString DefaultChildName = InstanceName + TEXT("_") + FString::FromInt(Index + 1);
 			int32 ChildId = TraverseOccurrence(ExternalChildren[Index], DefaultChildName, OccurrenceUnit);
-			Component.Children.Add(ChildId);
+			Reference.Children.Add(ChildId);
 		}
 	}
 
-	return Instance.ObjectId;
+	return Instance.Id;
 }
 
 void FTechSoftFileParser::CountUnderOccurrence(const A3DAsmProductOccurrence* Occurrence)
@@ -1132,7 +1132,7 @@ void FTechSoftFileParser::ProcessPrototype(const A3DAsmProductOccurrence* InProt
 			}
 		}
 		
-		// Case of DWG file, external reference is in fact internal data as a part or a component
+		// Case of DWG file, external reference is in fact internal data as a part or a Reference
 		// Case of CATProduct referencing CGR files: unloaded references are saved in ExternalData. In this case non empty ExternalFile means unloaded reference
 		if (PrototypeData->m_pPart != nullptr|| PrototypeData->m_uiPOccurrencesSize != 0 || (PrototypeData->m_pExternalData != nullptr && OutPrototypeMetaData.ExternalFile.IsEmpty()))
 		{
@@ -1185,7 +1185,7 @@ void FTechSoftFileParser::CountUnderPrototype(const A3DAsmProductOccurrence* Pro
 	ComponentCount[EComponentType::Reference] ++;
 }
 
-void FTechSoftFileParser::TraversePartDefinition(const A3DAsmPartDefinition* PartDefinitionPtr, FArchiveComponent& Part, double ParentUnit)
+void FTechSoftFileParser::TraversePartDefinition(const A3DAsmPartDefinition* PartDefinitionPtr, FArchiveReference& Part, double ParentUnit)
 {
 	FEntityMetaData PartMetaData;
 	ExtractMetaData(PartDefinitionPtr, PartMetaData);
@@ -1221,7 +1221,7 @@ void FTechSoftFileParser::TraversePartDefinition(const A3DAsmPartDefinition* Par
 	{
 		for (unsigned int Index = 0; Index < PartData->m_uiRepItemsSize; ++Index)
 		{
-			int32 ChildId = TraverseRepresentationItem(PartData->m_ppRepItems[Index], PartMetaData, Part.ObjectId, ParentUnit, Index);
+			int32 ChildId = TraverseRepresentationItem(PartData->m_ppRepItems[Index], PartMetaData, Part.Id, ParentUnit, Index);
 			Part.Children.Add(ChildId);
 		}
 	}
@@ -1298,11 +1298,11 @@ FCadId FTechSoftFileParser::TraverseRepresentationSet(const A3DRiSet* Representa
 	}
 
 	FCadId RepresentationSetId = 0;
-	FArchiveComponent& RepresentationSet = AddOccurence(RepresentationSetMetaData, RepresentationSetId);
+	FArchiveReference& RepresentationSet = AddOccurence(RepresentationSetMetaData, RepresentationSetId);
 
 	for (A3DUns32 Index = 0; Index < RepresentationSetData->m_uiRepItemsSize; ++Index)
 	{
-		int32 ChildId = TraverseRepresentationItem(RepresentationSetData->m_ppRepItems[Index], RepresentationSetMetaData, RepresentationSet.ObjectId, ParentUnit, Index);
+		int32 ChildId = TraverseRepresentationItem(RepresentationSetData->m_ppRepItems[Index], RepresentationSetMetaData, RepresentationSet.Id, ParentUnit, Index);
 		RepresentationSet.Children.Add(ChildId);
 	}
 	return RepresentationSetId;
@@ -1341,9 +1341,9 @@ FCadId FTechSoftFileParser::TraverseBRepModel(A3DRiBrepModel* BRepModelPtr, cons
 	TUniqueTSObj<A3DRiBrepModelData> BRepModelData(BRepModelPtr);
 	BuildBodyName(BRepMetaData, PartMetaData, ItemIndex, (bool) BRepModelData->m_bSolid);
 
-	if (int32* BodyIndexPtr = RepresentationItemsCache.Find(BRepModelPtr))
+	if (FCadId* BodyIndexPtr = RepresentationItemsCache.Find(BRepModelPtr))
 	{
-		return CADFileData.GetBodyAt(*BodyIndexPtr).ObjectId;
+		return CADFileData.GetBodyAt(*BodyIndexPtr).Id;
 	}
 
 	// if BRep model has not material or color, use the ones from the part
@@ -1378,7 +1378,7 @@ FCadId FTechSoftFileParser::TraverseBRepModel(A3DRiBrepModel* BRepModelPtr, cons
 
 	RepresentationItemsCache.Add(BRepModelPtr, BodyIndex);
 
-	return Body.ObjectId;
+	return Body.Id;
 }
 
 FCadId FTechSoftFileParser::TraversePolyBRepModel(A3DRiPolyBrepModel* PolygonalPtr, const FEntityMetaData& PartMetaData, const FCadId ParentId, double ParentUnit, int32 ItemIndex)
@@ -1399,9 +1399,9 @@ FCadId FTechSoftFileParser::TraversePolyBRepModel(A3DRiPolyBrepModel* PolygonalP
 	TUniqueTSObj<A3DRiPolyBrepModelData> BRepModelData(PolygonalPtr);
 	BuildBodyName(BRepMetaData, PartMetaData, ItemIndex, (bool) BRepModelData->m_bIsClosed);
 
-	if (int32* BodyIndexPtr = RepresentationItemsCache.Find(PolygonalPtr))
+	if (FCadId* BodyIndexPtr = RepresentationItemsCache.Find(PolygonalPtr))
 	{
-		return CADFileData.GetBodyAt(*BodyIndexPtr).ObjectId;
+		return CADFileData.GetBodyAt(*BodyIndexPtr).Id;
 	}
 
 	// if BRep model has not material or color, add part one
@@ -1434,7 +1434,7 @@ FCadId FTechSoftFileParser::TraversePolyBRepModel(A3DRiPolyBrepModel* PolygonalP
 
 	RepresentationItemsCache.Add(PolygonalPtr, BodyIndex);
 
-	return Body.ObjectId;
+	return Body.Id;
 }
 
 void FTechSoftFileParser::ExtractMetaData(const A3DEntity* Entity, FEntityMetaData& OutMetaData)
@@ -1591,7 +1591,7 @@ void FTechSoftFileParser::BuildInstanceName(FEntityMetaData& InstanceData, const
 	}
 }
 
-void FTechSoftFileParser::BuildPartName(FEntityMetaData& PartData, const FArchiveComponent& Component)
+void FTechSoftFileParser::BuildPartName(FEntityMetaData& PartData, const FArchiveReference& Reference)
 {
 	TMap<FString, FString>& MetaData = PartData.MetaData;
 	if (!MetaData.IsEmpty())
@@ -1613,10 +1613,10 @@ void FTechSoftFileParser::BuildPartName(FEntityMetaData& PartData, const FArchiv
 	}
 
 	// If a name hasn't been found for the Part, the name of its parent (component) is used
-	const FString* ComponentName = Component.MetaData.Find(TEXT("Name"));
-	if (ComponentName && !ComponentName->IsEmpty())
+	const FString* ReferenceName = Reference.MetaData.Find(TEXT("Name"));
+	if (ReferenceName && !ReferenceName->IsEmpty())
 	{
-		MetaData.Add(TEXT("Name"), *ComponentName);
+		MetaData.Add(TEXT("Name"), *ReferenceName);
 		return;
 	}
 }
@@ -1763,7 +1763,7 @@ FArchiveColor& FTechSoftFileParser::FindOrAddColor(uint32 ColorIndex, uint8 Alph
 	NewColor.Color = TechSoftUtils::GetColorAt(ColorIndex);
 	NewColor.Color.A = Alpha;
 	
-	NewColor.UEMaterialName = BuildColorUId(NewColor.Color);
+	NewColor.UEMaterialUId = BuildColorUId(NewColor.Color);
 	return NewColor;
 }
 
@@ -1787,13 +1787,13 @@ FArchiveMaterial& FTechSoftFileParser::AddMaterialAt(uint32 MaterialIndexToSave,
 		// Material.Emissive = GetColor(MaterialData->m_uiEmissive);
 		// Material.Reflexion;
 	}
-	NewMaterial.UEMaterialName = BuildMaterialUId(Material);
+	NewMaterial.UEMaterialUId = BuildMaterialUId(Material);
 	return NewMaterial;
 }
 
 
 // Look at TechSoftUtils::BuildCADMaterial if any loigc changes in this method
-// or any of the methos it calls
+// or any of the methods it calls
 FArchiveMaterial& FTechSoftFileParser::FindOrAddMaterial(uint32 MaterialIndex, const A3DGraphStyleData& GraphStyleData)
 {
 	if (FArchiveMaterial* MaterialArchive = CADFileData.FindMaterial(MaterialIndex))
@@ -1886,7 +1886,7 @@ void FTechSoftFileParser::ExtractGraphStyleProperties(uint32 StyleIndex, FCadUui
 		if (GraphStyleData->m_bMaterial)
 		{
 			FArchiveMaterial& MaterialArchive = FindOrAddMaterial(GraphStyleData->m_uiRgbColorIndex, *GraphStyleData);
-			OutMaterialName = MaterialArchive.UEMaterialName;
+			OutMaterialName = MaterialArchive.UEMaterialUId;
 		}
 		else
 		{
@@ -1897,7 +1897,7 @@ void FTechSoftFileParser::ExtractGraphStyleProperties(uint32 StyleIndex, FCadUui
 			}
 
 			FArchiveColor& ColorArchive = FindOrAddColor(GraphStyleData->m_uiRgbColorIndex, Alpha);
-			OutColorName = ColorArchive.UEMaterialName;
+			OutColorName = ColorArchive.UEMaterialUId;
 		}
 	}
 }
diff --git a/Engine/Plugins/Enterprise/DatasmithCADImporter/Source/CADInterfaces/Private/TechSoftFileParser.h b/Engine/Plugins/Enterprise/DatasmithCADImporter/Source/CADInterfaces/Private/TechSoftFileParser.h
--- a/Engine/Plugins/Enterprise/DatasmithCADImporter/Source/CADInterfaces/Private/TechSoftFileParser.h
+++ b/Engine/Plugins/Enterprise/DatasmithCADImporter/Source/CADInterfaces/Private/TechSoftFileParser.h
@@ -14,11 +14,11 @@ namespace CADLibrary
 {
 
 class FArchiveBody;
-class FArchiveComponent;
 class FArchiveColor;
 class FArchiveInstance;
 class FArchiveMaterial;
-class FArchiveUnloadedComponent;
+class FArchiveReference;
+class FArchiveUnloadedReference;
 class FCADFileData;
 
 struct FEntityMetaData;
@@ -27,7 +27,7 @@ enum EComponentType : uint32
 {
 	Reference = 0,
 	Occurrence,
-	UnloadedComponent,
+	UnloadedReference,
 	Body,
 	Undefined,
 	LastType
@@ -108,7 +108,7 @@ private:
 	void TraverseConfigurationSet(const A3DAsmProductOccurrence* ConfigurationSet, double ParentUnit);
 	FCadId TraverseOccurrence(const A3DAsmProductOccurrence* Occurrence, const FString& DefaultOccurrenceName, double ParentUnit);
 	void ProcessPrototype(const A3DAsmProductOccurrence* InPrototype, FEntityMetaData& OutMetaData, A3DMiscTransformation** OutLocation);
-	void TraversePartDefinition(const A3DAsmPartDefinition* PartDefinition, FArchiveComponent& Component, double ParentUnit);
+	void TraversePartDefinition(const A3DAsmPartDefinition* PartDefinition, FArchiveReference& Reference, double ParentUnit);
 	FCadId TraverseRepresentationSet(const A3DRiSet* pSet, const FEntityMetaData& PartMetaData, double ParentUnit);
 	FCadId TraverseRepresentationItem(A3DRiRepresentationItem* RepresentationItem, const FEntityMetaData& PartMetaData, const FCadId ParentId, double ParentUnit, int32 ItemIndex);
 	FCadId TraverseBRepModel(A3DRiBrepModel* BrepModel, const FEntityMetaData& PartMetaData, const FCadId ParentId, double ParentUnit, int32 ItemIndex);
@@ -119,7 +119,7 @@ private:
 
 	void BuildInstanceName(FEntityMetaData& MetaData, const FString& DefaultInstanceName);
 	void BuildReferenceName(FEntityMetaData& MetaData);
-	void BuildPartName(FEntityMetaData& MetaData, const FArchiveComponent& Component);
+	void BuildPartName(FEntityMetaData& MetaData, const FArchiveReference& Reference);
 	void BuildBodyName(FEntityMetaData& MetaData, const FEntityMetaData& PartMetaData, int32 ItemIndex, bool bIsSolid);
 
 	// Graphic properties
@@ -151,10 +151,10 @@ private:
 
 	// Archive methods
 	FArchiveInstance& AddInstance(FEntityMetaData& InstanceMetaData);
-	FArchiveComponent& AddComponent(FEntityMetaData& ComponentMetaData, FArchiveInstance& Instance);
-	FArchiveUnloadedComponent& AddUnloadedComponent(FEntityMetaData& ComponentMetaData, FArchiveInstance& Instance);
-	FArchiveComponent& AddOccurence(FEntityMetaData& InstanceMetaData, FEntityMetaData& ReferenceMetaData, FCadId& OutComponentId);
-	FArchiveComponent& AddOccurence(FEntityMetaData& InstanceMetaData, FCadId& OutComponentId);
+	FArchiveReference& AddReference(FEntityMetaData& ReferenceMetaData, FArchiveInstance& Instance);
+	FArchiveUnloadedReference& AddUnloadedReference(FEntityMetaData& ReferenceMetaData, FArchiveInstance& Instance);
+	FArchiveReference& AddOccurence(FEntityMetaData& InstanceMetaData, FEntityMetaData& ReferenceMetaData, FCadId& OutReferenceId);
+	FArchiveReference& AddOccurence(FEntityMetaData& InstanceMetaData, FCadId& OutReferenceId);
 	int32 AddBody(FEntityMetaData& BodyMetaData, const FMatrix& Matrix, const FCadId ParentId, double BodyUnit);
 #endif
 
@@ -175,7 +175,7 @@ protected:
 
 	FCadId LastEntityId = 1;
 
-	TMap<A3DRiRepresentationItem*, int32> RepresentationItemsCache;
+	TMap<A3DRiRepresentationItem*, FCadId> RepresentationItemsCache;
 };
 
 } // ns CADLibrary
\ No newline at end of file
diff --git a/Engine/Plugins/Enterprise/DatasmithCADImporter/Source/CADInterfaces/Private/TechSoftFileParserCADKernelTessellator.cpp b/Engine/Plugins/Enterprise/DatasmithCADImporter/Source/CADInterfaces/Private/TechSoftFileParserCADKernelTessellator.cpp
--- a/Engine/Plugins/Enterprise/DatasmithCADImporter/Source/CADInterfaces/Private/TechSoftFileParserCADKernelTessellator.cpp
+++ b/Engine/Plugins/Enterprise/DatasmithCADImporter/Source/CADInterfaces/Private/TechSoftFileParserCADKernelTessellator.cpp
@@ -49,7 +49,7 @@ void FTechSoftFileParserCADKernelTessellator::SewAndGenerateBodyMeshes()
 {
 	int32 RepresentationCount = RepresentationItemsCache.Num();
 	TMap<FCadId, TArray<A3DRiRepresentationItem*>> OccurenceIdToRepresentations;
-	for (TPair<A3DRiRepresentationItem*, int32>& Entry : RepresentationItemsCache)
+	for (TPair<A3DRiRepresentationItem*, FCadId>& Entry : RepresentationItemsCache)
 	{
 		A3DRiRepresentationItem* RepresentationItemPtr = Entry.Key;
 		FArchiveBody& Body = CADFileData.GetBodyAt(Entry.Value);
@@ -83,7 +83,7 @@ void FTechSoftFileParserCADKernelTessellator::SewAndMesh(TArray<A3DRiRepresentat
 
 	for (A3DRiRepresentationItem* Representation : Representations)
 	{
-		int32* BodyIndex = RepresentationItemsCache.Find(Representation);
+		FCadId* BodyIndex = RepresentationItemsCache.Find(Representation);
 		if (BodyIndex == nullptr)
 		{
 			continue;
@@ -141,7 +141,7 @@ void FTechSoftFileParserCADKernelTessellator::SewAndMesh(TArray<A3DRiRepresentat
 		CADKernel::FBody* Body = TechSoftBridge.GetBody(Representation);
 		if (Body == nullptr)
 		{
-			int32* ArchiveBodyIndex = RepresentationItemsCache.Find(Representation);
+			FCadId* ArchiveBodyIndex = RepresentationItemsCache.Find(Representation);
 			if (ArchiveBodyIndex == nullptr)
 			{
 				continue; // should not append
@@ -155,7 +155,7 @@ void FTechSoftFileParserCADKernelTessellator::SewAndMesh(TArray<A3DRiRepresentat
 	{
 		const A3DRiRepresentationItem* Representation = TechSoftBridge.GetA3DBody(Body);
 
-		int32* ArchiveBodyIndex = RepresentationItemsCache.Find(Representation);
+		FCadId* ArchiveBodyIndex = RepresentationItemsCache.Find(Representation);
 		if (ArchiveBodyIndex == nullptr)
 		{
 			continue; // should not append
@@ -211,7 +211,7 @@ void FTechSoftFileParserCADKernelTessellator::GenerateBodyMesh(A3DRiRepresentati
 
 void FTechSoftFileParserCADKernelTessellator::MeshAndGetTessellation(CADKernel::FSession& CADKernelSession, FArchiveBody& ArchiveBody, CADKernel::FBody& CADKernelBody)
 {
-	FBodyMesh& BodyMesh = CADFileData.AddBodyMesh(ArchiveBody.ObjectId, ArchiveBody);
+	FBodyMesh& BodyMesh = CADFileData.AddBodyMesh(ArchiveBody.Id, ArchiveBody);
 	ArchiveBody.ColorFaceSet = BodyMesh.ColorSet;
 	ArchiveBody.MaterialFaceSet = BodyMesh.MaterialSet;
 
diff --git a/Engine/Plugins/Enterprise/DatasmithCADImporter/Source/CADInterfaces/Public/CADFileData.h b/Engine/Plugins/Enterprise/DatasmithCADImporter/Source/CADInterfaces/Public/CADFileData.h
--- a/Engine/Plugins/Enterprise/DatasmithCADImporter/Source/CADInterfaces/Public/CADFileData.h
+++ b/Engine/Plugins/Enterprise/DatasmithCADImporter/Source/CADInterfaces/Public/CADFileData.h
@@ -114,28 +114,28 @@ public:
 		SceneGraphArchive.SerializeMockUp(*SceneGraphFilePath);
 	}
 
-	bool HasComponentOfId(FCadId ComponentId) const
+	bool HasReferenceOfId(FCadId ReferenceId) const
 	{
-		return SceneGraphArchive.CADIdToComponentIndex.Find(ComponentId) != nullptr;
+		return SceneGraphArchive.CADIdToReferenceIndex.Find(ReferenceId) != nullptr;
 	}
 
-	int32 ComponentCount()
+	int32 ReferenceCount()
 	{
-		return SceneGraphArchive.Components.Num();
+		return SceneGraphArchive.References.Num();
 	}
 
-	int32 AddComponent(FCadId ComponentId)
+	int32 AddReference(FCadId ReferenceId)
 	{
-		ensure(SceneGraphArchive.Components.Num() < SceneGraphArchive.Components.Max());
+		ensure(SceneGraphArchive.References.Num() < SceneGraphArchive.References.Max());
 
-		int32 Index = SceneGraphArchive.Components.Emplace(ComponentId);
-		SceneGraphArchive.CADIdToComponentIndex.Add(ComponentId, Index);
+		int32 Index = SceneGraphArchive.References.Emplace(ReferenceId);
+		SceneGraphArchive.CADIdToReferenceIndex.Add(ReferenceId, Index);
 		return Index;
 	}
 
-	FArchiveComponent& GetComponentAt(int32 Index)
+	FArchiveReference& GetReferenceAt(int32 Index)
 	{
-		return SceneGraphArchive.Components[Index];
+		return SceneGraphArchive.References[Index];
 	}
 
 
@@ -186,43 +186,43 @@ public:
 	}
 
 
-	int32* FindUnloadedComponentOfId(FCadId ComponentId)
+	int32* FindUnloadedReferenceOfId(FCadId ReferenceId)
 	{
-		return SceneGraphArchive.CADIdToUnloadedComponentIndex.Find(ComponentId);
+		return SceneGraphArchive.CADIdToUnloadedReferenceIndex.Find(ReferenceId);
 	}
 
-	bool HasUnloadedComponentOfId(FCadId ComponentId)
+	bool HasUnloadedReferenceOfId(FCadId ReferenceId)
 	{
-		return FindUnloadedComponentOfId(ComponentId) != nullptr;
+		return FindUnloadedReferenceOfId(ReferenceId) != nullptr;
 	}
 
-	int32 AddUnloadedComponent(FCadId ComponentId)
+	int32 AddUnloadedReference(FCadId ReferenceId)
 	{
-		ensure(SceneGraphArchive.UnloadedComponents.Num() < SceneGraphArchive.UnloadedComponents.Max());
+		ensure(SceneGraphArchive.UnloadedReferences.Num() < SceneGraphArchive.UnloadedReferences.Max());
 
-		int32 Index = SceneGraphArchive.UnloadedComponents.Emplace(ComponentId);
-		SceneGraphArchive.CADIdToUnloadedComponentIndex.Add(ComponentId, Index);
+		int32 Index = SceneGraphArchive.UnloadedReferences.Emplace(ReferenceId);
+		SceneGraphArchive.CADIdToUnloadedReferenceIndex.Add(ReferenceId, Index);
 		return Index;
 	}
 
-	FArchiveUnloadedComponent& GetUnloadedComponentAt(int32 Index)
+	FArchiveUnloadedReference& GetUnloadedReferenceAt(int32 Index)
 	{
-		return SceneGraphArchive.UnloadedComponents[Index];
+		return SceneGraphArchive.UnloadedReferences[Index];
 	}
 
 	FFileDescriptor& GetExternalReferences(int32 Index)
 	{
-		return SceneGraphArchive.ExternalReferences[Index];
+		return SceneGraphArchive.ExternalReferenceFiles[Index];
 	}
 
 	FFileDescriptor& AddExternalRef(const TCHAR* InFilePath, const TCHAR* InConfiguration, const TCHAR* InRootFilePath)
 	{
-		return SceneGraphArchive.ExternalReferences.Emplace_GetRef(InFilePath, InConfiguration, InRootFilePath);
+		return SceneGraphArchive.ExternalReferenceFiles.Emplace_GetRef(InFilePath, InConfiguration, InRootFilePath);
 	}
 
 	FFileDescriptor& AddExternalRef(const FFileDescriptor& InFileDescription)
 	{
-		return SceneGraphArchive.ExternalReferences.Emplace_GetRef(InFileDescription);
+		return SceneGraphArchive.ExternalReferenceFiles.Emplace_GetRef(InFileDescription);
 	}
 
 	/** return a unique value that will be used to define the static mesh name */
@@ -247,7 +247,7 @@ public:
 
 	const TArray<FFileDescriptor>& GetExternalRefSet() const
 	{
-		return SceneGraphArchive.ExternalReferences;
+		return SceneGraphArchive.ExternalReferenceFiles;
 	}
 
 	const FString& GetSceneGraphFileName() const
diff --git a/Engine/Plugins/Enterprise/DatasmithCADImporter/Source/CADInterfaces/Public/CADSceneGraph.h b/Engine/Plugins/Enterprise/DatasmithCADImporter/Source/CADInterfaces/Public/CADSceneGraph.h
--- a/Engine/Plugins/Enterprise/DatasmithCADImporter/Source/CADInterfaces/Public/CADSceneGraph.h
+++ b/Engine/Plugins/Enterprise/DatasmithCADImporter/Source/CADInterfaces/Public/CADSceneGraph.h
@@ -10,29 +10,29 @@ class FArchive;
 namespace CADLibrary
 {
 
-class CADINTERFACES_API FCADArchiveObject
+class CADINTERFACES_API FArchiveCADObject
 {
 public:
-	FCADArchiveObject(FCadId Id = 0)
-		: ObjectId(Id)
+	FArchiveCADObject(FCadId InId = 0)
+		: Id(InId)
 	{
 	}
 
-	virtual ~FCADArchiveObject() = default;
+	virtual ~FArchiveCADObject() = default;
 
-	friend FArchive& operator<<(FArchive& Ar, FCADArchiveObject& C);
+	friend FArchive& operator<<(FArchive& Ar, FArchiveCADObject& C);
 
 public:
-	uint32 ObjectId;
+	uint32 Id;
 	TMap<FString, FString> MetaData;
 	FMatrix TransformMatrix = FMatrix::Identity;
 };
 
-class CADINTERFACES_API FArchiveInstance : public FCADArchiveObject
+class CADINTERFACES_API FArchiveInstance : public FArchiveCADObject
 {
 public:
 	FArchiveInstance(FCadId Id = 0)
-		: FCADArchiveObject(Id)
+		: FArchiveCADObject(Id)
 	{
 	}
 
@@ -44,36 +44,36 @@ public:
 	FFileDescriptor ExternalReference;
 };
 
-class CADINTERFACES_API FArchiveComponent : public FCADArchiveObject
+class CADINTERFACES_API FArchiveReference : public FArchiveCADObject
 {
 public:
-	FArchiveComponent(FCadId Id = 0)
-		: FCADArchiveObject(Id)
+	FArchiveReference(FCadId Id = 0)
+		: FArchiveCADObject(Id)
 	{
 	}
 
-	friend FArchive& operator<<(FArchive& Ar, FArchiveComponent& C);
+	friend FArchive& operator<<(FArchive& Ar, FArchiveReference& C);
 
 public:
 	TArray<FCadId> Children;
 };
 
-class CADINTERFACES_API FArchiveUnloadedComponent : public FArchiveComponent
+class CADINTERFACES_API FArchiveUnloadedReference : public FArchiveReference
 {
 public:
-	FArchiveUnloadedComponent(FCadId Id = 0)
-		: FArchiveComponent(Id)
+	FArchiveUnloadedReference(FCadId Id = 0)
+		: FArchiveReference(Id)
 	{
 	}
 
-	friend FArchive& operator<<(FArchive& Ar, FArchiveUnloadedComponent& C);
+	friend FArchive& operator<<(FArchive& Ar, FArchiveUnloadedReference& C);
 };
 
-class CADINTERFACES_API FArchiveBody : public FCADArchiveObject
+class CADINTERFACES_API FArchiveBody : public FArchiveCADObject
 {
 public:
 	FArchiveBody(FCadId Id = 0)
-		: FCADArchiveObject(Id)
+		: FArchiveCADObject(Id)
 	{
 	}
 
@@ -92,33 +92,33 @@ public:
 class CADINTERFACES_API FArchiveColor
 {
 public:
-	FArchiveColor(FMaterialUId Id = 0)
-		: ObjectId(Id)
+	FArchiveColor(FMaterialUId InId = 0)
+		: Id(InId)
 	{
 	}
 
 	friend FArchive& operator<<(FArchive& Ar, FArchiveColor& C);
 
 public:
-	FMaterialUId ObjectId;
+	FMaterialUId Id;
 	FColor Color;
-	FCadUuid UEMaterialName;
+	FMaterialUId UEMaterialUId;
 };
 
 class CADINTERFACES_API FArchiveMaterial
 {
 public:
-	FArchiveMaterial(FMaterialUId Id = 0)
-		: ObjectId(Id)
+	FArchiveMaterial(FMaterialUId InId = 0)
+		: Id(InId)
 	{
 	}
 
 	friend FArchive& operator<<(FArchive& Ar, FArchiveMaterial& C);
 
 public:
-	FMaterialUId ObjectId;
-	FCadUuid UEMaterialName;
+	FMaterialUId Id;
 	FCADMaterial Material;
+	FMaterialUId UEMaterialUId;
 };
 
 class CADINTERFACES_API FArchiveSceneGraph
@@ -138,27 +138,27 @@ public:
 	TMap<FMaterialUId, FArchiveMaterial> MaterialHIdToMaterial;
 
 	TArray<FArchiveBody> Bodies;
-	TArray<FArchiveComponent> Components;
-	TArray<FArchiveUnloadedComponent> UnloadedComponents;
-	TArray<FFileDescriptor> ExternalReferences;
+	TArray<FArchiveReference> References;
+	TArray<FArchiveUnloadedReference> UnloadedReferences;
+	TArray<FFileDescriptor> ExternalReferenceFiles;
 	TArray<FArchiveInstance> Instances;
 
 	TMap<FCadId, int32> CADIdToBodyIndex;
-	TMap<FCadId, int32> CADIdToComponentIndex;
-	TMap<FCadId, int32> CADIdToUnloadedComponentIndex;
+	TMap<FCadId, int32> CADIdToReferenceIndex;
+	TMap<FCadId, int32> CADIdToUnloadedReferenceIndex;
 	TMap<FCadId, int32> CADIdToInstanceIndex;
 
-	void Reserve(int32 InstanceNum, int32 ComponentNum, int32 BodyNum)
+	void Reserve(int32 InstanceNum, int32 ReferenceNum, int32 BodyNum)
 	{
 		Instances.Reserve(InstanceNum);
-		Components.Reserve(ComponentNum);
-		UnloadedComponents.Reserve(ComponentNum);
-		ExternalReferences.Reserve(ComponentNum);
+		References.Reserve(ReferenceNum);
+		UnloadedReferences.Reserve(ReferenceNum);
+		ExternalReferenceFiles.Reserve(ReferenceNum);
 		Bodies.Reserve(BodyNum);
 
 		CADIdToInstanceIndex.Reserve(InstanceNum);
-		CADIdToComponentIndex.Reserve(ComponentNum);
-		CADIdToUnloadedComponentIndex.Reserve(ComponentNum);
+		CADIdToReferenceIndex.Reserve(ReferenceNum);
+		CADIdToUnloadedReferenceIndex.Reserve(ReferenceNum);
 		CADIdToBodyIndex.Reserve(BodyNum);
 	}
 };
diff --git a/Engine/Plugins/Enterprise/DatasmithCADImporter/Source/DatasmithCADTranslator/Private/DatasmithSceneGraphBuilder.cpp b/Engine/Plugins/Enterprise/DatasmithCADImporter/Source/DatasmithCADTranslator/Private/DatasmithSceneGraphBuilder.cpp
--- a/Engine/Plugins/Enterprise/DatasmithCADImporter/Source/DatasmithCADTranslator/Private/DatasmithSceneGraphBuilder.cpp
+++ b/Engine/Plugins/Enterprise/DatasmithCADImporter/Source/DatasmithCADTranslator/Private/DatasmithSceneGraphBuilder.cpp
@@ -31,7 +31,7 @@ namespace DatasmithSceneGraphBuilderImpl
 		}
 	}
 
-	void AddTransformToActor(const CADLibrary::FCADArchiveObject& Object, TSharedPtr< IDatasmithActorElement > Actor, const CADLibrary::FImportParameters& ImportParameters)
+	void AddTransformToActor(const CADLibrary::FArchiveCADObject& Object, TSharedPtr< IDatasmithActorElement > Actor, const CADLibrary::FImportParameters& ImportParameters)
 	{
 		if (!Actor.IsValid())
 		{
@@ -115,12 +115,12 @@ void FDatasmithSceneGraphBuilder::LoadSceneGraphDescriptionFiles()
 
 		for(const auto& ColorPair : MockUpDescription.ColorHIdToColor)
 		{
-			ColorUIdToColorArchive.Emplace(ColorPair.Value.UEMaterialName, ColorPair.Value);
+			ColorUIdToColorArchive.Emplace(ColorPair.Value.UEMaterialUId, ColorPair.Value);
 		}
 
 		for (const auto& MaterialPair : MockUpDescription.MaterialHIdToMaterial)
 		{
-			MaterialUIdToMaterialArchive.Emplace(MaterialPair.Value.UEMaterialName, MaterialPair.Value);
+			MaterialUIdToMaterialArchive.Emplace(MaterialPair.Value.UEMaterialUId, MaterialPair.Value);
 		}
 
 	}
@@ -139,7 +139,7 @@ void FDatasmithSceneGraphBuilder::FillAnchorActor(const TSharedRef<IDatasmithAct
 	}
 
 	FCadId RootId = 1;
-	const int32* Index = SceneGraph->CADIdToComponentIndex.Find(RootId);
+	const int32* Index = SceneGraph->CADIdToReferenceIndex.Find(RootId);
 	if (!Index)
 	{
 		return;
@@ -149,19 +149,19 @@ void FDatasmithSceneGraphBuilder::FillAnchorActor(const TSharedRef<IDatasmithAct
 
 	// TODO: check ParentData and Index validity?
 	ActorData ParentData(ActorElement->GetName());
-	CADLibrary::FArchiveComponent& Component = SceneGraph->Components[*Index];
+	CADLibrary::FArchiveReference& Reference = SceneGraph->References[*Index];
 
 	TMap<FString, FString> InstanceNodeMetaDataMap;
 	FString ActorUUID;
 	FString ActorLabel;
-	GetNodeUUIDAndName(InstanceNodeMetaDataMap, Component.MetaData, Component.ObjectId, ParentData.Uuid, ActorUUID, ActorLabel);
+	GetNodeUUIDAndName(InstanceNodeMetaDataMap, Reference.MetaData, Reference.Id, ParentData.Uuid, ActorUUID, ActorLabel);
 
-	AddMetaData(ActorElement, InstanceNodeMetaDataMap, Component.MetaData);
+	AddMetaData(ActorElement, InstanceNodeMetaDataMap, Reference.MetaData);
 
-	ActorData ComponentData(*ActorUUID, ParentData);
-	DatasmithSceneGraphBuilderImpl::GetMainMaterial(Component.MetaData, ComponentData, bMaterialPropagationIsTopDown);
+	ActorData ReferenceData(*ActorUUID, ParentData);
+	DatasmithSceneGraphBuilderImpl::GetMainMaterial(Reference.MetaData, ReferenceData, bMaterialPropagationIsTopDown);
 
-	AddChildren(ActorElement, Component, ComponentData);
+	AddChildren(ActorElement, Reference, ReferenceData);
 
 	ActorElement->SetLabel(*ActorLabel);
 }
@@ -181,13 +181,13 @@ FDatasmithSceneBaseGraphBuilder::FDatasmithSceneBaseGraphBuilder(CADLibrary::FAr
 		ColorUIdToColorArchive.Reserve(SceneGraph->ColorHIdToColor.Num());
 		for(const auto& ColorPair : SceneGraph->ColorHIdToColor)
 		{
-			ColorUIdToColorArchive.Emplace(ColorPair.Value.UEMaterialName, ColorPair.Value);
+			ColorUIdToColorArchive.Emplace(ColorPair.Value.UEMaterialUId, ColorPair.Value);
 		}
 
 		MaterialUIdToMaterialArchive.Reserve(SceneGraph->MaterialHIdToMaterial.Num());
 		for (const auto& MaterialPair : SceneGraph->MaterialHIdToMaterial)
 		{
-			MaterialUIdToMaterialArchive.Emplace(MaterialPair.Value.UEMaterialName, MaterialPair.Value);
+			MaterialUIdToMaterialArchive.Emplace(MaterialPair.Value.UEMaterialUId, MaterialPair.Value);
 		}
 	}
 }
@@ -195,28 +195,28 @@ FDatasmithSceneBaseGraphBuilder::FDatasmithSceneBaseGraphBuilder(CADLibrary::FAr
 bool FDatasmithSceneBaseGraphBuilder::Build()
 {
 	FCadId RootId = 1;
-	const int32* Index = SceneGraph->CADIdToComponentIndex.Find(RootId);
+	const int32* Index = SceneGraph->CADIdToReferenceIndex.Find(RootId);
 	if (!Index)
 	{
 		return false;
 	}
 
 	ActorData Data(TEXT(""));
-	CADLibrary::FArchiveComponent& Component = SceneGraph->Components[*Index];
-	TSharedPtr< IDatasmithActorElement > RootActor = BuildComponent(Component, Data);
+	CADLibrary::FArchiveReference& Reference = SceneGraph->References[*Index];
+	TSharedPtr< IDatasmithActorElement > RootActor = BuildReference(Reference, Data);
 	DatasmithScene->AddActor(RootActor);
 
 	// Set ProductName, ProductVersion in DatasmithScene for Analytics purpose
 	// application_name is something like "Catia V5"
 	DatasmithScene->SetVendor(TEXT("Techsoft"));
 
-	if (const FString* ProductVersion = Component.MetaData.Find(TEXT("TechsoftVersion")))
+	if (const FString* ProductVersion = Reference.MetaData.Find(TEXT("TechsoftVersion")))
 	{
 		DatasmithScene->SetProductVersion(**ProductVersion);
 	}
 
 	FString ProductName;
-	const FString* ProductNamePtr = Component.MetaData.Find(TEXT("Input_Format_and_Emitter"));
+	const FString* ProductNamePtr = Reference.MetaData.Find(TEXT("Input_Format_and_Emitter"));
 	if(ProductNamePtr)
 	{
 		ProductName = *ProductNamePtr;
@@ -287,8 +287,8 @@ bool FDatasmithSceneBaseGraphBuilder::Build()
 
 TSharedPtr< IDatasmithActorElement >  FDatasmithSceneBaseGraphBuilder::BuildInstance(int32 InstanceIndex, const ActorData& ParentData)
 {
-	CADLibrary::FArchiveComponent* Reference = nullptr;
-	CADLibrary::FArchiveComponent EmptyReference;
+	CADLibrary::FArchiveReference* Reference = nullptr;
+	CADLibrary::FArchiveReference EmptyReference;
 
 	CADLibrary::FArchiveInstance& Instance = SceneGraph->Instances[InstanceIndex];
 
@@ -307,10 +307,10 @@ TSharedPtr< IDatasmithActorElement >  FDatasmithSceneBaseGraphBuilder::BuildInst
 					AncestorSceneGraphHash.Add(InstanceSceneGraphHash);
 
 					FCadId RootId = 1;
-					const int32* Index = SceneGraph->CADIdToComponentIndex.Find(RootId);
+					const int32* Index = SceneGraph->CADIdToReferenceIndex.Find(RootId);
 					if (Index)
 					{
-						Reference = &(SceneGraph->Components[*Index]);
+						Reference = &(SceneGraph->References[*Index]);
 					}
 				}
 			}
@@ -319,19 +319,19 @@ TSharedPtr< IDatasmithActorElement >  FDatasmithSceneBaseGraphBuilder::BuildInst
 		if(!Reference)
 		{
 			SceneGraph = InstanceSceneGraph;
-			const int32* Index = SceneGraph->CADIdToUnloadedComponentIndex.Find(Instance.ReferenceNodeId);
+			const int32* Index = SceneGraph->CADIdToUnloadedReferenceIndex.Find(Instance.ReferenceNodeId);
 			if (Index)
 			{
-				Reference = &(SceneGraph->UnloadedComponents[*Index]);
+				Reference = &(SceneGraph->UnloadedReferences[*Index]);
 			}
 		}
 	}
 	else 
 	{
-		const int32* Index = SceneGraph->CADIdToComponentIndex.Find(Instance.ReferenceNodeId);
+		const int32* Index = SceneGraph->CADIdToReferenceIndex.Find(Instance.ReferenceNodeId);
 		if (Index)
 		{
-			Reference = &(SceneGraph->Components[*Index]);
+			Reference = &(SceneGraph->References[*Index]);
 		}
 	}
 
@@ -342,7 +342,7 @@ TSharedPtr< IDatasmithActorElement >  FDatasmithSceneBaseGraphBuilder::BuildInst
 
 	FString ActorUUID;
 	FString ActorLabel;
-	GetNodeUUIDAndName(Instance.MetaData, Reference->MetaData, Instance.ObjectId, ParentData.Uuid, ActorUUID, ActorLabel);
+	GetNodeUUIDAndName(Instance.MetaData, Reference->MetaData, Instance.Id, ParentData.Uuid, ActorUUID, ActorLabel);
 
 	TSharedPtr<IDatasmithActorElement> Actor = CreateActor(*ActorUUID, *ActorLabel);
 	if (!Actor.IsValid())
@@ -382,7 +382,7 @@ TSharedPtr< IDatasmithActorElement >  FDatasmithSceneBaseGraphBuilder::CreateAct
 void FDatasmithSceneBaseGraphBuilder::GetNodeUUIDAndName(
 	const TMap<FString, FString>& InInstanceNodeMetaDataMap,
 	const TMap<FString, FString>& InReferenceNodeMetaDataMap,
-	int32 InComponentIndex,
+	int32 InReferenceIndex,
 	const TCHAR* InParentUEUUID,
 	FString& OutUEUUID,
 	FString& OutName
@@ -419,7 +419,7 @@ void FDatasmithSceneBaseGraphBuilder::GetNodeUUIDAndName(
 		OutName = "NoName";
 	}
 
-	FCadUuid UEUUID = HashCombine(GetTypeHash(InParentUEUUID), GetTypeHash(InComponentIndex));
+	FCadUuid UEUUID = HashCombine(GetTypeHash(InParentUEUUID), GetTypeHash(InReferenceIndex));
 
 	if (InstanceUUID)
 	{
@@ -450,13 +450,13 @@ void FDatasmithSceneBaseGraphBuilder::GetNodeUUIDAndName(
 	OutUEUUID = FString::Printf(TEXT("0x%08x"), UEUUID);
 }
 
-TSharedPtr<IDatasmithActorElement> FDatasmithSceneBaseGraphBuilder::BuildComponent(CADLibrary::FArchiveComponent& Component, const ActorData& ParentData)
+TSharedPtr<IDatasmithActorElement> FDatasmithSceneBaseGraphBuilder::BuildReference(CADLibrary::FArchiveReference& Reference, const ActorData& ParentData)
 {
 	TMap<FString, FString> InstanceNodeMetaDataMap;
 
 	FString ActorUUID;
 	FString ActorLabel;
-	GetNodeUUIDAndName(InstanceNodeMetaDataMap, Component.MetaData, Component.ObjectId, ParentData.Uuid, ActorUUID, ActorLabel);
+	GetNodeUUIDAndName(InstanceNodeMetaDataMap, Reference.MetaData, Reference.Id, ParentData.Uuid, ActorUUID, ActorLabel);
 
 	TSharedPtr< IDatasmithActorElement > Actor = CreateActor(*ActorUUID, *ActorLabel);
 	if (!Actor.IsValid())
@@ -464,14 +464,14 @@ TSharedPtr<IDatasmithActorElement> FDatasmithSceneBaseGraphBuilder::BuildCompone
 		return TSharedPtr< IDatasmithActorElement >();
 	}
 
-	AddMetaData(Actor, InstanceNodeMetaDataMap, Component.MetaData);
+	AddMetaData(Actor, InstanceNodeMetaDataMap, Reference.MetaData);
 
-	ActorData ComponentData(*ActorUUID, ParentData);
-	DatasmithSceneGraphBuilderImpl::GetMainMaterial(Component.MetaData, ComponentData, bMaterialPropagationIsTopDown);
+	ActorData ReferenceData(*ActorUUID, ParentData);
+	DatasmithSceneGraphBuilderImpl::GetMainMaterial(Reference.MetaData, ReferenceData, bMaterialPropagationIsTopDown);
 
-	AddChildren(Actor, Component, ComponentData);
+	AddChildren(Actor, Reference, ReferenceData);
 
-	DatasmithSceneGraphBuilderImpl::AddTransformToActor(Component, Actor, ImportParameters);
+	DatasmithSceneGraphBuilderImpl::AddTransformToActor(Reference, Actor, ImportParameters);
 
 	return Actor;
 }
@@ -489,7 +489,7 @@ TSharedPtr<IDatasmithActorElement> FDatasmithSceneBaseGraphBuilder::BuildBody(in
 
 	FString BodyUUID;
 	FString BodyLabel;
-	GetNodeUUIDAndName(InstanceNodeMetaDataMap, Body.MetaData, Body.ObjectId, ParentData.Uuid, BodyUUID, BodyLabel);
+	GetNodeUUIDAndName(InstanceNodeMetaDataMap, Body.MetaData, Body.Id, ParentData.Uuid, BodyUUID, BodyLabel);
 
 	// Apply materials on the current part
 	uint32 MaterialUuid = 0;
@@ -758,9 +758,9 @@ bool FDatasmithSceneBaseGraphBuilder::DoesActorHaveChildrenOrIsAStaticMesh(const
 }
 
 
-void FDatasmithSceneBaseGraphBuilder::AddChildren(TSharedPtr< IDatasmithActorElement > Actor, const CADLibrary::FArchiveComponent& Component, const ActorData& ParentData)
+void FDatasmithSceneBaseGraphBuilder::AddChildren(TSharedPtr< IDatasmithActorElement > Actor, const CADLibrary::FArchiveReference& Reference, const ActorData& ParentData)
 {
-	for (const int32 ChildId : Component.Children)
+	for (const int32 ChildId : Reference.Children)
 	{
 		if (int32* ChildNodeIndex = SceneGraph->CADIdToInstanceIndex.Find(ChildId))
 		{
diff --git a/Engine/Plugins/Enterprise/DatasmithCADImporter/Source/DatasmithCADTranslator/Public/DatasmithSceneGraphBuilder.h b/Engine/Plugins/Enterprise/DatasmithCADImporter/Source/DatasmithCADTranslator/Public/DatasmithSceneGraphBuilder.h
--- a/Engine/Plugins/Enterprise/DatasmithCADImporter/Source/DatasmithCADTranslator/Public/DatasmithSceneGraphBuilder.h
+++ b/Engine/Plugins/Enterprise/DatasmithCADImporter/Source/DatasmithCADTranslator/Public/DatasmithSceneGraphBuilder.h
@@ -65,11 +65,11 @@ public:
 
 protected:
 	TSharedPtr< IDatasmithActorElement > BuildInstance(int32 InstanceIndex, const ActorData& ParentData);
-	TSharedPtr< IDatasmithActorElement > BuildComponent(CADLibrary::FArchiveComponent& Component, const ActorData& ParentData);
+	TSharedPtr< IDatasmithActorElement > BuildReference(CADLibrary::FArchiveReference& Component, const ActorData& ParentData);
 	TSharedPtr< IDatasmithActorElement > BuildBody(int32 BodyIndex, const ActorData& ParentData);
 
 	void AddMetaData(TSharedPtr<IDatasmithActorElement> ActorElement, TMap<FString, FString>& InstanceNodeAttributeSetMap, TMap<FString, FString>& ReferenceNodeAttributeSetMap);
-	void AddChildren(TSharedPtr<IDatasmithActorElement> Actor, const CADLibrary::FArchiveComponent& Component, const ActorData& ParentData);
+	void AddChildren(TSharedPtr<IDatasmithActorElement> Actor, const CADLibrary::FArchiveReference& Component, const ActorData& ParentData);
 	bool DoesActorHaveChildrenOrIsAStaticMesh(const TSharedPtr< IDatasmithActorElement >& ActorElement);
 
 	TSharedPtr< IDatasmithUEPbrMaterialElement > GetDefaultMaterial();
