Improve  CADImport naming

diff --git a/Engine/Plugins/Enterprise/DatasmithCADImporter/Source/CADInterfaces/Private/TechSoftFileParser.cpp b/Engine/Plugins/Enterprise/DatasmithCADImporter/Source/CADInterfaces/Private/TechSoftFileParser.cpp
--- a/Engine/Plugins/Enterprise/DatasmithCADImporter/Source/CADInterfaces/Private/TechSoftFileParser.cpp
+++ b/Engine/Plugins/Enterprise/DatasmithCADImporter/Source/CADInterfaces/Private/TechSoftFileParser.cpp
@@ -1498,18 +1498,18 @@ void FTechSoftFileParser::BuildReferenceName(FEntityMetaData& ReferenceData)
 
 	if (Format == ECADFormat::JT)
 	{
-		if (TechSoftUtils::ReplaceOrAddNameValue(MetaData, TEXT("SDKName")))
+		if (TechSoftFileParserImpl::ReplaceOrAddNameValue(MetaData, TEXT("SDKName")))
 		{
 			return;
 		}
 	}
 
-	if (TechSoftUtils::CheckIfNameExists(MetaData))
+	if (TechSoftFileParserImpl::CheckIfNameExists(MetaData))
 	{
 		return;
 	}
 
-	if (TechSoftUtils::ReplaceOrAddNameValue(MetaData, TEXT("PartNumber")))
+	if (TechSoftFileParserImpl::ReplaceOrAddNameValue(MetaData, TEXT("PartNumber")))
 	{
 		return;
 	}
@@ -1625,7 +1625,7 @@ void FTechSoftFileParser::BuildBodyName(FEntityMetaData& BodyData, const FEntity
 {
 	TMap<FString, FString>& MetaData = BodyData.MetaData;
 
-	if (TechSoftUtils::CheckIfNameExists(MetaData))
+	if (TechSoftFileParserImpl::CheckIfNameExists(MetaData))
 	{
 		return;
 	}
diff --git a/Engine/Plugins/Enterprise/DatasmithCADImporter/Source/CADInterfaces/Private/TechSoftInterface.cpp b/Engine/Plugins/Enterprise/DatasmithCADImporter/Source/CADInterfaces/Private/TechSoftInterface.cpp
--- a/Engine/Plugins/Enterprise/DatasmithCADImporter/Source/CADInterfaces/Private/TechSoftInterface.cpp
+++ b/Engine/Plugins/Enterprise/DatasmithCADImporter/Source/CADInterfaces/Private/TechSoftInterface.cpp
@@ -13,7 +13,14 @@
 namespace CADLibrary
 {
 
-uint32 FTechSoftInterface::InvalidScriptIndex = -1;
+#ifdef USE_TECHSOFT_SDK
+const uint32 FTechSoftDefaultValue::Material = A3D_DEFAULT_MATERIAL_INDEX;
+const uint32 FTechSoftDefaultValue::Picture = A3D_DEFAULT_PICTURE_INDEX;
+const uint32 FTechSoftDefaultValue::RgbColor = A3D_DEFAULT_COLOR_INDEX;
+const uint32 FTechSoftDefaultValue::Style = A3D_DEFAULT_STYLE_INDEX;
+const uint32 FTechSoftDefaultValue::TextureApplication = A3D_DEFAULT_TEXTURE_APPLICATION_INDEX;
+const uint32 FTechSoftDefaultValue::TextureDefinition = A3D_DEFAULT_TEXTURE_DEFINITION_INDEX;
+#endif
 
 FTechSoftInterface& FTechSoftInterface::Get()
 {
diff --git a/Engine/Plugins/Enterprise/DatasmithCADImporter/Source/CADInterfaces/Private/TechSoftInterfaceUtils.inl b/Engine/Plugins/Enterprise/DatasmithCADImporter/Source/CADInterfaces/Private/TechSoftInterfaceUtils.inl
--- a/Engine/Plugins/Enterprise/DatasmithCADImporter/Source/CADInterfaces/Private/TechSoftInterfaceUtils.inl
+++ b/Engine/Plugins/Enterprise/DatasmithCADImporter/Source/CADInterfaces/Private/TechSoftInterfaceUtils.inl
@@ -169,8 +169,8 @@ private:
 			// there is a bijection between A3DTess3DData->m_psFaceTessData and A3DTopoShellData->m_ppFaces
 			Tessellation.PatchId = Index;
 
-			Tessellation.MaterialUId = FTechSoftInterface::InvalidScriptIndex;
-			if (FaceTessData.m_uiStyleIndexesSize > 0)
+			Tessellation.MaterialUId = 0;
+			if (FaceTessData.m_uiStyleIndexesSize != FTechSoftDefaultValue::Style)
 			{
 				// Store the StyleIndex on the MaterialName. It will be processed after tessellation
 				Tessellation.MaterialUId = FaceTessData.m_puiStyleIndexes[0];
diff --git a/Engine/Plugins/Enterprise/DatasmithCADImporter/Source/CADInterfaces/Private/TechSoftUtils.cpp b/Engine/Plugins/Enterprise/DatasmithCADImporter/Source/CADInterfaces/Private/TechSoftUtils.cpp
--- a/Engine/Plugins/Enterprise/DatasmithCADImporter/Source/CADInterfaces/Private/TechSoftUtils.cpp
+++ b/Engine/Plugins/Enterprise/DatasmithCADImporter/Source/CADInterfaces/Private/TechSoftUtils.cpp
@@ -262,7 +262,7 @@ void BuildCADMaterial(uint32 MaterialIndex, const A3DGraphStyleData& GraphStyleD
 #endif
 
 // Replicates the logic in FTechSoftFileParser::ExtractGraphStyleProperties
-void GetMaterialValues(uint32 StyleIndex, FCadUuid& OutColorName, FCadUuid& OutMaterialName)
+void GetMaterialValues(uint32 StyleIndex, FMaterialUId& OutColorName, FMaterialUId& OutMaterialName)
 {
 #if defined USE_TECHSOFT_SDK && !defined CADKERNEL_DEV
 	TUniqueTSObjFromIndex<A3DGraphStyleData> GraphStyleData(StyleIndex);
@@ -295,8 +295,8 @@ void GetMaterialValues(uint32 StyleIndex, FCadUuid& OutColorName, FCadUuid& OutM
 
 void RestoreMaterials(const TSharedPtr<FJsonObject>& DefaultValues, FBodyMesh& BodyMesh)
 {
-	FCadUuid DefaultColorName = 0;
-	FCadUuid DefaultMaterialName = 0;
+	FMaterialUId DefaultColorName = 0;
+	FMaterialUId DefaultMaterialName = 0;
 
 	DefaultValues->TryGetNumberField(JSON_ENTRY_COLOR_NAME, DefaultColorName);
 	DefaultValues->TryGetNumberField(JSON_ENTRY_MATERIAL_NAME, DefaultMaterialName);
@@ -310,21 +310,21 @@ void RestoreMaterials(const TSharedPtr<FJsonObject>& DefaultValues, FBodyMesh& B
 		uint32 CachedStyleIndex = Tessellation.MaterialUId;
 		Tessellation.MaterialUId = 0;
 
-		FCadUuid ColorName = DefaultColorName;
-		FCadUuid MaterialName = DefaultMaterialName;
+		FMaterialUId ColorUId = DefaultColorName;
+		FMaterialUId MaterialUId = DefaultMaterialName;
 
-		GetMaterialValues(CachedStyleIndex, ColorName, MaterialName);
+		GetMaterialValues(CachedStyleIndex, ColorUId, MaterialUId);
 
-		if (ColorName)
+		if (ColorUId)
 		{
-			Tessellation.ColorUId = ColorName;
-			BodyMesh.ColorSet.Add(ColorName);
+			Tessellation.ColorUId = ColorUId;
+			BodyMesh.ColorSet.Add(ColorUId);
 		}
 
-		if (MaterialName)
+		if (MaterialUId)
 		{
-			Tessellation.MaterialUId = MaterialName;
-			BodyMesh.MaterialSet.Add(MaterialName);
+			Tessellation.MaterialUId = MaterialUId;
+			BodyMesh.MaterialSet.Add(MaterialUId);
 		}
 	}
 }
@@ -770,27 +770,6 @@ FString CleanNameByRemoving_prt(const FString& Name)
 	}
 	return Name;
 }
-bool CheckIfNameExists(TMap<FString, FString>& MetaData)
-{
-	FString* NamePtr = MetaData.Find(TEXT("Name"));
-	if (NamePtr != nullptr)
-	{
-		return true;
-	}
-	return false;
-}
-bool ReplaceOrAddNameValue(TMap<FString, FString>& MetaData, const TCHAR* Key)
-{
-	FString* NamePtr = MetaData.Find(Key);
-	if (NamePtr != nullptr)
-	{
-		FString& Name = MetaData.FindOrAdd(TEXT("Name"));
-		Name = *NamePtr;
-		return true;
-	}
-	return false;
-}
-
 
 } // NS TechSoftUtils
 
diff --git a/Engine/Plugins/Enterprise/DatasmithCADImporter/Source/CADInterfaces/Public/TUniqueTechSoftObj.h b/Engine/Plugins/Enterprise/DatasmithCADImporter/Source/CADInterfaces/Public/TUniqueTechSoftObj.h
--- a/Engine/Plugins/Enterprise/DatasmithCADImporter/Source/CADInterfaces/Public/TUniqueTechSoftObj.h
+++ b/Engine/Plugins/Enterprise/DatasmithCADImporter/Source/CADInterfaces/Public/TUniqueTechSoftObj.h
@@ -66,7 +66,7 @@ public:
 	A3DStatus FillFrom(IndexerType EntityPtr)
 	{
 		ResetData();
-	
+
 		if (EntityPtr == GetDefaultIndexerValue())
 		{
 			Status = A3DStatus::A3D_ERROR;
@@ -94,7 +94,7 @@ public:
 		else
 		{
 			Status = Getter(EntityPtr, &Data, Forward<InArgTypes>(Args)...);
-			if(Status == A3DStatus::A3D_SUCCESS)
+			if (Status == A3DStatus::A3D_SUCCESS)
 			{
 				bDataFromTechSoft = true;
 			}
@@ -202,7 +202,7 @@ private:
 
 	void ResetData()
 	{
-		if(bDataFromTechSoft)
+		if (bDataFromTechSoft)
 		{
 			GetData(GetDefaultIndexerValue());
 		}
@@ -710,7 +710,16 @@ CADINTERFACES_API A3DStatus TUniqueTSObj<A3DUTF8Char*>::GetData(const A3DEntity*
 template<>
 CADINTERFACES_API const A3DEntity* TUniqueTSObj<A3DUTF8Char*>::GetDefaultIndexerValue() const;
 
+CADINTERFACES_API class FTechSoftDefaultValue
+{
+public:
+	static const uint32 Material;
+	static const uint32 Picture;
+	static const uint32 RgbColor;
+	static const uint32 Style;
+	static const uint32 TextureApplication;
+	static const uint32 TextureDefinition;
+};
 
 }
-
-#endif
\ No newline at end of file
+#endif
diff --git a/Engine/Plugins/Enterprise/DatasmithCADImporter/Source/CADInterfaces/Public/TechSoftInterface.h b/Engine/Plugins/Enterprise/DatasmithCADImporter/Source/CADInterfaces/Public/TechSoftInterface.h
--- a/Engine/Plugins/Enterprise/DatasmithCADImporter/Source/CADInterfaces/Public/TechSoftInterface.h
+++ b/Engine/Plugins/Enterprise/DatasmithCADImporter/Source/CADInterfaces/Public/TechSoftInterface.h
@@ -61,8 +61,6 @@ public:
 
 	bool InitializeKernel(const TCHAR* = TEXT(""));
 
-	static uint32 InvalidScriptIndex;
-
 private:
 
 	bool bIsExternal = false;
diff --git a/Engine/Plugins/Enterprise/DatasmithCADImporter/Source/CADLibrary/Private/MeshDescriptionHelper.cpp b/Engine/Plugins/Enterprise/DatasmithCADImporter/Source/CADLibrary/Private/MeshDescriptionHelper.cpp
--- a/Engine/Plugins/Enterprise/DatasmithCADImporter/Source/CADLibrary/Private/MeshDescriptionHelper.cpp
+++ b/Engine/Plugins/Enterprise/DatasmithCADImporter/Source/CADLibrary/Private/MeshDescriptionHelper.cpp
@@ -467,8 +467,8 @@ namespace CADLibrary
 				}
 
 				// Get the polygonGroup (material is preferred over color)
-				FCadUuid GraphicName = Tessellation.MaterialUId ? Tessellation.MaterialUId : Tessellation.ColorUId;
-				const FPolygonGroupID* PolygonGroupID = MaterialToPolygonGroupMapping.Find(GraphicName);
+				FMaterialUId GraphicUId = Tessellation.MaterialUId ? Tessellation.MaterialUId : Tessellation.ColorUId;
+				const FPolygonGroupID* PolygonGroupID = MaterialToPolygonGroupMapping.Find(GraphicUId);
 				if (PolygonGroupID == nullptr)
 				{
 					continue;
