Restored support for CAD files on Linux
Fixed Linux build breakage

diff --git a/Engine/Plugins/Enterprise/DatasmithCADImporter/Source/CADInterfaces/Public/TUniqueTechSoftObj.h b/Engine/Plugins/Enterprise/DatasmithCADImporter/Source/CADInterfaces/Public/TUniqueTechSoftObj.h
--- a/Engine/Plugins/Enterprise/DatasmithCADImporter/Source/CADInterfaces/Public/TUniqueTechSoftObj.h
+++ b/Engine/Plugins/Enterprise/DatasmithCADImporter/Source/CADInterfaces/Public/TUniqueTechSoftObj.h
@@ -710,7 +710,7 @@ CADINTERFACES_API A3DStatus TUniqueTSObj<A3DUTF8Char*>::GetData(const A3DEntity*
 template<>
 CADINTERFACES_API const A3DEntity* TUniqueTSObj<A3DUTF8Char*>::GetDefaultIndexerValue() const;
 
-CADINTERFACES_API class FTechSoftDefaultValue
+class CADINTERFACES_API FTechSoftDefaultValue
 {
 public:
 	static const uint32 Material;
