Fix dll directory push/pop missmatch

diff --git a/Engine/Plugins/Enterprise/DatasmithCADImporter/Source/CADInterfaces/Private/CADInterfacesModule.cpp b/Engine/Plugins/Enterprise/DatasmithCADImporter/Source/CADInterfaces/Private/CADInterfacesModule.cpp
--- a/Engine/Plugins/Enterprise/DatasmithCADImporter/Source/CADInterfaces/Private/CADInterfacesModule.cpp
+++ b/Engine/Plugins/Enterprise/DatasmithCADImporter/Source/CADInterfaces/Private/CADInterfacesModule.cpp
@@ -46,13 +46,12 @@ ECADInterfaceAvailability ICADInterfacesModule::GetAvailability()
 
 void FCADInterfacesModule::StartupModule()
 {
-	// determine directory paths
-	FString CADImporterDllPath = FPaths::Combine(FPaths::EnginePluginsDir(), TEXT("Enterprise/DatasmithCADImporter"), TEXT("Binaries"), FPlatformProcess::GetBinariesSubdirectory());
-	FPlatformProcess::PushDllDirectory(*CADImporterDllPath);
 
 #if WITH_EDITOR & defined(USE_TECHSOFT_SDK)
 	check(TechSoftLibHandle == nullptr);
 
+	// determine directory paths
+	FString CADImporterDllPath = FPaths::Combine(FPaths::EnginePluginsDir(), TEXT("Enterprise/DatasmithCADImporter"), TEXT("Binaries"), FPlatformProcess::GetBinariesSubdirectory());
 	FString TechSoftDllPath = FPaths::ConvertRelativePathToFull(FPaths::Combine(CADImporterDllPath, "TechSoft"));
 	FPlatformProcess::PushDllDirectory(*TechSoftDllPath);
 
@@ -77,8 +76,8 @@ void FCADInterfacesModule::StartupModule()
 			UE_LOG(LogCADInterfaces, Warning, TEXT("Failed to load required library %s. Plug-in will not be functional."), *TechSoftDll);
 		}
 
-		FPlatformProcess::PopDllDirectory(*TechSoftDllPath);
 	}
+	FPlatformProcess::PopDllDirectory(*TechSoftDllPath);
 #endif
 
 }
