[CAD Import] Add PlmXml format in Linux Platform

diff --git a/Engine/Plugins/Enterprise/DatasmithCADImporter/DatasmithCADImporter.uplugin b/Engine/Plugins/Enterprise/DatasmithCADImporter/DatasmithCADImporter.uplugin
--- a/Engine/Plugins/Enterprise/DatasmithCADImporter/DatasmithCADImporter.uplugin
+++ b/Engine/Plugins/Enterprise/DatasmithCADImporter/DatasmithCADImporter.uplugin
@@ -149,7 +149,8 @@
 			"Type": "Editor",
 			"LoadingPhase": "Default",
 			"PlatformAllowList": [
-				"Win64"
+				"Win64",
+				"Linux"
 			]
 		}
 	],
