Improve TechSoft parsing and SceneGraph building
#swarm https://p4-swarm.epicgames.net/reviews/20578906

diff --git a/Engine/Plugins/Enterprise/DatasmithCADImporter/Source/CADInterfaces/Private/CADSceneGraph.cpp b/Engine/Plugins/Enterprise/DatasmithCADImporter/Source/CADInterfaces/Private/CADSceneGraph.cpp
--- a/Engine/Plugins/Enterprise/DatasmithCADImporter/Source/CADInterfaces/Private/CADSceneGraph.cpp
+++ b/Engine/Plugins/Enterprise/DatasmithCADImporter/Source/CADInterfaces/Private/CADSceneGraph.cpp
@@ -12,17 +12,33 @@ namespace CADLibrary
 FArchive& operator<<(FArchive& Ar, FArchiveCADObject& Object)
 {
 	Ar << Object.Id;
+	Ar << Object.Label;
 	Ar << Object.MetaData;
 	Ar << Object.TransformMatrix;
+	Ar << Object.ColorUId;
+	Ar << Object.MaterialUId;
+	Ar << Object.bIsRemoved;
+	Ar << Object.Inheritance;
+	Ar << Object.Unit;
 	return Ar;
 }
 
+bool FArchiveCADObject::SetNameWithAttributeValue(const TCHAR* Key)
+{
+	FString* LabelPtr = MetaData.Find(Key);
+	if (LabelPtr != nullptr)
+	{
+		Label = *LabelPtr;
+		return true;
+	}
+	return false;
+}
+
 FArchive& operator<<(FArchive& Ar, FArchiveInstance& Instance) 
 {
 	Ar << (FArchiveCADObject&) Instance;
 	Ar << Instance.ReferenceNodeId;
 	Ar << Instance.bIsExternalReference;
-	Ar << Instance.ExternalReference;
 	return Ar;
 }
 
@@ -35,7 +51,8 @@ FArchive& operator<<(FArchive& Ar, FArchiveReference& Component)
 
 FArchive& operator<<(FArchive& Ar, FArchiveUnloadedReference& Unloaded) 
 {
-	Ar << (FArchiveReference&) Unloaded;
+	Ar << (FArchiveCADObject&) Unloaded;
+	Ar << Unloaded.ExternalFile;
 	return Ar;
 }
 
@@ -82,10 +99,7 @@ FArchive& operator<<(FArchive& Ar, FArchiveSceneGraph& SceneGraph)
 	Ar << SceneGraph.UnloadedReferences;
 	Ar << SceneGraph.Bodies;
 
-	Ar << SceneGraph.CADIdToInstanceIndex;
-	Ar << SceneGraph.CADIdToReferenceIndex;
-	Ar << SceneGraph.CADIdToUnloadedReferenceIndex;
-	Ar << SceneGraph.CADIdToBodyIndex;
+	Ar << SceneGraph.CADIdToIndex;
 
 	return Ar;
 }
@@ -110,6 +124,270 @@ void FArchiveSceneGraph::DeserializeMockUpFile(const TCHAR* Filename)
 	}
 }
 
+FArchiveInstance& FArchiveSceneGraph::AddInstance(const FArchiveCADObject& Parent)
+{
+	ensure(Instances.Num() < Instances.Max());
+
+	const int32 InstanceId = LastEntityId++;
+
+	int32 Index = Instances.Emplace(InstanceId, Parent);
+	CADIdToIndex.Add(Index);
+
+	return Instances[Index];
+}
+
+FArchiveInstance& FArchiveSceneGraph::GetInstance(FCadId CadId)
+{
+	ensure(CADIdToIndex.IsValidIndex(CadId));
+	const int32 Index = CADIdToIndex[CadId];
+	ensure(Instances.IsValidIndex(Index));
+	return Instances[Index];
+}
+
+void FArchiveSceneGraph::RemoveLastInstance()
+{
+	ensure(!Instances.IsEmpty());
+	LastEntityId--;
+	ensure(Instances.Last().Id == LastEntityId);
+	Instances.SetNum(Instances.Num() - 1, false);
+	CADIdToIndex.SetNum(CADIdToIndex.Num() - 1, false);
+}
+
+bool FArchiveSceneGraph::IsAInstance(FCadId CadId) const 
+{
+	if(!CADIdToIndex.IsValidIndex(CadId))
+	{
+		return false;
+	}
+	const int32 Index = CADIdToIndex[CadId];
+
+	if(!Instances.IsValidIndex(Index))
+	{
+		return false;
+	}
+	return Instances[Index].Id == CadId;
+}
+
+FArchiveReference& FArchiveSceneGraph::AddReference(FArchiveUnloadedReference& Reference)
+{
+	ensure(References.Num() < References.Max());
+
+	int32 Index = References.Emplace(Reference.Id, Reference);
+
+	FArchiveReference& NewReference = References[Index];
+	NewReference.MoveTemp(Reference);
+
+	// do not call here RemoveLastUnloadedReference otherwise the Reference.Id will be delete
+	UnloadedReferences.SetNum(UnloadedReferences.Num() - 1, false);
+
+	CADIdToIndex[Reference.Id] = Index;
+
+	return NewReference;
+}
+
+FArchiveReference& FArchiveSceneGraph::AddReference(FArchiveInstance& Parent)
+{
+	ensure(References.Num() < References.Max());
+
+	const int32 ReferenceId = LastEntityId++;
+
+	int32 Index = References.Emplace(ReferenceId, Parent);
+	CADIdToIndex.Add(Index);
+
+	Parent.bIsExternalReference = false;
+	Parent.ReferenceNodeId = ReferenceId;
+
+	return References[Index];
+}
+
+FArchiveReference& FArchiveSceneGraph::AddOccurence(FArchiveReference& Parent)
+{
+	FArchiveInstance& Instance = AddInstance(Parent);
+	Parent.AddChild(Instance.Id);
+	FArchiveReference& Reference = AddReference(Instance);
+	Instance.ReferenceNodeId = Reference.Id;
+	Reference.Unit = Parent.Unit;
+	return Reference;
+}
+
+void FArchiveSceneGraph::RemoveLastOccurence()
+{
+	RemoveLastReference();
+	RemoveLastInstance();
+}
+
+
+FArchiveReference& FArchiveSceneGraph::GetReference(FCadId CadId)
+{
+	ensure(CADIdToIndex.IsValidIndex(CadId));
+	const int32 Index = CADIdToIndex[CadId];
+	ensure(References.IsValidIndex(Index));
+	return References[Index];
+}
+
+bool FArchiveSceneGraph::IsAReference(FCadId CadId) const
+{
+	if (!CADIdToIndex.IsValidIndex(CadId))
+	{
+		return false;
+	}
+	const int32 Index = CADIdToIndex[CadId];
+
+	if (!References.IsValidIndex(Index))
+	{
+		return false;
+	}
+	return References[Index].Id == CadId;
+}
+
+void FArchiveSceneGraph::RemoveLastReference()
+{
+	ensure(!References.IsEmpty());
+	LastEntityId--;
+	ensure(References.Last().Id == LastEntityId);
+
+	References.SetNum(References.Num() - 1, false);
+	CADIdToIndex.SetNum(CADIdToIndex.Num() - 1, false);
+}
+
+void FArchiveReference::AddChild(const FCadId ChildId)
+{
+	Children.Add(ChildId);
+}
+
+int32 FArchiveReference::ChildrenCount()
+{
+	return Children.Num();
+}
+
+void FArchiveBody::Delete()
+{
+	bIsRemoved = true;
+	MetaData.Empty();
+	ColorFaceSet.Empty();
+	MaterialFaceSet.Empty();
+	Label.Empty();
+	MeshActorUId = 0;
+}
+
+
+void FArchiveReference::MoveTemp(FArchiveUnloadedReference& Reference)
+{
+	Id = Reference.Id;
+	Label = ::MoveTemp(Reference.Label);
+	MetaData = ::MoveTemp(Reference.MetaData);
+	TransformMatrix = ::MoveTemp(Reference.TransformMatrix);
+	ColorUId = Reference.ColorUId;
+	MaterialUId = Reference.MaterialUId;
+	Unit = Reference.Unit;
+}
+
+void FArchiveReference::CopyMetaData(FArchiveUnloadedReference& Reference)
+{
+	Label = Reference.Label;
+	MetaData = Reference.MetaData;
+	ColorUId = Reference.ColorUId;
+	MaterialUId = Reference.MaterialUId;
+	Unit = Reference.Unit;
+}
+
+FArchiveUnloadedReference& FArchiveSceneGraph::AddUnloadedReference(FArchiveInstance& Parent)
+{
+	ensure(UnloadedReferences.Num() < UnloadedReferences.Max());
+
+	const int32 ReferenceId = LastEntityId++;
+
+	int32 Index = UnloadedReferences.Emplace(ReferenceId, Parent);
+	CADIdToIndex.Add(Index);
+
+	Parent.bIsExternalReference = true;
+	Parent.ReferenceNodeId = ReferenceId;
+
+	return UnloadedReferences[Index];
+}
+
+FArchiveUnloadedReference& FArchiveSceneGraph::GetUnloadedReference(FCadId CadId)
+{
+	ensure(CADIdToIndex.IsValidIndex(CadId));
+	const int32 Index = CADIdToIndex[CadId];
+	ensure(UnloadedReferences.IsValidIndex(Index));
+	return UnloadedReferences[Index];
+}
+
+bool FArchiveSceneGraph::IsAUnloadedReference(FCadId CadId) const
+{
+	if (!CADIdToIndex.IsValidIndex(CadId))
+	{
+		return false;
+	}
+	const int32 Index = CADIdToIndex[CadId];
+
+	if (!UnloadedReferences.IsValidIndex(Index))
+	{
+		return false;
+	}
+	return UnloadedReferences[Index].Id == CadId;
+}
+
+void FArchiveSceneGraph::RemoveLastUnloadedReference()
+{
+	ensure(!UnloadedReferences.IsEmpty());
+	LastEntityId--;
+	ensure(UnloadedReferences.Last().Id == LastEntityId);
+
+	UnloadedReferences.SetNum(UnloadedReferences.Num() - 1, false);
+	CADIdToIndex.SetNum(CADIdToIndex.Num() - 1, false);
+}
+
+FArchiveBody& FArchiveSceneGraph::AddBody(FArchiveReference& Parent)
+{
+	ensure(Bodies.Num() < Bodies.Max());
+
+	const int32 BodyId = LastEntityId++;
+	int32 Index = Bodies.Emplace(BodyId, Parent);
+	CADIdToIndex.Add(Index);
+
+	return Bodies[Index];
+}
+
+FArchiveBody& FArchiveSceneGraph::GetBody(FCadId CadId)
+{
+	ensure(CADIdToIndex.IsValidIndex(CadId));
+	const int32 Index = CADIdToIndex[CadId];
+	ensure(Bodies.IsValidIndex(Index));
+	return Bodies[Index];
+}
+
+bool FArchiveSceneGraph::IsABody(FCadId CadId) const
+{
+	if (!CADIdToIndex.IsValidIndex(CadId))
+	{
+		return false;
+	}
+	const int32 Index = CADIdToIndex[CadId];
+
+	if (!Bodies.IsValidIndex(Index))
+	{
+		return false;
+	}
+	return Bodies[Index].Id == CadId;
+}
+
+void FArchiveSceneGraph::RemoveLastBody()
+{
+	ensure(!Bodies.IsEmpty());
+	LastEntityId--;
+	ensure(Bodies.Last().Id == LastEntityId);
+	Bodies.SetNum(Bodies.Num() - 1, false);
+	CADIdToIndex.SetNum(CADIdToIndex.Num() - 1, false);
+}
+
+void FArchiveSceneGraph::AddExternalReferenceFile(const FArchiveUnloadedReference& Reference)
+{
+	const int32 Index = CADIdToIndex[Reference.Id];
+	ensure(ExternalReferenceFiles.Add(Reference.ExternalFile) == Index);
+}
+
 }
 
 
diff --git a/Engine/Plugins/Enterprise/DatasmithCADImporter/Source/CADInterfaces/Private/TechSoftBridge.cpp b/Engine/Plugins/Enterprise/DatasmithCADImporter/Source/CADInterfaces/Private/TechSoftBridge.cpp
--- a/Engine/Plugins/Enterprise/DatasmithCADImporter/Source/CADInterfaces/Private/TechSoftBridge.cpp
+++ b/Engine/Plugins/Enterprise/DatasmithCADImporter/Source/CADInterfaces/Private/TechSoftBridge.cpp
@@ -4,11 +4,12 @@
 
 #ifdef USE_TECHSOFT_SDK
 
+#include "CADSceneGraph.h"
+#include "CADFileReport.h"
 #include "TechSoftInterface.h"
 #include "TechSoftFileParser.h"
 #include "TechSoftUtils.h"
 #include "TUniqueTechSoftObj.h"
-#include "CADFileReport.h"
 
 #include "CADKernel/Core/Session.h"
 
@@ -218,7 +219,7 @@ CADKernel::FBody* FTechSoftBridge::AddBody(A3DRiBrepModel* A3DBRepModel, TMap<FS
 	//CADKernel working unit is mm
 	BodyScale = InBodyScale * 10.;
 
-	FEntityMetaData BRepMetaData;
+	FArchiveCADObject BRepMetaData;
 	Parser.ExtractMetaData(A3DBRepModel, BRepMetaData);
 	FString* Name = MetaData.Find(TEXT("Name"));
 	if(Name != nullptr)
@@ -260,7 +261,7 @@ CADKernel::FBody* FTechSoftBridge::AddBody(A3DRiBrepModel* A3DBRepModel, TMap<FS
 
 void FTechSoftBridge::TraverseBrepData(const A3DTopoBrepData* A3DBrepData, TSharedRef<CADKernel::FBody>& Body)
 {
-	FEntityMetaData MetaData;
+	FArchiveCADObject MetaData;
 	Parser.ExtractMetaData(A3DBrepData, MetaData);
 
 	{
@@ -293,7 +294,7 @@ void FTechSoftBridge::TraverseBrepData(const A3DTopoBrepData* A3DBrepData, TShar
 
 void FTechSoftBridge::TraverseConnex(const A3DTopoConnex* A3DTopoConnex, TSharedRef<CADKernel::FBody>& Body)
 {
-	FEntityMetaData MetaData;
+	FArchiveCADObject MetaData;
 	Parser.ExtractMetaData(A3DTopoConnex, MetaData);
 
 	TUniqueTSObj<A3DTopoConnexData> TopoConnexData(A3DTopoConnex);
@@ -308,7 +309,7 @@ void FTechSoftBridge::TraverseConnex(const A3DTopoConnex* A3DTopoConnex, TShared
 
 void FTechSoftBridge::TraverseShell(const A3DTopoShell* A3DShell, TSharedRef<CADKernel::FBody>& Body)
 {
-	FEntityMetaData MetaData;
+	FArchiveCADObject MetaData;
 	Parser.ExtractMetaData(A3DShell, MetaData);
 
 	TSharedRef<CADKernel::FShell> Shell = CADKernel::FEntity::MakeShared<CADKernel::FShell>();
@@ -697,7 +698,7 @@ void FTechSoftBridge::AddFace(const A3DTopoFace* A3DFace, CADKernel::EOrientatio
 {
 	Report.FaceCount++;
 
-	FEntityMetaData MetaData;
+	FArchiveCADObject MetaData;
 	Parser.ExtractMetaData(A3DFace, MetaData);
 
 	TUniqueTSObj<A3DTopoFaceData> TopoFaceData(A3DFace);
@@ -796,7 +797,7 @@ TSharedPtr<CADKernel::FSurface> FTechSoftBridge::AddSurface(const A3DSurfBase* A
 {
 	Report.SurfaceCount++;
 
-	FEntityMetaData MetaData;
+	FArchiveCADObject MetaData;
 	Parser.ExtractMetaData(A3DSurface, MetaData);
 
 	A3DEEntityType Type;
@@ -1222,15 +1223,15 @@ TSharedPtr<CADKernel::FSurface> FTechSoftBridge::AddSurfaceAsNurbs(const A3DSurf
 
 }
 
-void FTechSoftBridge::AddMetadata(FEntityMetaData& MetaData, CADKernel::FTopologicalShapeEntity& Entity)
+void FTechSoftBridge::AddMetadata(FArchiveCADObject& MetaData, CADKernel::FTopologicalShapeEntity& Entity)
 {
 	FString* Name = MetaData.MetaData.Find(TEXT("Name"));
 	if(Name != nullptr)
 	{
 		Entity.SetName(*Name);
 	}
-	Entity.SetColorId(MetaData.ColorName);
-	Entity.SetMaterialId(MetaData.MaterialName);
+	Entity.SetColorId(MetaData.ColorUId);
+	Entity.SetMaterialId(MetaData.MaterialUId);
 }
 
 }
diff --git a/Engine/Plugins/Enterprise/DatasmithCADImporter/Source/CADInterfaces/Private/TechSoftBridge.h b/Engine/Plugins/Enterprise/DatasmithCADImporter/Source/CADInterfaces/Private/TechSoftBridge.h
--- a/Engine/Plugins/Enterprise/DatasmithCADImporter/Source/CADInterfaces/Private/TechSoftBridge.h
+++ b/Engine/Plugins/Enterprise/DatasmithCADImporter/Source/CADInterfaces/Private/TechSoftBridge.h
@@ -35,8 +35,8 @@ class FTopologicalShapeEntity;
 namespace CADLibrary
 {
 
+class FArchiveCADObject;
 class FTechSoftFileParser;
-struct FEntityMetaData;
 
 namespace TechSoftUtils
 {
@@ -243,7 +243,7 @@ private:
 
 	TSharedPtr<CADKernel::FCurve> AddCurveAsNurbs(const A3DCrvBase* A3DCurve, const TechSoftUtils::FUVReparameterization& UVReparameterization);
 
-	void AddMetadata(FEntityMetaData& EntityData, CADKernel::FTopologicalShapeEntity& Entity);
+	void AddMetadata(FArchiveCADObject& EntityData, CADKernel::FTopologicalShapeEntity& Entity);
 };
 }
 
diff --git a/Engine/Plugins/Enterprise/DatasmithCADImporter/Source/CADInterfaces/Private/TechSoftFileParser.cpp b/Engine/Plugins/Enterprise/DatasmithCADImporter/Source/CADInterfaces/Private/TechSoftFileParser.cpp
--- a/Engine/Plugins/Enterprise/DatasmithCADImporter/Source/CADInterfaces/Private/TechSoftFileParser.cpp
+++ b/Engine/Plugins/Enterprise/DatasmithCADImporter/Source/CADInterfaces/Private/TechSoftFileParser.cpp
@@ -94,28 +94,6 @@ inline void RemoveUnwantedChar(FString& StringToClean, TCHAR UnwantedChar)
 	StringToClean = MoveTemp(NewString);
 }
 
-bool CheckIfNameExists(TMap<FString, FString>& MetaData)
-{
-	FString* NamePtr = MetaData.Find(TEXT("Name"));
-	if (NamePtr != nullptr)
-	{
-		return true;
-	}
-	return false;
-}
-
-bool ReplaceOrAddNameValue(TMap<FString, FString>& MetaData, const TCHAR* Key)
-{
-	FString* NamePtr = MetaData.Find(Key);
-	if (NamePtr != nullptr)
-	{
-		FString& Name = MetaData.FindOrAdd(TEXT("Name"));
-		Name = *NamePtr;
-		return true;
-	}
-	return false;
-}
-
 // Functions used in traverse model process
 
 void TraverseAttribute(const A3DMiscAttributeData& AttributeData, TMap<FString, FString>& OutMetaData)
@@ -280,6 +258,7 @@ double ExtractUniformScale(FVector3d& Scale)
 
 FTechSoftFileParser::FTechSoftFileParser(FCADFileData& InCADData, const FString& EnginePluginsPath)
 	: CADFileData(InCADData)
+	, SceneGraph(InCADData.GetSceneGraphArchive())
 	, TechSoftInterface(FTechSoftInterface::Get())
 {
 	TechSoftInterface.InitializeKernel(*EnginePluginsPath);
@@ -385,12 +364,13 @@ ECADParsingResult FTechSoftFileParser::Process()
 	if (Result == ECADParsingResult::ProcessOk)
 	{
 		GenerateBodyMeshes();
-	}
 
-	FString TechSoftVersion = TechSoftInterface::GetTechSoftVersion();
-	if (!TechSoftVersion.IsEmpty() && CADFileData.ReferenceCount() > 0)
-	{
-		CADFileData.GetReferenceAt(0).MetaData.Add(TEXT("TechsoftVersion"), TechSoftVersion);
+		FString TechSoftVersion = TechSoftInterface::GetTechSoftVersion();
+		if (!TechSoftVersion.IsEmpty())
+		{
+			FArchiveReference& RootReference = SceneGraph.GetReference(1);
+			RootReference.MetaData.Add(TEXT("TechsoftVersion"), TechSoftVersion);
+		}
 	}
 
 	ModelFile.Reset();
@@ -412,7 +392,7 @@ void FTechSoftFileParser::GenerateBodyMeshes()
 	for (TPair<A3DRiRepresentationItem*, FCadId>& Entry : RepresentationItemsCache)
 	{
 		A3DRiRepresentationItem* RepresentationItemPtr = Entry.Key;
-		FArchiveBody& Body = CADFileData.GetBodyAt(Entry.Value);
+		FArchiveBody& Body = SceneGraph.GetBody(Entry.Value);
 		GenerateBodyMesh(RepresentationItemPtr, Body);
 	}
 }
@@ -420,6 +400,7 @@ void FTechSoftFileParser::GenerateBodyMeshes()
 void FTechSoftFileParser::GenerateBodyMesh(A3DRiRepresentationItem* Representation, FArchiveBody& Body)
 {
 	FBodyMesh& BodyMesh = CADFileData.AddBodyMesh(Body.Id, Body);
+	BodyMesh.AddGraphicPropertiesFrom(Body);
 
 	uint32 NewBRepCount = 0;
 	A3DRiBrepModel** NewBReps = nullptr;
@@ -440,52 +421,34 @@ void FTechSoftFileParser::GenerateBodyMesh(A3DRiRepresentationItem* Representati
 	{
 		for (uint32 Index = 0; Index < NewBRepCount; ++Index)
 		{
-			TechSoftUtils::FillBodyMesh(NewBReps[Index], CADFileData.GetImportParameters(), Body.BodyUnit, BodyMesh);
+			TechSoftUtils::FillBodyMesh(NewBReps[Index], CADFileData.GetImportParameters(), Body.Unit, BodyMesh);
 		}
 	} 
 	else
 	{
-		TechSoftUtils::FillBodyMesh(Representation, CADFileData.GetImportParameters(), Body.BodyUnit, BodyMesh);
+		TechSoftUtils::FillBodyMesh(Representation, CADFileData.GetImportParameters(), Body.Unit, BodyMesh);
 	}
 
 	if (BodyMesh.TriangleCount == 0)
 	{
 		// the mesh of the body is empty, the body is deleted.
-		// Todo (jira UETOOL-5148): add a boolean in Body to flag that the body should not be build
-		Body.ParentId = 0;
-		Body.MeshActorUId = 0;
+		Body.Delete();
 	}
 
 	// Convert material
-	FCadUuid DefaultColorUId = Body.ColorFaceSet.Num() > 0 ? *Body.ColorFaceSet.begin() : 0;
-	FCadUuid DefaultMaterialUId = Body.MaterialFaceSet.Num() > 0 ? *Body.MaterialFaceSet.begin() : 0;
-
 	for (FTessellationData& Tessellation : BodyMesh.Faces)
 	{
-		FCadUuid ColorUId = DefaultColorUId;
-		FCadUuid MaterialUId = DefaultMaterialUId;
-
 		// Extract proper color or material based on style index
-		uint32 CachedStyleIndex = Tessellation.MaterialUId;
+		FMaterialUId CachedStyleIndex = Tessellation.MaterialUId;
 		Tessellation.MaterialUId = 0;
 
-		constexpr uint32 GraphStyleDataDefaultValue =  65535;
-		if (CachedStyleIndex != GraphStyleDataDefaultValue)
+		if (CachedStyleIndex != FTechSoftDefaultValue::Style)
 		{
-			ExtractGraphStyleProperties(CachedStyleIndex, ColorUId, MaterialUId);
+			ExtractGraphStyleProperties(CachedStyleIndex, Tessellation);
 		}
+		Tessellation.DefineGraphicsPropertiesFromNoOverwrite(Body);
 
-		if (ColorUId)
-		{
-			Tessellation.ColorUId = ColorUId;
-			BodyMesh.ColorSet.Add(ColorUId);
-		}
-
-		if (MaterialUId)
-		{
-			Tessellation.MaterialUId = MaterialUId;
-			BodyMesh.MaterialSet.Add(MaterialUId);
-		}
+		BodyMesh.AddGraphicPropertiesFrom(Tessellation);
 	}
 
 	Body.ColorFaceSet = BodyMesh.ColorSet;
@@ -505,11 +468,18 @@ void FTechSoftFileParser::GenerateBodyMesh(A3DRiRepresentationItem* Representati
 			// Save body unit and default color and material attributes in a json string
 			// This will be used when the file is reloaded
 			JsonObject->SetNumberField(JSON_ENTRY_BODY_UNIT, Body.BodyUnit);
-			JsonObject->SetNumberField(JSON_ENTRY_COLOR_NAME, DefaultColorUId);
-			JsonObject->SetNumberField(JSON_ENTRY_MATERIAL_NAME, DefaultMaterialUId);
+
+			if(Body.ColorUId)
+			{
+				JsonObject->SetNumberField(JSON_ENTRY_COLOR_NAME, Body.ColorUId);
+			}
+			if(Body.MaterialUId)
+			{
+				JsonObject->SetNumberField(JSON_ENTRY_MATERIAL_NAME, Body.MaterialUId);
+			}
 
 			FString JsonString;
-			TSharedRef< TJsonWriter< TCHAR, TPrettyJsonPrintPolicy<TCHAR> > > JsonWriter = TJsonWriterFactory< TCHAR, TPrettyJsonPrintPolicy<TCHAR> >::Create(&JsonString);
+			TSharedRef<TJsonWriter<TCHAR, TPrettyJsonPrintPolicy<TCHAR>>> JsonWriter = TJsonWriterFactory<TCHAR, TPrettyJsonPrintPolicy<TCHAR>>::Create(&JsonString);
 
 			FJsonSerializer::Serialize(JsonObject.ToSharedRef(), JsonWriter);
 
@@ -525,11 +495,9 @@ void FTechSoftFileParser::ReserveCADFileData()
 
 	CADFileData.ReserveBodyMeshes(ComponentCount[EComponentType::Body]);
 
-	FArchiveSceneGraph& SceneGraphArchive = CADFileData.GetSceneGraphArchive();
-	SceneGraphArchive.Reserve(ComponentCount[EComponentType::Occurrence], ComponentCount[EComponentType::Reference], ComponentCount[EComponentType::Body]);
-
+	SceneGraph.Reserve(ComponentCount);
 	uint32 MaterialNum = CountColorAndMaterial();
-	SceneGraphArchive.MaterialHIdToMaterial.Reserve(MaterialNum);
+	SceneGraph.MaterialHIdToMaterial.Reserve(MaterialNum);
 }
 
 void FTechSoftFileParser::CountUnderModel()
@@ -540,7 +508,7 @@ void FTechSoftFileParser::CountUnderModel()
 		return;
 	}
 
-	ComponentCount[EComponentType::Occurrence] ++;
+	ComponentCount[EComponentType::Reference] ++;
 
 	for (uint32 Index = 0; Index < ModelFileData->m_uiPOccurrencesSize; ++Index)
 	{
@@ -563,28 +531,29 @@ ECADParsingResult FTechSoftFileParser::TraverseModel()
 		return ECADParsingResult::ProcessFailed;
 	}
 
-	FEntityMetaData MetaData;
-	ExtractMetaData(ModelFile.Get(), MetaData);
-	ExtractSpecificMetaData(ModelFile.Get(), MetaData);
-	BuildReferenceName(MetaData);
+	FArchiveInstance EmptyInstance;
+	FArchiveReference& Reference = SceneGraph.AddReference(EmptyInstance);
+	ExtractMetaData(ModelFile.Get(), Reference);
+	ExtractSpecificMetaData(ModelFile.Get(), Reference);
+	BuildReferenceName(Reference);
+	Reference.Unit = FileUnit;
 
 	for (uint32 Index = 0; Index < ModelFileData->m_uiPOccurrencesSize; ++Index)
 	{
 		if (IsConfigurationSet(ModelFileData->m_ppPOccurrences[Index]))
 		{
-			TraverseConfigurationSet(ModelFileData->m_ppPOccurrences[Index], FileUnit);
+			TraverseConfigurationSet(ModelFileData->m_ppPOccurrences[Index], Reference);
 		}
 		else
 		{
-			FString* ModelNamePtr = MetaData.MetaData.Find(TEXT("Name"));
-			TraverseReference(ModelFileData->m_ppPOccurrences[Index], ModelNamePtr ? *ModelNamePtr : FString(), FMatrix::Identity, FileUnit);
+			TraverseReference(ModelFileData->m_ppPOccurrences[Index], Reference);
 		}
 	}
 
 	return ECADParsingResult::ProcessOk;
 }
 
-void FTechSoftFileParser::TraverseConfigurationSet(const A3DAsmProductOccurrence* ConfigurationSetPtr, double ParentUnit)
+void FTechSoftFileParser::TraverseConfigurationSet(const A3DAsmProductOccurrence* ConfigurationSetPtr, FArchiveReference& Reference)
 {
 	TUniqueTSObj<A3DAsmProductOccurrenceData> ConfigurationSetData(ConfigurationSetPtr);
 	if (!ConfigurationSetData.IsValid())
@@ -592,20 +561,16 @@ void FTechSoftFileParser::TraverseConfigurationSet(const A3DAsmProductOccurrence
 		return;
 	}
 
-	FEntityMetaData MetaData;
-	ExtractMetaData(ConfigurationSetPtr, MetaData);
-	ExtractSpecificMetaData(ConfigurationSetPtr, MetaData);
-	BuildReferenceName(MetaData);
-	FString ConfigurationName = MetaData.MetaData.FindOrAdd(TEXT("Name"));
+	ExtractMetaData(ConfigurationSetPtr, Reference);
+	ExtractSpecificMetaData(ConfigurationSetPtr, Reference);
+	BuildReferenceName(Reference);
 
 	const FString& ConfigurationToLoad = CADFileData.GetCADFileDescription().GetConfiguration();
 
-	FMatrix TransformMatrix = FMatrix::Identity;;
-	double OccurenceUnit = ParentUnit;
 	A3DMiscTransformation* Location = ConfigurationSetData->m_pLocation;
 	if (Location)
 	{
-		TransformMatrix = ExtractTransformation(Location, OccurenceUnit);
+		ExtractTransformation(Location, Reference);
 	}
 
 	TUniqueTSObj<A3DAsmProductOccurrenceData> ConfigurationData;
@@ -622,12 +587,11 @@ void FTechSoftFileParser::TraverseConfigurationSet(const A3DAsmProductOccurrence
 			bool bIsConfigurationToLoad = false;
 			if (!ConfigurationToLoad.IsEmpty())
 			{
-				FEntityMetaData ConfigurationMetaData;
-				ExtractMetaData(ConfigurationSetData->m_ppPOccurrences[Index], ConfigurationMetaData);
-				const FString* Configuration = ConfigurationMetaData.MetaData.Find(TEXT("SDKName"));
-				if (Configuration)
+				FArchiveCADObject Configuration;
+				ExtractMetaData(ConfigurationSetData->m_ppPOccurrences[Index], Configuration);
+				if (!Configuration.Label.IsEmpty())
 				{
-					bIsConfigurationToLoad = (Configuration->Equals(ConfigurationToLoad));
+					bIsConfigurationToLoad = (Configuration.Label.Equals(ConfigurationToLoad));
 				}
 			}
 			else
@@ -637,7 +601,7 @@ void FTechSoftFileParser::TraverseConfigurationSet(const A3DAsmProductOccurrence
 
 			if (bIsConfigurationToLoad)
 			{
-				TraverseReference(ConfigurationSetData->m_ppPOccurrences[Index], ConfigurationName, TransformMatrix, OccurenceUnit);
+				TraverseReference(ConfigurationSetData->m_ppPOccurrences[Index], Reference);
 				return;
 			}
 		}
@@ -656,7 +620,7 @@ void FTechSoftFileParser::TraverseConfigurationSet(const A3DAsmProductOccurrence
 
 			if (ConfigurationData->m_uiProductFlags & A3D_PRODUCT_FLAG_CONFIG)
 			{
-				TraverseReference(ConfigurationSetData->m_ppPOccurrences[Index], ConfigurationName, TransformMatrix, OccurenceUnit);
+				TraverseReference(ConfigurationSetData->m_ppPOccurrences[Index], Reference);
 			}
 		}
 	}
@@ -686,12 +650,11 @@ void FTechSoftFileParser::CountUnderConfigurationSet(const A3DAsmProductOccurren
 			bool bIsConfigurationToLoad = false;
 			if (!ConfigurationToLoad.IsEmpty())
 			{
-				FEntityMetaData ConfigurationMetaData;
-				ExtractMetaData(ConfigurationSetData->m_ppPOccurrences[Index], ConfigurationMetaData);
-				const FString* ConfigurationName = ConfigurationMetaData.MetaData.Find(TEXT("SDKName"));
-				if (ConfigurationName)
+				FArchiveCADObject Configuration;
+				ExtractMetaData(ConfigurationSetData->m_ppPOccurrences[Index], Configuration);
+				if (!Configuration.Label.IsEmpty())
 				{
-					bIsConfigurationToLoad = (ConfigurationName->Equals(ConfigurationToLoad));
+					bIsConfigurationToLoad = (Configuration.Label.Equals(ConfigurationToLoad));
 				}
 			}
 			else
@@ -726,154 +689,72 @@ void FTechSoftFileParser::CountUnderConfigurationSet(const A3DAsmProductOccurren
 	}
 }
 
-void FTechSoftFileParser::TraverseReference(const A3DAsmProductOccurrence* ReferencePtr, const FString& RootName, const FMatrix& ParentMatrix, double ParentUnit)
+void FTechSoftFileParser::TraverseReference(const A3DAsmProductOccurrence* A3DReferencePtr, FArchiveReference& Reference)
 {
-	FEntityMetaData MetaData;
-	ExtractMetaData(ReferencePtr, MetaData);
-
-	if (MetaData.bRemoved || !MetaData.bShow)
+	TUniqueTSObj<A3DAsmProductOccurrenceData> ReferenceData(A3DReferencePtr);
+	if (!ReferenceData.IsValid())
 	{
+		SceneGraph.RemoveLastReference();
 		return;
 	}
 
-	ExtractSpecificMetaData(ReferencePtr, MetaData);
-
-	if (!RootName.IsEmpty())
-	{
-		MetaData.MetaData.FindOrAdd(TEXT("Name"), RootName);
-	}
-	else
+	if (Reference.bIsRemoved || !Reference.bShow)
 	{
-		BuildReferenceName(MetaData);
+		SceneGraph.RemoveLastReference();
+		return;
 	}
 
-	FArchiveInstance EmptyInstance;
-	FArchiveReference& Reference = FTechSoftFileParser::AddReference(MetaData, EmptyInstance);
+	ExtractSpecificMetaData(A3DReferencePtr, Reference);
 
-	TUniqueTSObj<A3DAsmProductOccurrenceData> ReferenceData(ReferencePtr);
-	if (!ReferenceData.IsValid())
+	if (!Reference.IsNameDefined())
 	{
-		return;
+		BuildReferenceName(Reference);
 	}
 
 	FMatrix ReferenceMatrix = FMatrix::Identity;
-	double ReferenceUnit = ParentUnit;
-
 	A3DMiscTransformation* Location = ReferenceData->m_pLocation;
 	if (Location)
 	{
-		ReferenceMatrix = ExtractTransformation(Location, ReferenceUnit);
+		ExtractTransformation(Location, Reference);
 	}
+	Reference.TransformMatrix = Reference.TransformMatrix * ReferenceMatrix;
 
-	Reference.TransformMatrix = ParentMatrix * ReferenceMatrix;
-
-	const FString& OccurenceNameBase = Reference.MetaData.FindOrAdd(TEXT("Name"));
-	for (uint32 OccurenceIndex = 0; OccurenceIndex < ReferenceData->m_uiPOccurrencesSize; ++OccurenceIndex)
-	{
-		int32 ChildrenId = TraverseOccurrence(ReferenceData->m_ppPOccurrences[OccurenceIndex], OccurenceNameBase + TEXT("_") + FString::FromInt(OccurenceIndex+1), ReferenceUnit);
-		Reference.Children.Add(ChildrenId);
-	}
-
-	if (ReferenceData->m_pPart)
-	{
-		TraversePartDefinition(ReferenceData->m_pPart, Reference, ReferenceUnit);
-	}
-}
-
-FArchiveInstance& FTechSoftFileParser::AddInstance(FEntityMetaData& InstanceMetaData)
-{
-	FCadId InstanceId = LastEntityId++;
-	int32 InstanceIndex = CADFileData.AddInstance(InstanceId);
-	FArchiveInstance& Instance = CADFileData.GetInstanceAt(InstanceIndex);
-	Instance.MetaData = MoveTemp(InstanceMetaData.MetaData);
-	return Instance;
-}
-
-FArchiveReference& FTechSoftFileParser::AddReference(FEntityMetaData& ReferenceMetaData, FArchiveInstance& Instance)
-{
-	FCadId ReferenceId = LastEntityId++;
-	int32 ReferenceIndex = CADFileData.AddReference(ReferenceId);
-	FArchiveReference& Prototype = CADFileData.GetReferenceAt(ReferenceIndex);
-	Prototype.MetaData = MoveTemp(ReferenceMetaData.MetaData);
-
-	Instance.ReferenceNodeId = ReferenceId;
-
-	return Prototype;
+	FArchiveInstance EmptyInstance;
+	ProcessReference(A3DReferencePtr, EmptyInstance, Reference);
 }
 
-FArchiveUnloadedReference& FTechSoftFileParser::AddUnloadedReference(FEntityMetaData& ReferenceMetaData, FArchiveInstance& Instance)
+void FTechSoftFileParser::ProcessUnloadedReference(const FArchiveInstance& Instance, FArchiveUnloadedReference& Reference)
 {
-	FCadId ReferenceId = LastEntityId++;
-	int32 ReferenceIndex = CADFileData.AddUnloadedReference(ReferenceId);
-	FArchiveUnloadedReference& Reference = CADFileData.GetUnloadedReferenceAt(ReferenceIndex);
-	Reference.MetaData = MoveTemp(ReferenceMetaData.MetaData);
-
-	Instance.bIsExternalReference = true;
-	Instance.ReferenceNodeId = ReferenceId;
-
 	// Make sure that the ExternalFile path is right otherwise try to find the file and update.
-	TechSoftFileParserImpl::UpdateFileDescriptor(ReferenceMetaData.ExternalFile);
+	TechSoftFileParserImpl::UpdateFileDescriptor(Reference.ExternalFile);
 
-	Instance.ExternalReference = ReferenceMetaData.ExternalFile;
-
-	if (Format == ECADFormat::SOLIDWORKS /*|| CADFileData.FileFormat() == ECADFormat::JT*/)
+	switch (Format)
 	{
-		if (FString* ConfigurationName = Reference.MetaData.Find(TEXT("ConfigurationName")))
+	case ECADFormat::SOLIDWORKS:
+		if (const FString* ConfigurationName = Instance.MetaData.Find(TEXT("ConfigurationName")))
 		{
-			Instance.ExternalReference.SetConfiguration(*ConfigurationName);
+			Reference.ExternalFile.SetConfiguration(*ConfigurationName);
 		}
-	}
-
-	CADFileData.AddExternalRef(Instance.ExternalReference);
-
-	return Reference;
-}
-
-FArchiveReference& FTechSoftFileParser::AddOccurence(FEntityMetaData& InstanceMetaData, FCadId& OutReferenceId)
-{
-	FArchiveInstance& Instance = AddInstance(InstanceMetaData);
-	OutReferenceId = Instance.Id;
-	FEntityMetaData ReferenceMetaData;
-	FArchiveReference& Prototype = AddReference(ReferenceMetaData, Instance);
-	return Prototype;
-}
-
-FArchiveReference& FTechSoftFileParser::AddOccurence(FEntityMetaData& InstanceMetaData, FEntityMetaData& ReferenceMetaData, FCadId& OutReferenceId)
-{
-	FArchiveInstance& Instance = AddInstance(InstanceMetaData);
-	OutReferenceId = Instance.Id;
-	FArchiveReference& Prototype = AddReference(ReferenceMetaData, Instance);
-	return Prototype;
-}
+		break;
 
-int32 FTechSoftFileParser::AddBody(FEntityMetaData& BodyMetaData, const FMatrix& Matrix, const FCadId ParentId, double BodyUnit)
-{
-	FCadId BodyId = LastEntityId++;
-	int32 BodyIndex = CADFileData.AddBody(BodyId);
-	FArchiveBody& Body = CADFileData.GetBodyAt(BodyIndex);
-	Body.ParentId = ParentId;
-	Body.TransformMatrix = Matrix;
-	Body.BodyUnit = BodyUnit;
-	Body.MetaData = MoveTemp(BodyMetaData.MetaData);
-	if (BodyMetaData.ColorName != 0)
-	{
-		Body.ColorFaceSet.Add(BodyMetaData.ColorName);
-	}
-	if (BodyMetaData.MaterialName != 0)
-	{
-		Body.MaterialFaceSet.Add(BodyMetaData.MaterialName);
+	case ECADFormat::NX:
+		// todo
+		break;
+	default:
+		break;
 	}
-	return BodyIndex;
+
+	SceneGraph.AddExternalReferenceFile(Reference);
 }
 
-FCadId FTechSoftFileParser::TraverseOccurrence(const A3DAsmProductOccurrence* OccurrencePtr, const FString& DefaultOccurrenceName, double ParentUnit)
+void FTechSoftFileParser::TraverseOccurrence(const A3DAsmProductOccurrence* OccurrencePtr, FArchiveReference& ParentReference)
 {
 	// first product occurence with m_pPart != nullptr || m_uiPOccurrencesSize > 0
 	const A3DAsmProductOccurrence* CachedOccurrencePtr = OccurrencePtr;
 	TUniqueTSObj<A3DAsmProductOccurrenceData> OccurrenceData(OccurrencePtr);
 	if (!OccurrenceData.IsValid())
 	{
-		return 0;
+		return;
 	}
 
 	bool bContinueTraverse = OccurrenceData->m_pPrototype
@@ -882,59 +763,64 @@ FCadId FTechSoftFileParser::TraverseOccurrence(const A3DAsmProductOccurrence* Oc
 		|| OccurrenceData->m_uiPOccurrencesSize > 0;
 	if (!bContinueTraverse)
 	{
-		return 0;
+		return;
 	}
 
-	// Todo
-	// the life time of FEntityMetaData end at the creation of FArchiveInstance.
-	// This variable should not be accessible after
-	FEntityMetaData InstanceMetaData;
-	ExtractMetaData(OccurrencePtr, InstanceMetaData);
+	FArchiveInstance& Instance = SceneGraph.AddInstance(ParentReference);
+	ExtractMetaData(OccurrencePtr, Instance);
+	Instance.DefineGraphicsPropertiesFromNoOverwrite(ParentReference);
 
-	if (InstanceMetaData.bRemoved || !InstanceMetaData.bShow)
+	if (Instance.bIsRemoved || !Instance.bShow)
 	{
-		return 0;
+		SceneGraph.RemoveLastInstance();
+		return;
 	}
 
-	ExtractSpecificMetaData(OccurrencePtr, InstanceMetaData);
-	BuildInstanceName(InstanceMetaData, DefaultOccurrenceName);
+	ParentReference.AddChild(Instance.Id);
 
-	FArchiveInstance& Instance = AddInstance(InstanceMetaData);
+	ExtractSpecificMetaData(OccurrencePtr, Instance);
+	BuildInstanceName(Instance, ParentReference);
 
 	A3DMiscTransformation* Location = OccurrenceData->m_pLocation;
 
-	// Todo
-	// the life time of FEntityMetaData end at the creation of FArchive(Unloaded)Reference.
-	// This variable should not be accessible after
-	FEntityMetaData PrototypeMetaData;
+	FArchiveUnloadedReference& UnloadedReference = SceneGraph.AddUnloadedReference(Instance);
 	if (OccurrenceData->m_pPrototype)
 	{
-		ProcessPrototype(OccurrenceData->m_pPrototype, PrototypeMetaData, &Location);
+		ProcessPrototype(OccurrenceData->m_pPrototype, UnloadedReference, &Location);
+	}
+	else
+	{
+		UnloadedReference.bIsUnloaded = false;
 	}
 
-	double OccurrenceUnit = ParentUnit;
+	Instance.Unit = ParentReference.Unit;
 	if (Location)
 	{
-		Instance.TransformMatrix = ExtractTransformation(Location, OccurrenceUnit);
+		ExtractTransformation(Location, Instance);
 	}
 
-	if (PrototypeMetaData.bUnloaded)
+	if (UnloadedReference.bIsUnloaded)
 	{
-		FArchiveUnloadedReference& UnloadedReference = AddUnloadedReference(PrototypeMetaData, Instance);
-		return Instance.Id;
+		ProcessUnloadedReference(Instance, UnloadedReference);
 	}
-	
-	// If the prototype hasn't name, set its name with the name of the instance 
+	else
 	{
-		FString& PrototypeName = PrototypeMetaData.MetaData.FindOrAdd(TEXT("Name"));
-		if (PrototypeName.IsEmpty())
-		{
-			const FString& InstanceName = Instance.MetaData.FindOrAdd(TEXT("Name"));
-			PrototypeName = InstanceName;
-		}
+		FArchiveReference& NewReference = SceneGraph.AddReference(UnloadedReference);
+		Instance.bIsExternalReference = false;
+		ProcessReference(OccurrencePtr, Instance, NewReference);
 	}
+}
 
-	FArchiveReference& Reference = AddReference(PrototypeMetaData, Instance);
+void FTechSoftFileParser::ProcessReference(const A3DAsmProductOccurrence* OccurrencePtr, FArchiveInstance& Instance, FArchiveReference& Reference)
+{
+	const A3DAsmProductOccurrence* CachedOccurrencePtr = OccurrencePtr;
+	TUniqueTSObj<A3DAsmProductOccurrenceData> OccurrenceData(OccurrencePtr);
+
+	// If the prototype hasn't name, set its name with the name of the instance 
+	if (!Reference.IsNameDefined())
+	{
+		Reference.Label = Instance.Label;
+	}
 
 	while (OccurrenceData->m_pPrototype != nullptr && OccurrenceData->m_pPart == nullptr && OccurrenceData->m_uiPOccurrencesSize == 0 && OccurrenceData->m_pExternalData == nullptr)
 	{
@@ -944,7 +830,7 @@ FCadId FTechSoftFileParser::TraverseOccurrence(const A3DAsmProductOccurrence* Oc
 
 	if(OccurrenceData->m_pPart == nullptr && OccurrenceData->m_uiPOccurrencesSize == 0 && OccurrenceData->m_pExternalData == nullptr)
 	{
-		return Instance.Id;
+		return;
 	}
 
 	// Add part
@@ -955,7 +841,7 @@ FCadId FTechSoftFileParser::TraverseOccurrence(const A3DAsmProductOccurrence* Oc
 	if (OccurrenceData->m_pPart != nullptr)
 	{
 		A3DAsmPartDefinition* PartDefinition = OccurrenceData->m_pPart;
-		TraversePartDefinition(PartDefinition, Reference, OccurrenceUnit);
+		TraversePartDefinition(PartDefinition, Reference);
 	}
 
 	// Add Occurrence's Children
@@ -967,12 +853,9 @@ FCadId FTechSoftFileParser::TraverseOccurrence(const A3DAsmProductOccurrence* Oc
 
 	uint32 ChildrenCount = OccurrenceData->m_uiPOccurrencesSize;
 	A3DAsmProductOccurrence** Children = OccurrenceData->m_ppPOccurrences;
-	const FString& InstanceName = Instance.MetaData.FindOrAdd(TEXT("Name"));
 	for (uint32 Index = 0; Index < ChildrenCount; ++Index)
 	{
-		const FString DefaultChildName = InstanceName + TEXT("_") + FString::FromInt(Index + 1);
-		int32 ChildrenId = TraverseOccurrence(Children[Index], DefaultChildName, OccurrenceUnit);
-		Reference.Children.Add(ChildrenId);
+		TraverseOccurrence(Children[Index], Reference);
 	}
 
 	// Add External data
@@ -987,20 +870,16 @@ FCadId FTechSoftFileParser::TraverseOccurrence(const A3DAsmProductOccurrence* Oc
 		if (ExternalData->m_pPart != nullptr)
 		{
 			A3DAsmPartDefinition* PartDefinition = ExternalData->m_pPart;
-			TraversePartDefinition(PartDefinition, Reference, OccurrenceUnit);
+			TraversePartDefinition(PartDefinition, Reference);
 		}
 
 		uint32 ExternalChildrenCount = ExternalData->m_uiPOccurrencesSize;
 		A3DAsmProductOccurrence** ExternalChildren = ExternalData->m_ppPOccurrences;
 		for (uint32 Index = 0; Index < ExternalChildrenCount; ++Index)
 		{
-			const FString DefaultChildName = InstanceName + TEXT("_") + FString::FromInt(Index + 1);
-			int32 ChildId = TraverseOccurrence(ExternalChildren[Index], DefaultChildName, OccurrenceUnit);
-			Reference.Children.Add(ChildId);
+			TraverseOccurrence(ExternalChildren[Index], Reference);
 		}
 	}
-
-	return Instance.Id;
 }
 
 void FTechSoftFileParser::CountUnderOccurrence(const A3DAsmProductOccurrence* Occurrence)
@@ -1008,7 +887,7 @@ void FTechSoftFileParser::CountUnderOccurrence(const A3DAsmProductOccurrence* Oc
 	TUniqueTSObj<A3DAsmProductOccurrenceData> OccurrenceData(Occurrence);
 	if (Occurrence && OccurrenceData.IsValid())
 	{
-		ComponentCount[EComponentType::Occurrence]++;
+		ComponentCount[EComponentType::Instance]++;
 		ComponentCount[EComponentType::Reference]++;
 
 		const A3DAsmProductOccurrence* CachedOccurrencePtr = Occurrence;
@@ -1095,10 +974,8 @@ void FTechSoftFileParser::CountUnderSubPrototype(const A3DAsmProductOccurrence*
 	}
 }
 
-void FTechSoftFileParser::ProcessPrototype(const A3DAsmProductOccurrence* InPrototypePtr, FEntityMetaData& OutPrototypeMetaData, A3DMiscTransformation** OutLocation)
+void FTechSoftFileParser::ProcessPrototype(const A3DAsmProductOccurrence* InPrototypePtr, FArchiveUnloadedReference& OutReference, A3DMiscTransformation** OutLocation)
 {
-	OutPrototypeMetaData.bUnloaded = true;
-
 	const A3DAsmProductOccurrence* PrototypePtr = InPrototypePtr;
 	TUniqueTSObj<A3DAsmProductOccurrenceData> PrototypeData;
 	while (PrototypePtr)
@@ -1111,8 +988,8 @@ void FTechSoftFileParser::ProcessPrototype(const A3DAsmProductOccurrence* InProt
 
 		if (PrototypeData->m_pPart != nullptr || PrototypeData->m_uiPOccurrencesSize != 0 || PrototypeData->m_pExternalData != nullptr || PrototypeData->m_pPrototype == nullptr)
 		{
-			ExtractMetaData(PrototypePtr, OutPrototypeMetaData);
-			ExtractSpecificMetaData(PrototypePtr, OutPrototypeMetaData);
+			ExtractMetaData(PrototypePtr, OutReference);
+			ExtractSpecificMetaData(PrototypePtr, OutReference);
 
 			TUniqueTSObj<A3DUTF8Char*> FilePathUTF8Ptr;
 			FilePathUTF8Ptr.FillWith(&TechSoftInterface::GetFilePathName, PrototypePtr);
@@ -1127,16 +1004,16 @@ void FTechSoftFileParser::ProcessPrototype(const A3DAsmProductOccurrence* InProt
 				FString FileName = FPaths::GetCleanFilename(FilePath);
 				if (!FileName.IsEmpty() && FileName != CADFileData.GetCADFileDescription().GetFileName())
 				{
-					OutPrototypeMetaData.ExternalFile = FFileDescriptor(*FilePath, nullptr, *CADFileData.GetCADFileDescription().GetRootFolder());
+					OutReference.ExternalFile = FFileDescriptor(*FilePath, nullptr, *CADFileData.GetCADFileDescription().GetRootFolder());
 				}
 			}
 		}
 		
-		// Case of DWG file, external reference is in fact internal data as a part or a Reference
+		// Case of DWG file, external reference is in fact internal data as a part or a component
 		// Case of CATProduct referencing CGR files: unloaded references are saved in ExternalData. In this case non empty ExternalFile means unloaded reference
-		if (PrototypeData->m_pPart != nullptr|| PrototypeData->m_uiPOccurrencesSize != 0 || (PrototypeData->m_pExternalData != nullptr && OutPrototypeMetaData.ExternalFile.IsEmpty()))
+		if (PrototypeData->m_pPart != nullptr|| PrototypeData->m_uiPOccurrencesSize != 0 || (PrototypeData->m_pExternalData != nullptr && OutReference.ExternalFile.IsEmpty()))
 		{
-			OutPrototypeMetaData.bUnloaded = false;
+			OutReference.bIsUnloaded = false;
 			PrototypePtr = nullptr;
 		}
 		else
@@ -1162,16 +1039,16 @@ void FTechSoftFileParser::ProcessPrototype(const A3DAsmProductOccurrence* InProt
 		}
 	}
 
-	if (OutPrototypeMetaData.bUnloaded)
+	if (OutReference.bIsUnloaded)
 	{
-		OutPrototypeMetaData.MetaData.Add(TEXT("Name"), OutPrototypeMetaData.ExternalFile.GetFileName());
+		OutReference.Label = OutReference.ExternalFile.GetFileName();
 	}
 	else
 	{
-		OutPrototypeMetaData.ExternalFile.Empty();
+		OutReference.ExternalFile.Empty();
 	}
 
-	BuildReferenceName(OutPrototypeMetaData);
+	BuildReferenceName(OutReference);
 }
 
 void FTechSoftFileParser::CountUnderPrototype(const A3DAsmProductOccurrence* Prototype)
@@ -1185,44 +1062,25 @@ void FTechSoftFileParser::CountUnderPrototype(const A3DAsmProductOccurrence* Pro
 	ComponentCount[EComponentType::Reference] ++;
 }
 
-void FTechSoftFileParser::TraversePartDefinition(const A3DAsmPartDefinition* PartDefinitionPtr, FArchiveReference& Part, double ParentUnit)
+void FTechSoftFileParser::TraversePartDefinition(const A3DAsmPartDefinition* PartDefinitionPtr, FArchiveReference& Part)
 {
-	FEntityMetaData PartMetaData;
-	ExtractMetaData(PartDefinitionPtr, PartMetaData);
+	ExtractMetaData(PartDefinitionPtr, Part);
 
-	if (PartMetaData.bRemoved || !PartMetaData.bShow)
+	if (Part.bIsRemoved || !Part.bShow)
 	{
 		return;
 	}
 
-	ExtractSpecificMetaData(PartDefinitionPtr, PartMetaData);
+	ExtractSpecificMetaData(PartDefinitionPtr, Part);
 
-	// A top/down propagation of the color and material is needed. This is done in 5.0.3 with minimal code modification. 
-	// However the model parsing need to be rewrite in the next version. (Jira UE-152691)
-	const FString* Material = Part.MetaData.Find(TEXT("MaterialName"));
-	const FString* Color = Part.MetaData.Find(TEXT("ColorName"));
-	if (Material && !Material->IsEmpty() && PartMetaData.MaterialName == 0)
-	{
-		uint32 MaterialName = (uint32) FCString::Atoi64(**Material);
-		PartMetaData.MetaData.Add(TEXT("Material"), *Material);
-		PartMetaData.MaterialName = MaterialName;
-	}
-	if (Color && !Color->IsEmpty() && PartMetaData.ColorName == 0)
-	{
-		uint32 ColorName = (uint32)FCString::Atoi64(**Color);
-		PartMetaData.MetaData.Add(TEXT("Color"), *Color);
-		PartMetaData.ColorName = ColorName;
-	}
-
-	BuildPartName(PartMetaData, Part);
+	BuildPartName(Part);
 
 	TUniqueTSObj<A3DAsmPartDefinitionData> PartData(PartDefinitionPtr);
 	if (PartData.IsValid())
 	{
-		for (unsigned int Index = 0; Index < PartData->m_uiRepItemsSize; ++Index)
+		for (A3DUns32 Index = 0; Index < PartData->m_uiRepItemsSize; ++Index)
 		{
-			int32 ChildId = TraverseRepresentationItem(PartData->m_ppRepItems[Index], PartMetaData, Part.Id, ParentUnit, Index);
-			Part.Children.Add(ChildId);
+			TraverseRepresentationItem(PartData->m_ppRepItems[Index], Part);
 		}
 	}
 }
@@ -1233,7 +1091,7 @@ void FTechSoftFileParser::CountUnderPartDefinition(const A3DAsmPartDefinition* P
 	if (PartDefinition && PartData.IsValid())
 	{
 		ComponentCount[EComponentType::Reference] ++;
-		ComponentCount[EComponentType::Occurrence] ++;
+		ComponentCount[EComponentType::Instance] ++;
 
 		for (unsigned int Index = 0; Index < PartData->m_uiRepItemsSize; ++Index)
 		{
@@ -1242,23 +1100,33 @@ void FTechSoftFileParser::CountUnderPartDefinition(const A3DAsmPartDefinition* P
 	}
 }
 
-FCadId FTechSoftFileParser::TraverseRepresentationItem(A3DRiRepresentationItem* RepresentationItem, const FEntityMetaData& PartMetaData, const FCadId ParentId, double ParentUnit, int32 ItemIndex)
+void FTechSoftFileParser::TraverseRepresentationItem(A3DRiRepresentationItem* RepresentationItem, FArchiveReference& Part)
 {
+	if (!RepresentationItem)
+	{
+		return;
+	}
+
+	if (FCadId* BodyIndexPtr = RepresentationItemsCache.Find(RepresentationItem))
+	{
+		Part.AddChild(*BodyIndexPtr);
+		return;
+	}
+
 	A3DEEntityType Type;
 	A3DEntityGetType(RepresentationItem, &Type);
 
 	switch (Type)
 	{
 	case kA3DTypeRiSet:
-		return TraverseRepresentationSet(RepresentationItem, PartMetaData, ParentUnit);
+		return TraverseRepresentationSet(RepresentationItem, Part);
 	case kA3DTypeRiBrepModel:
-		return TraverseBRepModel(RepresentationItem, PartMetaData, ParentId, ParentUnit, ItemIndex);
+		return TraverseBRepModel(RepresentationItem, Part);
 	case kA3DTypeRiPolyBrepModel:
-		return TraversePolyBRepModel(RepresentationItem, PartMetaData, ParentId, ParentUnit, ItemIndex);
+		return TraversePolyBRepModel(RepresentationItem, Part);
 	default:
 		break;
 	}
-	return 0;
 }
 
 void FTechSoftFileParser::CountUnderRepresentationItem(const A3DRiRepresentationItem* RepresentationItem)
@@ -1281,31 +1149,30 @@ void FTechSoftFileParser::CountUnderRepresentationItem(const A3DRiRepresentation
 	}
 }
 
-FCadId FTechSoftFileParser::TraverseRepresentationSet(const A3DRiSet* RepresentationSetPtr, const FEntityMetaData& PartMetaData, double ParentUnit)
+void FTechSoftFileParser::TraverseRepresentationSet(const A3DRiSet* RepresentationSetPtr, FArchiveReference& Parent)
 {
 	TUniqueTSObj<A3DRiSetData> RepresentationSetData(RepresentationSetPtr);
 	if (!RepresentationSetData.IsValid())
 	{
-		return 0;
+		return;
 	}
 
-	FEntityMetaData RepresentationSetMetaData;
-	ExtractMetaData(RepresentationSetPtr, RepresentationSetMetaData);
+	FCadId RepresentationSetId = 0;
+	FArchiveReference& RepresentationSet = SceneGraph.AddOccurence(Parent);
+	ExtractMetaData(RepresentationSetPtr, RepresentationSet);
+	RepresentationSet.DefineGraphicsPropertiesFromNoOverwrite(Parent);
+	BuildRepresentationSetName(RepresentationSet, Parent);
 
-	if (RepresentationSetMetaData.bRemoved || !RepresentationSetMetaData.bShow)
+	if (RepresentationSet.bIsRemoved || !RepresentationSet.bShow)
 	{
-		return 0;
+		SceneGraph.RemoveLastOccurence();
+		return;
 	}
 
-	FCadId RepresentationSetId = 0;
-	FArchiveReference& RepresentationSet = AddOccurence(RepresentationSetMetaData, RepresentationSetId);
-
 	for (A3DUns32 Index = 0; Index < RepresentationSetData->m_uiRepItemsSize; ++Index)
 	{
-		int32 ChildId = TraverseRepresentationItem(RepresentationSetData->m_ppRepItems[Index], RepresentationSetMetaData, RepresentationSet.Id, ParentUnit, Index);
-		RepresentationSet.Children.Add(ChildId);
+		TraverseRepresentationItem(RepresentationSetData->m_ppRepItems[Index], RepresentationSet);
 	}
-	return RepresentationSetId;
 }
 
 void FTechSoftFileParser::CountUnderRepresentationSet(const A3DRiSet* RepresentationSet)
@@ -1313,7 +1180,7 @@ void FTechSoftFileParser::CountUnderRepresentationSet(const A3DRiSet* Representa
 	TUniqueTSObj<A3DRiSetData> RepresentationSetData(RepresentationSet);
 	if (RepresentationSet && RepresentationSetData.IsValid())
 	{
-		ComponentCount[EComponentType::Occurrence] ++;
+		ComponentCount[EComponentType::Instance] ++;
 		ComponentCount[EComponentType::Reference] ++;
 
 		for (A3DUns32 Index = 0; Index < RepresentationSetData->m_uiRepItemsSize; ++Index)
@@ -1323,132 +1190,70 @@ void FTechSoftFileParser::CountUnderRepresentationSet(const A3DRiSet* Representa
 	}
 }
 
-FCadId FTechSoftFileParser::TraverseBRepModel(A3DRiBrepModel* BRepModelPtr, const FEntityMetaData& PartMetaData, const FCadId ParentId, double ParentUnit, int32 ItemIndex)
+void FTechSoftFileParser::TraverseBRepModel(A3DRiBrepModel* BRepModelPtr, FArchiveReference& Parent)
 {
-	if (!BRepModelPtr)
-	{
-		return 0;
-	}
+	FArchiveBody& BRep = SceneGraph.AddBody(Parent);
+	ExtractMetaData(BRepModelPtr, BRep);
+	BRep.DefineGraphicsPropertiesFromNoOverwrite(Parent);
 
-	FEntityMetaData BRepMetaData;
-	ExtractMetaData(BRepModelPtr, BRepMetaData);
+	ExtractSpecificMetaData(BRepModelPtr, BRep);
 
-	if (!BRepMetaData.bShow || BRepMetaData.bRemoved)
-	{
-		return 0;
-	}
 
-	TUniqueTSObj<A3DRiBrepModelData> BRepModelData(BRepModelPtr);
-	BuildBodyName(BRepMetaData, PartMetaData, ItemIndex, (bool) BRepModelData->m_bSolid);
-
-	if (FCadId* BodyIndexPtr = RepresentationItemsCache.Find(BRepModelPtr))
+	if (!BRep.bShow || BRep.bIsRemoved)
 	{
-		return CADFileData.GetBodyAt(*BodyIndexPtr).Id;
+		SceneGraph.RemoveLastBody();
+		return;
 	}
 
-	// if BRep model has not material or color, use the ones from the part
-	// A top/down propagation of the color and material is needed. This is done in 5.0.3 with minimal code modification. 
-	// However the model parsing need to be rewrite in the next version. (Jira UE-152691)
-	if (BRepMetaData.MaterialName == 0 && BRepMetaData.ColorName == 0)
-	{
-		if (PartMetaData.MaterialName)
-		{
-			BRepMetaData.MaterialName = PartMetaData.MaterialName;
-			BRepMetaData.MetaData.Add(TEXT("MaterialName"), FString::Printf(TEXT("%u"), PartMetaData.MaterialName));
-		}
-		if (PartMetaData.ColorName)
-		{
-			BRepMetaData.ColorName = PartMetaData.ColorName;
-			BRepMetaData.MetaData.Add(TEXT("ColorName"), FString::Printf(TEXT("%u"), PartMetaData.ColorName));
-		}
-	}
+	Parent.AddChild(BRep.Id);
+	RepresentationItemsCache.Add(BRepModelPtr, BRep.Id);
 
-	ExtractSpecificMetaData(BRepModelPtr, BRepMetaData);
+	TUniqueTSObj<A3DRiBrepModelData> BRepModelData(BRepModelPtr);
+	BRep.bIsASolid = (bool)BRepModelData->m_bSolid;
+	BuildBodyName(BRep, Parent);
 
-	FMatrix Matrix = FMatrix::Identity;
-	double BRepUnit = ParentUnit;
 	TUniqueTSObj<A3DRiRepresentationItemData> RepresentationData(BRepModelPtr);
-	if (RepresentationData->m_pCoordinateSystem)
-	{
-		Matrix = ExtractCoordinateSystem(RepresentationData->m_pCoordinateSystem, BRepUnit);
-	}
-
-	int32 BodyIndex = AddBody(BRepMetaData, Matrix, ParentId, BRepUnit);
-	FArchiveBody& Body = CADFileData.GetBodyAt(BodyIndex);
-
-	RepresentationItemsCache.Add(BRepModelPtr, BodyIndex);
-
-	return Body.Id;
+	ExtractCoordinateSystem(RepresentationData->m_pCoordinateSystem, BRep);
 }
 
-FCadId FTechSoftFileParser::TraversePolyBRepModel(A3DRiPolyBrepModel* PolygonalPtr, const FEntityMetaData& PartMetaData, const FCadId ParentId, double ParentUnit, int32 ItemIndex)
+void FTechSoftFileParser::TraversePolyBRepModel(A3DRiPolyBrepModel* PolygonalPtr, FArchiveReference& Parent)
 {
-	if (!PolygonalPtr)
-	{
-		return 0;
-	}
+	FArchiveBody& BRep = SceneGraph.AddBody(Parent);
+	ExtractMetaData(PolygonalPtr, BRep);
+	BRep.DefineGraphicsPropertiesFromNoOverwrite(Parent);
 
-	FEntityMetaData BRepMetaData;
-	ExtractMetaData(PolygonalPtr, BRepMetaData);
+	ExtractSpecificMetaData(PolygonalPtr, BRep);
 
-	if (!BRepMetaData.bShow || BRepMetaData.bRemoved)
+	if (!BRep.bShow || BRep.bIsRemoved)
 	{
-		return 0;
-	}
-
-	TUniqueTSObj<A3DRiPolyBrepModelData> BRepModelData(PolygonalPtr);
-	BuildBodyName(BRepMetaData, PartMetaData, ItemIndex, (bool) BRepModelData->m_bIsClosed);
-
-	if (FCadId* BodyIndexPtr = RepresentationItemsCache.Find(PolygonalPtr))
-	{
-		return CADFileData.GetBodyAt(*BodyIndexPtr).Id;
+		SceneGraph.RemoveLastBody();
+		return;
 	}
 
-	// if BRep model has not material or color, add part one
-	if (BRepMetaData.MaterialName == 0 && BRepMetaData.ColorName == 0)
-	{
-		if (PartMetaData.MaterialName)
-		{
-			BRepMetaData.MaterialName = PartMetaData.MaterialName;
-			BRepMetaData.MetaData.Add(TEXT("MaterialName"), FString::Printf(TEXT("%u"), PartMetaData.MaterialName));
-		}
-		if (PartMetaData.ColorName)
-		{
-			BRepMetaData.ColorName = PartMetaData.ColorName;
-			BRepMetaData.MetaData.Add(TEXT("ColorName"), FString::Printf(TEXT("%u"), PartMetaData.ColorName));
-		}
-	}
+	Parent.AddChild(BRep.Id);
+	RepresentationItemsCache.Add(PolygonalPtr, BRep.Id);
 
-	ExtractSpecificMetaData(PolygonalPtr, BRepMetaData);
+	TUniqueTSObj<A3DRiPolyBrepModelData> BRepModelData(PolygonalPtr);
+	BRep.bIsASolid = (bool)BRepModelData->m_bIsClosed;
+	BuildBodyName(BRep, Parent);
 
-	FMatrix Matrix = FMatrix::Identity;
-	double BRepUnit = ParentUnit;
 	TUniqueTSObj<A3DRiRepresentationItemData> RepresentationData(PolygonalPtr);
-	if (RepresentationData->m_pCoordinateSystem)
-	{
-		Matrix = ExtractCoordinateSystem(RepresentationData->m_pCoordinateSystem, BRepUnit);
-	}
-
-	int32 BodyIndex = AddBody(BRepMetaData, Matrix, ParentId, BRepUnit);
-	FArchiveBody& Body = CADFileData.GetBodyAt(BodyIndex);
-
-	RepresentationItemsCache.Add(PolygonalPtr, BodyIndex);
-
-	return Body.Id;
+	ExtractCoordinateSystem(RepresentationData->m_pCoordinateSystem, BRep);
 }
 
-void FTechSoftFileParser::ExtractMetaData(const A3DEntity* Entity, FEntityMetaData& OutMetaData)
+
+void FTechSoftFileParser::ExtractMetaData(const A3DEntity* Entity, FArchiveCADObject& OutObject)
 {
 	TUniqueTSObj<A3DRootBaseData> MetaData(Entity);
 	if (MetaData.IsValid())
 	{
 		if (MetaData->m_pcName && MetaData->m_pcName[0] != '\0')
 		{
-			FString SDKName = UTF8_TO_TCHAR(MetaData->m_pcName);
-			if(SDKName != TEXT("unnamed"))  // "unnamed" is create by Techsoft. This name is ignored 
+			FString Name = UTF8_TO_TCHAR(MetaData->m_pcName);
+			if(Name != TEXT("unnamed"))  // "unnamed" is create by Techsoft. This name is ignored 
 			{
-				SDKName = TechSoftUtils::CleanSdkName(SDKName);
-				OutMetaData.MetaData.Emplace(TEXT("SDKName"), SDKName);
+				Name = TechSoftUtils::CleanSdkName(Name);
+				OutObject.Label = Name;
 			}
 		}
 
@@ -1458,7 +1263,7 @@ void FTechSoftFileParser::ExtractMetaData(const A3DEntity* Entity, FEntityMetaDa
 			AttributeData.FillFrom(MetaData->m_ppAttributes[Index]);
 			if (AttributeData.IsValid())
 			{
-				TechSoftFileParserImpl::TraverseAttribute(*AttributeData, OutMetaData.MetaData);
+				TechSoftFileParserImpl::TraverseAttribute(*AttributeData, OutObject.MetaData);
 			}
 		}
 	}
@@ -1470,218 +1275,154 @@ void FTechSoftFileParser::ExtractMetaData(const A3DEntity* Entity, FEntityMetaDa
 		{
 			if (MetaDataWithGraphics->m_pGraphics != NULL)
 			{
-				ExtractGraphicProperties(MetaDataWithGraphics->m_pGraphics, OutMetaData);
+				ExtractGraphicProperties(MetaDataWithGraphics->m_pGraphics, OutObject);
 			}
 		}
 	}
 }
 
-void FTechSoftFileParser::BuildReferenceName(FEntityMetaData& ReferenceData)
+void FTechSoftFileParser::BuildReferenceName(FArchiveCADObject& ReferenceData)
 {
 	TMap<FString, FString>& MetaData = ReferenceData.MetaData;
-	if (MetaData.IsEmpty())
-	{
-		return;
-	}
 
 	FString* NamePtr = MetaData.Find(TEXT("InstanceName"));
 	if (NamePtr != nullptr && !NamePtr->IsEmpty())
 	{
-		FString& Name = MetaData.FindOrAdd(TEXT("Name"));
-		Name = *NamePtr;
 		if (Format == ECADFormat::CATIA)
 		{
-			Name = TechSoftUtils::CleanCatiaReferenceName(Name);
+			ReferenceData.Label = TechSoftUtils::CleanCatiaReferenceName(*NamePtr);
 		}
-		return;
-	}
-
-	if (Format == ECADFormat::JT)
-	{
-		if (TechSoftFileParserImpl::ReplaceOrAddNameValue(MetaData, TEXT("SDKName")))
+		else
 		{
-			return;
+			ReferenceData.Label = *NamePtr;
 		}
-	}
-
-	if (TechSoftFileParserImpl::CheckIfNameExists(MetaData))
-	{
 		return;
 	}
 
-	if (TechSoftFileParserImpl::ReplaceOrAddNameValue(MetaData, TEXT("PartNumber")))
-	{
-		return;
-	}
-
-	NamePtr = MetaData.Find(TEXT("SDKName"));
-	if (NamePtr != nullptr && !NamePtr->IsEmpty())
+	switch (Format)
 	{
-		FString SdkName = *NamePtr;
-
-		switch (Format)
-		{
-		case ECADFormat::CATIA_3DXML:
-			SdkName = TechSoftUtils::Clean3dxmlReferenceSdkName(SdkName);
-			break;
-
-		case ECADFormat::SOLIDWORKS:
-			SdkName = TechSoftUtils::CleanSwReferenceSdkName(SdkName);
-			break;
+	case ECADFormat::CATIA_3DXML:
+		ReferenceData.Label = TechSoftUtils::Clean3dxmlReferenceSdkName(ReferenceData.Label);
+		break;
 
-		default:
-			break;
-		}
+	case ECADFormat::SOLIDWORKS:
+		ReferenceData.Label = TechSoftUtils::CleanSwReferenceSdkName(ReferenceData.Label);
+		break;
 
-		MetaData.FindOrAdd(TEXT("Name"), SdkName);
-		return;
+	default:
+		break;
 	}
 }
 
-void FTechSoftFileParser::BuildInstanceName(FEntityMetaData& InstanceData, const FString& DefaultName)
+void FTechSoftFileParser::BuildInstanceName(FArchiveInstance& InstanceData, const FArchiveReference& Parent)
 {
 	TMap<FString, FString>& MetaData = InstanceData.MetaData;
 
-	if (TechSoftFileParserImpl::ReplaceOrAddNameValue(MetaData, TEXT("InstanceName")))
+	if (InstanceData.SetNameWithAttributeValue(TEXT("InstanceName")))
 	{
 		return;
 	}
 
-	if (TechSoftFileParserImpl::CheckIfNameExists(MetaData))
+	if (InstanceData.IsNameDefined())
 	{
-		return;
-	}
-
-	FString* NamePtr = MetaData.Find(TEXT("SDKName"));
-	if (NamePtr != nullptr && !NamePtr->IsEmpty())
-	{
-		FString SdkName = *NamePtr;
-
 		switch (Format)
 		{
 		case ECADFormat::CATIA:
-			SdkName = TechSoftUtils::CleanCatiaInstanceSdkName(SdkName);
+			InstanceData.Label = TechSoftUtils::CleanCatiaInstanceSdkName(InstanceData.Label);
 			break;
 		
 		case ECADFormat::CATIA_3DXML:
-			SdkName = TechSoftUtils::Clean3dxmlInstanceSdkName(SdkName);
+			InstanceData.Label = TechSoftUtils::Clean3dxmlInstanceSdkName(InstanceData.Label);
 			break;
 
 		case ECADFormat::SOLIDWORKS:
-			SdkName = TechSoftUtils::CleanSwInstanceSdkName(SdkName);
+			InstanceData.Label = TechSoftUtils::CleanSwInstanceSdkName(InstanceData.Label);
 			break;
 
 		default:
 			break;
 		}
-
-		MetaData.FindOrAdd(TEXT("Name"), SdkName);
 		return;
 	}
 
-	if (TechSoftFileParserImpl::ReplaceOrAddNameValue(MetaData, TEXT("PartNumber")))
+	if (InstanceData.SetNameWithAttributeValue(TEXT("PartNumber")))
 	{
 		return;
 	}
 
-	FString& Name = MetaData.FindOrAdd(TEXT("Name"));
-	if (Name.IsEmpty())
+	if (InstanceData.Label.IsEmpty())
 	{
-		Name = DefaultName;
+		InstanceData.Label = Parent.Label + "_" + FString::FromInt(Parent.Children.Num());
 	}
 }
 
-void FTechSoftFileParser::BuildPartName(FEntityMetaData& PartData, const FArchiveReference& Reference)
+void FTechSoftFileParser::BuildPartName(FArchiveCADObject& PartData)
 {
-	TMap<FString, FString>& MetaData = PartData.MetaData;
-	if (!MetaData.IsEmpty())
+	if (PartData.SetNameWithAttributeValue(TEXT("PartNumber")))
 	{
-		if (TechSoftFileParserImpl::CheckIfNameExists(MetaData))
-		{
-			return;
-		}
-
-		if (TechSoftFileParserImpl::ReplaceOrAddNameValue(MetaData, TEXT("PartNumber")))
-		{
-			return;
-		}
-
-		if (TechSoftFileParserImpl::ReplaceOrAddNameValue(MetaData, TEXT("SDKName")))
-		{
-			return;
-		}
-	}
-
-	// If a name hasn't been found for the Part, the name of its parent (component) is used
-	const FString* ReferenceName = Reference.MetaData.Find(TEXT("Name"));
-	if (ReferenceName && !ReferenceName->IsEmpty())
-	{
-		MetaData.Add(TEXT("Name"), *ReferenceName);
 		return;
 	}
 }
 
-void FTechSoftFileParser::BuildBodyName(FEntityMetaData& BodyData, const FEntityMetaData& PartMetaData, int32 ItemIndex, bool bIsSolid)
+void FTechSoftFileParser::BuildBodyName(FArchiveBody& Body, const FArchiveReference& Parent)
 {
-	TMap<FString, FString>& MetaData = BodyData.MetaData;
-
-	if (TechSoftFileParserImpl::CheckIfNameExists(MetaData))
+	if (Format == ECADFormat::CREO)
 	{
-		return;
+		Body.Label = TechSoftUtils::CleanCreoName(Body.Label);
 	}
 
-	FString* NamePtr = MetaData.Find(TEXT("SDKName"));
-	if (NamePtr != nullptr)
+	if(Body.IsNameDefined())
 	{
-		FString SdkName = *NamePtr;
-		if (Format == ECADFormat::CREO)
-		{
-			SdkName = TechSoftUtils::CleanNameByRemoving_prt(SdkName);
-		}
-
-		MetaData.FindOrAdd(TEXT("Name"), SdkName);
 		return;
 	}
 
 	if (Format == ECADFormat::CATIA)
 	{
-		NamePtr = MetaData.Find(TEXT("BodyID"));
-		if (NamePtr != nullptr && !NamePtr->IsEmpty())
+		if(Body.SetNameWithAttributeValue(TEXT("BodyID")))
 		{
-			MetaData.FindOrAdd(TEXT("Name"), *NamePtr);
 			return;
 		}
 	}
 
-	FString& Name = MetaData.FindOrAdd(TEXT("Name"));
-	if (Name.IsEmpty())
+	FString& Label = Body.Label;
+	if (Parent.IsNameDefined())
 	{
-		const FString* PartName = PartMetaData.MetaData.Find(TEXT("Name"));
-		if (PartName != nullptr && !PartName->IsEmpty())
-		{
-			Name = *PartName;
-		}
-
-		if (!Name.IsEmpty())
+		Label = Parent.Label + TEXT("_body");
+	}
+	else
+	{
+		if (Body.bIsASolid)
 		{
-			Name += TEXT("_body");
+			Label = TEXT("Solid");
 		}
 		else
 		{
-			if (bIsSolid)
-			{
-				Name = TEXT("Solid");
-			}
-			else
-			{
-				Name = TEXT("Shell");
-			}
+			Label = TEXT("Shell");
 		}
-		Name += FString::FromInt(ItemIndex + 1);
 	}
+	Label += FString::FromInt(Parent.Children.Num());
+}
+
+void FTechSoftFileParser::BuildRepresentationSetName(FArchiveCADObject& Occurrence, const FArchiveReference& Parent)
+{
+	if (Occurrence.IsNameDefined())
+	{
+		return;
+	}
+
+	FString& Label = Occurrence.Label;
+	if (Parent.IsNameDefined())
+	{
+		Label = Parent.Label;
+	}
+	else
+	{
+		Label = TEXT("Product");
+	}
+	Label += FString::FromInt(Parent.Children.Num());
 }
 
-void FTechSoftFileParser::ExtractSpecificMetaData(const A3DAsmProductOccurrence* Occurrence, FEntityMetaData& OutMetaData)
+void FTechSoftFileParser::ExtractSpecificMetaData(const A3DAsmProductOccurrence* Occurrence, FArchiveCADObject& OutMetaData)
 {
 	//----------- Export Specific information per CAD format -----------
 	switch (ModellerType)
@@ -1703,6 +1444,7 @@ void FTechSoftFileParser::ExtractSpecificMetaData(const A3DAsmProductOccurrence*
 	}
 	case ModellerUnigraphics:
 	{
+#ifdef WIP
 		TUniqueTSObj<A3DAsmProductOccurrenceDataUg> UnigraphicsSpecificData(Occurrence);
 		if (UnigraphicsSpecificData.IsValid())
 		{
@@ -1711,10 +1453,25 @@ void FTechSoftFileParser::ExtractSpecificMetaData(const A3DAsmProductOccurrence*
 				FString PartUID = UTF8_TO_TCHAR(UnigraphicsSpecificData->m_psPartUID);
 				OutMetaData.MetaData.Emplace(TEXT("UnigraphicsPartUID"), PartUID);
 			}
+			if (UnigraphicsSpecificData->m_psFileName)
+			{
+				FString FileName = UTF8_TO_TCHAR(UnigraphicsSpecificData->m_psFileName);
+				OutMetaData.MetaData.Emplace(TEXT("UnigraphicsFileName"), FileName);
+			}
 			if (UnigraphicsSpecificData->m_psInstanceFileName)
 			{
 				FString InstanceFileName = UTF8_TO_TCHAR(UnigraphicsSpecificData->m_psInstanceFileName);
-				OutMetaData.MetaData.Emplace(TEXT("UnigraphicsPartUID"), InstanceFileName);
+				OutMetaData.MetaData.Emplace(TEXT("UnigraphicsInstanceFileName"), InstanceFileName);
+			}
+			if (UnigraphicsSpecificData->m_psRefSet)
+			{
+				FString RefSet = UTF8_TO_TCHAR(UnigraphicsSpecificData->m_psRefSet);
+				OutMetaData.MetaData.Emplace(TEXT("UnigraphicsInstanceRefSet"), RefSet);
+			}
+			if (UnigraphicsSpecificData->m_psPartUID)
+			{
+				FString PartUID = UTF8_TO_TCHAR(UnigraphicsSpecificData->m_psPartUID);
+				OutMetaData.MetaData.Emplace(TEXT("UnigraphicsInstancePartUID"), PartUID);
 			}
 
 			if (UnigraphicsSpecificData->m_uiInstanceTag)
@@ -1722,7 +1479,34 @@ void FTechSoftFileParser::ExtractSpecificMetaData(const A3DAsmProductOccurrence*
 				FString InstanceTag = FString::FromInt(UnigraphicsSpecificData->m_uiInstanceTag);
 				OutMetaData.MetaData.Emplace(TEXT("UnigraphicsInstanceTag"), InstanceTag);
 			}
+
+			for (uint32 Index = 0; Index < UnigraphicsSpecificData->m_uiPromotedBodiesSize; ++Index)
+			{
+				A3DPromotedBodyUg PromotedBody = UnigraphicsSpecificData->m_asPromotedBodies[Index];
+			}
+
+			for (uint32 Index = 0; Index < UnigraphicsSpecificData->m_uiChildrenByRefsetsSize; ++Index)
+			{
+				A3DElementsByRefsetUg Refset = UnigraphicsSpecificData->m_asChildrenByRefsets[Index];
+			}
+
+			if(UnigraphicsSpecificData->m_uiSolidsByRefsetsSize)
+			{
+				OutMetaData.BodyIdentsOfConfigurations.SetNum(UnigraphicsSpecificData->m_uiSolidsByRefsetsSize);
+				for (uint32 Index = 0; Index < UnigraphicsSpecificData->m_uiSolidsByRefsetsSize; ++Index)
+				{
+					A3DElementsByRefsetUg Refset = UnigraphicsSpecificData->m_asSolidsByRefsets[Index];
+					TPair<FString, TArray<uint32>>& BodyIdentsOfConfiguration = OutMetaData.BodyIdentsOfConfigurations[Index];
+					BodyIdentsOfConfiguration.Key = UTF8_TO_TCHAR(Refset.m_psRefset);
+					BodyIdentsOfConfiguration.Value.Reserve(Refset.m_uiElementsSize);
+					for (uint32 Andex = 0; Andex < Refset.m_uiElementsSize; ++Andex)
+					{
+						BodyIdentsOfConfiguration.Value.Add(Refset.m_auiElements[Andex]);
+					}
+				}
+			}
 		}
+#endif
 		break;
 	}
 
@@ -1794,7 +1578,7 @@ FArchiveMaterial& FTechSoftFileParser::AddMaterialAt(uint32 MaterialIndexToSave,
 
 // Look at TechSoftUtils::BuildCADMaterial if any loigc changes in this method
 // or any of the methods it calls
-FArchiveMaterial& FTechSoftFileParser::FindOrAddMaterial(uint32 MaterialIndex, const A3DGraphStyleData& GraphStyleData)
+FArchiveMaterial& FTechSoftFileParser::FindOrAddMaterial(FMaterialUId MaterialIndex, const A3DGraphStyleData& GraphStyleData)
 {
 	if (FArchiveMaterial* MaterialArchive = CADFileData.FindMaterial(MaterialIndex))
 	{
@@ -1825,7 +1609,7 @@ FArchiveMaterial& FTechSoftFileParser::FindOrAddMaterial(uint32 MaterialIndex, c
 	}
 }
 
-void FTechSoftFileParser::ExtractGraphicProperties(const A3DGraphics* Graphics, FEntityMetaData& OutMetaData)
+void FTechSoftFileParser::ExtractGraphicProperties(const A3DGraphics* Graphics, FArchiveCADObject& OutMetaData)
 {
 	TUniqueTSObj<A3DGraphicsData> GraphicsData(Graphics);
 	if (!GraphicsData.IsValid())
@@ -1833,51 +1617,29 @@ void FTechSoftFileParser::ExtractGraphicProperties(const A3DGraphics* Graphics,
 		return;
 	}
 
-	bool bFatherHeritColor = GraphicsData->m_usBehaviour & kA3DGraphicsFatherHeritColor;
-	bool bSonHeritColor = GraphicsData->m_usBehaviour & kA3DGraphicsSonHeritColor;
-
-	//To do if needed
-	//bool bFatherHeritLayer = GraphicsData->m_usBehaviour & kA3DGraphicsFatherHeritLayer;
-	//bool bSonHeritLayer = GraphicsData->m_usBehaviour & kA3DGraphicsSonHeritLayer;
-	//bool bFatherHeritTransparency = GraphicsData->m_usBehaviour & kA3DGraphicsFatherHeritTransparency;
-	//bool bSonHeritTransparency = GraphicsData->m_usBehaviour & kA3DGraphicsSonHeritTransparency;
-
-	OutMetaData.bRemoved = GraphicsData->m_usBehaviour & kA3DGraphicsRemoved;
+	OutMetaData.bIsRemoved = GraphicsData->m_usBehaviour & kA3DGraphicsRemoved;
 	OutMetaData.bShow = GraphicsData->m_usBehaviour & kA3DGraphicsShow;
 
-	if (GraphicsData->m_uiStyleIndex == A3D_DEFAULT_STYLE_INDEX)
+	if (GraphicsData->m_usBehaviour & kA3DGraphicsFatherHeritColor)
 	{
-		return;
+		OutMetaData.Inheritance = ECADGraphicPropertyInheritance::FatherHerit;
 	}
-
-	FCadUuid& ColorName = OutMetaData.ColorName;
-	FCadUuid& MaterialName = OutMetaData.MaterialName;
-
-	ExtractGraphStyleProperties(GraphicsData->m_uiStyleIndex, ColorName, MaterialName);
-
-	if(bSonHeritColor)
-	{
-		OutMetaData.MetaData.Add(TEXT("GraphicsBehaviour"), TEXT("SonHerit"));
-	}
-	else if (bFatherHeritColor)
+	else if (GraphicsData->m_usBehaviour & kA3DGraphicsSonHeritColor)
 	{
-		OutMetaData.MetaData.Add(TEXT("GraphicsBehaviour"), TEXT("FatherHerit"));
+		OutMetaData.Inheritance = ECADGraphicPropertyInheritance::ChildHerit;
 	}
 
-	if (ColorName)
+	if (GraphicsData->m_uiStyleIndex == A3D_DEFAULT_STYLE_INDEX)
 	{
-		OutMetaData.MetaData.Add(TEXT("ColorName"), FString::Printf(TEXT("%u"), ColorName));
+		return;
 	}
 
-	if (MaterialName)
-	{
-		OutMetaData.MetaData.Add(TEXT("MaterialName"), FString::Printf(TEXT("%u"), MaterialName));
-	}
+	ExtractGraphStyleProperties(GraphicsData->m_uiStyleIndex, OutMetaData);
 }
 
 // Please review TechSoftUtils::GetMaterialValues if anything changes
 // in this method or the methods it calls
-void FTechSoftFileParser::ExtractGraphStyleProperties(uint32 StyleIndex, FCadUuid& OutColorName, FCadUuid& OutMaterialName)
+void FTechSoftFileParser::ExtractGraphStyleProperties(uint32 StyleIndex, FArchiveGraphicProperties& OutGraphicProperties)
 {
 	TUniqueTSObjFromIndex<A3DGraphStyleData> GraphStyleData(StyleIndex);
 
@@ -1886,7 +1648,7 @@ void FTechSoftFileParser::ExtractGraphStyleProperties(uint32 StyleIndex, FCadUui
 		if (GraphStyleData->m_bMaterial)
 		{
 			FArchiveMaterial& MaterialArchive = FindOrAddMaterial(GraphStyleData->m_uiRgbColorIndex, *GraphStyleData);
-			OutMaterialName = MaterialArchive.UEMaterialUId;
+			OutGraphicProperties.MaterialUId = MaterialArchive.UEMaterialUId;
 		}
 		else
 		{
@@ -1897,7 +1659,7 @@ void FTechSoftFileParser::ExtractGraphStyleProperties(uint32 StyleIndex, FCadUui
 			}
 
 			FArchiveColor& ColorArchive = FindOrAddColor(GraphStyleData->m_uiRgbColorIndex, Alpha);
-			OutColorName = ColorArchive.UEMaterialUId;
+			OutGraphicProperties.ColorUId = ColorArchive.UEMaterialUId;
 		}
 	}
 }
@@ -1987,11 +1749,11 @@ FMatrix FTechSoftFileParser::ExtractGeneralTransformation(const A3DMiscTransform
 	return FMatrix::Identity;
 }
 
-FMatrix FTechSoftFileParser::ExtractTransformation(const A3DMiscTransformation* Transformation3D, double& InOutUnit)
+void FTechSoftFileParser::ExtractTransformation(const A3DMiscTransformation* Transformation3D, FArchiveCADObject& Component)
 {
 	if (Transformation3D == NULL)
 	{
-		return FMatrix::Identity;
+		return;
 	}
 
 	A3DEEntityType Type = kA3DTypeUnknown;
@@ -1999,23 +1761,25 @@ FMatrix FTechSoftFileParser::ExtractTransformation(const A3DMiscTransformation*
 
 	if (Type == kA3DTypeMiscCartesianTransformation)
 	{
-		return ExtractTransformation3D(Transformation3D, InOutUnit);
+		Component.TransformMatrix = ExtractTransformation3D(Transformation3D, Component.Unit);
 	}
 	else if (Type == kA3DTypeMiscGeneralTransformation)
 	{
-		return ExtractGeneralTransformation(Transformation3D, InOutUnit);
+		Component.TransformMatrix = ExtractGeneralTransformation(Transformation3D, Component.Unit);
 	}
-	return FMatrix::Identity;
 }
 
-FMatrix FTechSoftFileParser::ExtractCoordinateSystem(const A3DRiCoordinateSystem* CoordinateSystem, double& InOutUnit)
+void FTechSoftFileParser::ExtractCoordinateSystem(const A3DRiCoordinateSystem* CoordinateSystem, FArchiveCADObject& OutMetaData)
 {
 	TUniqueTSObj<A3DRiCoordinateSystemData> CoordinateSystemData(CoordinateSystem);
 	if (CoordinateSystemData.IsValid())
 	{
-		return ExtractTransformation3D(CoordinateSystemData->m_pTransformation, InOutUnit);
+		OutMetaData.TransformMatrix = ExtractTransformation3D(CoordinateSystemData->m_pTransformation, OutMetaData.Unit);
+	}
+	else
+	{
+		OutMetaData.TransformMatrix = FMatrix::Identity;
 	}
-	return FMatrix::Identity;
 }
 
 bool FTechSoftFileParser::IsConfigurationSet(const A3DAsmProductOccurrence* Occurrence)
@@ -2141,7 +1905,7 @@ void FTechSoftFileParser::ReadMaterialsAndColors()
 				A3DEntity* PicturePtr = TechSoftInterface::GetPointerFromIndex(PictureIndex, kA3DTypeGraphPicture);
 				if (PicturePtr)
 				{
-					FEntityMetaData PictureMetaData;
+					FArchiveCADObject PictureMetaData;
 					ExtractMetaData(PicturePtr, PictureMetaData);
 				}
 
diff --git a/Engine/Plugins/Enterprise/DatasmithCADImporter/Source/CADInterfaces/Private/TechSoftFileParser.h b/Engine/Plugins/Enterprise/DatasmithCADImporter/Source/CADInterfaces/Private/TechSoftFileParser.h
--- a/Engine/Plugins/Enterprise/DatasmithCADImporter/Source/CADInterfaces/Private/TechSoftFileParser.h
+++ b/Engine/Plugins/Enterprise/DatasmithCADImporter/Source/CADInterfaces/Private/TechSoftFileParser.h
@@ -14,25 +14,16 @@ namespace CADLibrary
 {
 
 class FArchiveBody;
+class FArchiveCADObject;
 class FArchiveColor;
+class FArchiveGraphicProperties;
 class FArchiveInstance;
 class FArchiveMaterial;
 class FArchiveReference;
+class FArchiveSceneGraph;
 class FArchiveUnloadedReference;
 class FCADFileData;
 
-struct FEntityMetaData;
-
-enum EComponentType : uint32
-{
-	Reference = 0,
-	Occurrence,
-	UnloadedReference,
-	Body,
-	Undefined,
-	LastType
-};
-
 class FTechSoftFileParser : public ICADFileParser
 {
 public:
@@ -56,7 +47,7 @@ public:
 		return FileUnit;
 	}
 
-	void ExtractMetaData(const A3DEntity* Entity, FEntityMetaData& OutMetaData);
+	void ExtractMetaData(const A3DEntity* Entity, FArchiveCADObject& OutMetaData);
 
 protected:
 
@@ -103,35 +94,39 @@ private:
 	// To be able to build a name, the name of the parent has to be known.
 	// The implementation of the naming policy is done in 5.0.3 with minimal code modification. However the model parsing need to be rewrite in the next version. (Jira UE-152624)
 
-	void TraverseReference(const A3DAsmProductOccurrence* Reference, const FString& RootName, const FMatrix& ParentMatrix, double ParentUnit);
+	void TraverseReference(const A3DAsmProductOccurrence* A3DReference, FArchiveReference& Reference);
 	bool IsConfigurationSet(const A3DAsmProductOccurrence* Occurrence);
-	void TraverseConfigurationSet(const A3DAsmProductOccurrence* ConfigurationSet, double ParentUnit);
-	FCadId TraverseOccurrence(const A3DAsmProductOccurrence* Occurrence, const FString& DefaultOccurrenceName, double ParentUnit);
-	void ProcessPrototype(const A3DAsmProductOccurrence* InPrototype, FEntityMetaData& OutMetaData, A3DMiscTransformation** OutLocation);
-	void TraversePartDefinition(const A3DAsmPartDefinition* PartDefinition, FArchiveReference& Reference, double ParentUnit);
-	FCadId TraverseRepresentationSet(const A3DRiSet* pSet, const FEntityMetaData& PartMetaData, double ParentUnit);
-	FCadId TraverseRepresentationItem(A3DRiRepresentationItem* RepresentationItem, const FEntityMetaData& PartMetaData, const FCadId ParentId, double ParentUnit, int32 ItemIndex);
-	FCadId TraverseBRepModel(A3DRiBrepModel* BrepModel, const FEntityMetaData& PartMetaData, const FCadId ParentId, double ParentUnit, int32 ItemIndex);
-	FCadId TraversePolyBRepModel(A3DRiPolyBrepModel* PolygonalBrepModel, const FEntityMetaData& PartMetaData, const FCadId ParentId, double ParentUnit, int32 ItemIndex);
+	void TraverseConfigurationSet(const A3DAsmProductOccurrence* ConfigurationSet, FArchiveReference& Reference);
+	void TraverseOccurrence(const A3DAsmProductOccurrence* Occurrence, FArchiveReference& ParentReference);
+	void ProcessPrototype(const A3DAsmProductOccurrence* InPrototype, FArchiveUnloadedReference& OutReference, A3DMiscTransformation** OutLocation);
+	void ProcessReference(const A3DAsmProductOccurrence* OccurrencePtr, FArchiveInstance& Instance, FArchiveReference& Reference);
+	void ProcessUnloadedReference(const FArchiveInstance& Instance, FArchiveUnloadedReference& Reference);
+
+	void TraversePartDefinition(const A3DAsmPartDefinition* InPartDefinition, FArchiveReference& Parent);
+	void TraverseExternalData(const A3DAsmProductOccurrence* InExternalData, FArchiveReference& Parent);
+	void TraverseRepresentationSet(const A3DRiSet* InSetPtr, FArchiveReference& Parent);
+	void TraverseRepresentationItem(A3DRiRepresentationItem* RepresentationItem, FArchiveReference& Parent);
+	void TraverseBRepModel(A3DRiBrepModel* BrepModel, FArchiveReference& Parent);
+	void TraversePolyBRepModel(A3DRiPolyBrepModel* PolygonalBrepModel, FArchiveReference& Parent);
 
 	// MetaData
-	void ExtractSpecificMetaData(const A3DAsmProductOccurrence* Occurrence, FEntityMetaData& OutMetaData);
-
-	void BuildInstanceName(FEntityMetaData& MetaData, const FString& DefaultInstanceName);
-	void BuildReferenceName(FEntityMetaData& MetaData);
-	void BuildPartName(FEntityMetaData& MetaData, const FArchiveReference& Reference);
-	void BuildBodyName(FEntityMetaData& MetaData, const FEntityMetaData& PartMetaData, int32 ItemIndex, bool bIsSolid);
+	void ExtractSpecificMetaData(const A3DAsmProductOccurrence* Occurrence, FArchiveCADObject& OutMetaData);
 
+	void BuildInstanceName(FArchiveInstance& MetaData, const FArchiveReference& Parent);
+	void BuildReferenceName(FArchiveCADObject& Reference);
+	void BuildPartName(FArchiveCADObject& Part);
+	void BuildBodyName(FArchiveBody& MetaData, const FArchiveReference& Parent);
+	void BuildRepresentationSetName(FArchiveCADObject& MetaData, const FArchiveReference& Parent);
 	// Graphic properties
-	void ExtractGraphicProperties(const A3DGraphics* Graphics, FEntityMetaData& OutMetaData);
+	void ExtractGraphicProperties(const A3DGraphics* Graphics, FArchiveCADObject& OutMetaData);
 
 	/**
 	 * ColorName and MaterialName have to be initialized before.
 	 * This method update the value ColorName or MaterialName accordingly of the GraphStyleData type (material or color)
 	 */
-	void ExtractGraphStyleProperties(uint32 StyleIndex, FCadUuid& ColorName, FCadUuid& MaterialName);
+	void ExtractGraphStyleProperties(uint32 StyleIndex, FArchiveGraphicProperties& OutGraphicProperties);
 	FArchiveColor& FindOrAddColor(uint32 ColorIndex, uint8 Alpha);
-	FArchiveMaterial& FindOrAddMaterial(uint32 MaterialId, const A3DGraphStyleData& GraphStyleData);
+	FArchiveMaterial& FindOrAddMaterial(FMaterialUId MaterialId, const A3DGraphStyleData& GraphStyleData);
 
 	/**
 	 * @param GraphMaterialIndex is the Techsoft index of the graphic data
@@ -144,18 +139,11 @@ private:
 	}
 
 	// Transform
-	FMatrix ExtractCoordinateSystem(const A3DRiCoordinateSystem* CoordinateSystem, double& InOutUnit);
-	FMatrix ExtractTransformation(const A3DMiscTransformation* Transformation3d, double& InOutUnit);
+	void ExtractCoordinateSystem(const A3DRiCoordinateSystem* CoordinateSystem, FArchiveCADObject& Component);
+	void ExtractTransformation(const A3DMiscTransformation* Transformation3d, FArchiveCADObject& Component);
 	FMatrix ExtractGeneralTransformation(const A3DMiscTransformation* GeneralTransformation, double& InOutUnit);
 	FMatrix ExtractTransformation3D(const A3DMiscTransformation* CartesianTransformation, double& InOutUnit);
 
-	// Archive methods
-	FArchiveInstance& AddInstance(FEntityMetaData& InstanceMetaData);
-	FArchiveReference& AddReference(FEntityMetaData& ReferenceMetaData, FArchiveInstance& Instance);
-	FArchiveUnloadedReference& AddUnloadedReference(FEntityMetaData& ReferenceMetaData, FArchiveInstance& Instance);
-	FArchiveReference& AddOccurence(FEntityMetaData& InstanceMetaData, FEntityMetaData& ReferenceMetaData, FCadId& OutReferenceId);
-	FArchiveReference& AddOccurence(FEntityMetaData& InstanceMetaData, FCadId& OutReferenceId);
-	int32 AddBody(FEntityMetaData& BodyMetaData, const FMatrix& Matrix, const FCadId ParentId, double BodyUnit);
 #endif
 
 protected:
@@ -166,6 +154,8 @@ protected:
 	uint32 ComponentCount[EComponentType::LastType] = { 0 };
 
 	FCADFileData& CADFileData;
+	FArchiveSceneGraph& SceneGraph;
+
 	FTechSoftInterface& TechSoftInterface;
 	ECADFormat Format;
 	bool bForceSew = false;
@@ -173,8 +163,6 @@ protected:
 	EModellerType ModellerType;
 	double FileUnit = 1;
 
-	FCadId LastEntityId = 1;
-
 	TMap<A3DRiRepresentationItem*, FCadId> RepresentationItemsCache;
 };
 
diff --git a/Engine/Plugins/Enterprise/DatasmithCADImporter/Source/CADInterfaces/Private/TechSoftFileParserCADKernelTessellator.cpp b/Engine/Plugins/Enterprise/DatasmithCADImporter/Source/CADInterfaces/Private/TechSoftFileParserCADKernelTessellator.cpp
--- a/Engine/Plugins/Enterprise/DatasmithCADImporter/Source/CADInterfaces/Private/TechSoftFileParserCADKernelTessellator.cpp
+++ b/Engine/Plugins/Enterprise/DatasmithCADImporter/Source/CADInterfaces/Private/TechSoftFileParserCADKernelTessellator.cpp
@@ -52,7 +52,7 @@ void FTechSoftFileParserCADKernelTessellator::SewAndGenerateBodyMeshes()
 	for (TPair<A3DRiRepresentationItem*, FCadId>& Entry : RepresentationItemsCache)
 	{
 		A3DRiRepresentationItem* RepresentationItemPtr = Entry.Key;
-		FArchiveBody& Body = CADFileData.GetBodyAt(Entry.Value);
+		FArchiveBody& Body = SceneGraph.GetBody(Entry.Value);
 
 		TArray<A3DRiRepresentationItem*>& Representations = OccurenceIdToRepresentations.FindOrAdd(Body.ParentId);
 		if (Representations.Max() == 0)
@@ -88,9 +88,9 @@ void FTechSoftFileParserCADKernelTessellator::SewAndMesh(TArray<A3DRiRepresentat
 		{
 			continue;
 		}
-		FArchiveBody& Body = CADFileData.GetBodyAt(*BodyIndex);
+		FArchiveBody& Body = SceneGraph.GetBody(*BodyIndex);
 
-		CADKernel::FBody* CADKernelBody = TechSoftBridge.AddBody(Representation, Body.MetaData, Body.BodyUnit);
+		CADKernel::FBody* CADKernelBody = TechSoftBridge.AddBody(Representation, Body.MetaData, Body.Unit);
 	}
 
 	// Sew if needed
@@ -133,20 +133,20 @@ void FTechSoftFileParserCADKernelTessellator::SewAndMesh(TArray<A3DRiRepresentat
 
 	// find Representation of deleted bodies to find unused FArchiveBody
 	//int32 DeletedBodyCount = Representations.Num() - BodyCount;
-	TArray<int32> ArchiveBodyIndexOfDeletedRepresentation;
-	ArchiveBodyIndexOfDeletedRepresentation.Reserve(DeletedBodyCount);
+	TArray<FCadId> ArchiveBodyIdOfDeletedRepresentation;
+	ArchiveBodyIdOfDeletedRepresentation.Reserve(DeletedBodyCount);
 
 	for (A3DRiRepresentationItem* Representation : Representations)
 	{
 		CADKernel::FBody* Body = TechSoftBridge.GetBody(Representation);
 		if (Body == nullptr)
 		{
-			FCadId* ArchiveBodyIndex = RepresentationItemsCache.Find(Representation);
-			if (ArchiveBodyIndex == nullptr)
+			FCadId* ArchiveBodyId = RepresentationItemsCache.Find(Representation);
+			if (ArchiveBodyId == nullptr)
 			{
 				continue; // should not append
 			}
-			ArchiveBodyIndexOfDeletedRepresentation.Add(*ArchiveBodyIndex);
+			ArchiveBodyIdOfDeletedRepresentation.Add(*ArchiveBodyId);
 		}
 	}
 
@@ -155,13 +155,13 @@ void FTechSoftFileParserCADKernelTessellator::SewAndMesh(TArray<A3DRiRepresentat
 	{
 		const A3DRiRepresentationItem* Representation = TechSoftBridge.GetA3DBody(Body);
 
-		FCadId* ArchiveBodyIndex = RepresentationItemsCache.Find(Representation);
-		if (ArchiveBodyIndex == nullptr)
+		FCadId* ArchiveBodyId = RepresentationItemsCache.Find(Representation);
+		if (ArchiveBodyId == nullptr)
 		{
 			continue; // should not append
 		}
 
-		FArchiveBody& ArchiveBody = CADFileData.GetBodyAt(*ArchiveBodyIndex);
+		FArchiveBody& ArchiveBody = SceneGraph.GetBody(*ArchiveBodyId);
 		MeshAndGetTessellation(CADKernelSession, ArchiveBody, *Body);
 	}
 
@@ -169,19 +169,17 @@ void FTechSoftFileParserCADKernelTessellator::SewAndMesh(TArray<A3DRiRepresentat
 	int32 Index = 0;
 	for (CADKernel::FBody* Body : NewBodies)
 	{
-		int32 ArchiveBodyIndex = ArchiveBodyIndexOfDeletedRepresentation[Index++];
-		FArchiveBody& ArchiveBody = CADFileData.GetBodyAt(ArchiveBodyIndex);
+		FCadId ArchiveBodyId = ArchiveBodyIdOfDeletedRepresentation[Index++];
+		FArchiveBody& ArchiveBody = SceneGraph.GetBody(ArchiveBodyId);
 		MeshAndGetTessellation(CADKernelSession, ArchiveBody, *Body);
 	}
 
 	// Delete unused FArchiveBody
-	for (; Index < ArchiveBodyIndexOfDeletedRepresentation.Num(); ++Index)
+	for (; Index < ArchiveBodyIdOfDeletedRepresentation.Num(); ++Index)
 	{
-		int32 ArchiveBodyIndex = ArchiveBodyIndexOfDeletedRepresentation[Index++];
-		FArchiveBody& ArchiveBody = CADFileData.GetBodyAt(ArchiveBodyIndex);
-		ArchiveBody.MetaData.Empty();
-		ArchiveBody.ParentId = 0;
-		ArchiveBody.MeshActorUId = 0;
+		FCadId ArchiveBodyId = ArchiveBodyIdOfDeletedRepresentation[Index++];
+		FArchiveBody& ArchiveBody = SceneGraph.GetBody(ArchiveBodyId);
+		ArchiveBody.Delete();
 	}
 }
 
@@ -196,7 +194,7 @@ void FTechSoftFileParserCADKernelTessellator::GenerateBodyMesh(A3DRiRepresentati
 	CADKernel::FCADFileReport Report;
 	FTechSoftBridge TechSoftBridge(*this, CADKernelSession, Report);
 
-	CADKernel::FBody* CADKernelBody = TechSoftBridge.AddBody(Representation, ArchiveBody.MetaData, ArchiveBody.BodyUnit);
+	CADKernel::FBody* CADKernelBody = TechSoftBridge.AddBody(Representation, ArchiveBody.MetaData, ArchiveBody.Unit);
 
 	if (CADFileData.GetImportParameters().GetStitchingTechnique() == StitchingHeal)
 	{
diff --git a/Engine/Plugins/Enterprise/DatasmithCADImporter/Source/CADInterfaces/Private/TechSoftInterface.cpp b/Engine/Plugins/Enterprise/DatasmithCADImporter/Source/CADInterfaces/Private/TechSoftInterface.cpp
--- a/Engine/Plugins/Enterprise/DatasmithCADImporter/Source/CADInterfaces/Private/TechSoftInterface.cpp
+++ b/Engine/Plugins/Enterprise/DatasmithCADImporter/Source/CADInterfaces/Private/TechSoftInterface.cpp
@@ -464,7 +464,8 @@ A3DStatus SewModel(A3DAsmModelFile* ModelPtr, double ToleranceInCM, A3DSewOption
 
 A3DStatus SewBReps(A3DRiBrepModel** BRepsToSew, uint32 const BRepCount, double ToleranceInCM, double FileUnit, A3DSewOptionsData const* SewOptions, A3DRiBrepModel*** OutNewBReps, uint32& OutNewBRepCount)
 {
-	const double ToleranceInFileUnit = ToleranceInCM * 10. / FileUnit; 
+	const double CmToMm = 10.;
+	const double ToleranceInFileUnit = ToleranceInCM * CmToMm / FileUnit;
 
 	A3DUns32 NewBRepCount = 0;
 	A3DStatus Status = A3DSewBrep(&BRepsToSew, BRepCount, ToleranceInFileUnit, SewOptions, OutNewBReps, &NewBRepCount);
diff --git a/Engine/Plugins/Enterprise/DatasmithCADImporter/Source/CADInterfaces/Private/TechSoftInterfaceUtils.inl b/Engine/Plugins/Enterprise/DatasmithCADImporter/Source/CADInterfaces/Private/TechSoftInterfaceUtils.inl
--- a/Engine/Plugins/Enterprise/DatasmithCADImporter/Source/CADInterfaces/Private/TechSoftInterfaceUtils.inl
+++ b/Engine/Plugins/Enterprise/DatasmithCADImporter/Source/CADInterfaces/Private/TechSoftInterfaceUtils.inl
@@ -170,7 +170,7 @@ private:
 			Tessellation.PatchId = Index;
 
 			Tessellation.MaterialUId = 0;
-			if (FaceTessData.m_uiStyleIndexesSize != FTechSoftDefaultValue::Style)
+			if (FaceTessData.m_uiStyleIndexesSize != 0)
 			{
 				// Store the StyleIndex on the MaterialName. It will be processed after tessellation
 				Tessellation.MaterialUId = FaceTessData.m_puiStyleIndexes[0];
diff --git a/Engine/Plugins/Enterprise/DatasmithCADImporter/Source/CADInterfaces/Private/TechSoftUtils.cpp b/Engine/Plugins/Enterprise/DatasmithCADImporter/Source/CADInterfaces/Private/TechSoftUtils.cpp
--- a/Engine/Plugins/Enterprise/DatasmithCADImporter/Source/CADInterfaces/Private/TechSoftUtils.cpp
+++ b/Engine/Plugins/Enterprise/DatasmithCADImporter/Source/CADInterfaces/Private/TechSoftUtils.cpp
@@ -756,7 +756,7 @@ FString CleanCatiaReferenceName(const FString& Name)
 	}
 	return Name;
 }
-FString CleanNameByRemoving_prt(const FString& Name)
+FString CleanCreoName(const FString& Name)
 {
 	int32 Position;
 	if (Name.FindLastChar(TEXT('.'), Position))
diff --git a/Engine/Plugins/Enterprise/DatasmithCADImporter/Source/CADInterfaces/Private/TechSoftUtilsPrivate.h b/Engine/Plugins/Enterprise/DatasmithCADImporter/Source/CADInterfaces/Private/TechSoftUtilsPrivate.h
--- a/Engine/Plugins/Enterprise/DatasmithCADImporter/Source/CADInterfaces/Private/TechSoftUtilsPrivate.h
+++ b/Engine/Plugins/Enterprise/DatasmithCADImporter/Source/CADInterfaces/Private/TechSoftUtilsPrivate.h
@@ -20,7 +20,7 @@ FString Clean3dxmlInstanceSdkName(const FString& Name);
 FString Clean3dxmlReferenceSdkName(const FString& Name);
 FString CleanSwInstanceSdkName(const FString& Name);
 FString CleanSwReferenceSdkName(const FString& Name);
-FString CleanNameByRemoving_prt(const FString& Name);
+FString CleanCreoName(const FString& Name);
 bool CheckIfNameExists(TMap<FString, FString>& MetaData);
 bool ReplaceOrAddNameValue(TMap<FString, FString>& MetaData, const TCHAR* Key);
 
diff --git a/Engine/Plugins/Enterprise/DatasmithCADImporter/Source/CADInterfaces/Public/CADFileData.h b/Engine/Plugins/Enterprise/DatasmithCADImporter/Source/CADInterfaces/Public/CADFileData.h
--- a/Engine/Plugins/Enterprise/DatasmithCADImporter/Source/CADInterfaces/Public/CADFileData.h
+++ b/Engine/Plugins/Enterprise/DatasmithCADImporter/Source/CADInterfaces/Public/CADFileData.h
@@ -114,102 +114,6 @@ public:
 		SceneGraphArchive.SerializeMockUp(*SceneGraphFilePath);
 	}
 
-	bool HasReferenceOfId(FCadId ReferenceId) const
-	{
-		return SceneGraphArchive.CADIdToReferenceIndex.Find(ReferenceId) != nullptr;
-	}
-
-	int32 ReferenceCount()
-	{
-		return SceneGraphArchive.References.Num();
-	}
-
-	int32 AddReference(FCadId ReferenceId)
-	{
-		ensure(SceneGraphArchive.References.Num() < SceneGraphArchive.References.Max());
-
-		int32 Index = SceneGraphArchive.References.Emplace(ReferenceId);
-		SceneGraphArchive.CADIdToReferenceIndex.Add(ReferenceId, Index);
-		return Index;
-	}
-
-	FArchiveReference& GetReferenceAt(int32 Index)
-	{
-		return SceneGraphArchive.References[Index];
-	}
-
-
-	bool HasInstanceOfId(FCadId InstanceId) const
-	{
-		return SceneGraphArchive.CADIdToInstanceIndex.Find(InstanceId) != nullptr;
-	}
-
-	int32 AddInstance(FCadId InstanceId)
-	{
-		ensure(SceneGraphArchive.Instances.Num() < SceneGraphArchive.Instances.Max());
-
-		int32 Index = SceneGraphArchive.Instances.Emplace(InstanceId);
-		SceneGraphArchive.CADIdToInstanceIndex.Add(InstanceId, Index);
-		return Index;
-	}
-
-	FArchiveInstance& GetInstanceAt(int32 Index)
-	{
-		return SceneGraphArchive.Instances[Index];
-	}
-
-	int32 GetInstanceIndexFromId(FCadId InstanceId)
-	{
-		int32* IndexPtr = SceneGraphArchive.CADIdToInstanceIndex.Find(InstanceId);
-
-		return IndexPtr ? *IndexPtr : INDEX_NONE;
-	}
-
-
-	bool HasBodyOfId(FCadId BodyId) const
-	{
-		return SceneGraphArchive.CADIdToBodyIndex.Find(BodyId) != nullptr;
-	}
-
-	int32 AddBody(FCadId BodyId)
-	{
-		ensure(SceneGraphArchive.Bodies.Num() < SceneGraphArchive.Bodies.Max());
-
-		int32 Index = SceneGraphArchive.Bodies.Emplace(BodyId);
-		SceneGraphArchive.CADIdToBodyIndex.Add(BodyId, Index);
-		return Index;
-	}
-
-	FArchiveBody& GetBodyAt(int32 Index)
-	{
-		return SceneGraphArchive.Bodies[Index];
-	}
-
-
-	int32* FindUnloadedReferenceOfId(FCadId ReferenceId)
-	{
-		return SceneGraphArchive.CADIdToUnloadedReferenceIndex.Find(ReferenceId);
-	}
-
-	bool HasUnloadedReferenceOfId(FCadId ReferenceId)
-	{
-		return FindUnloadedReferenceOfId(ReferenceId) != nullptr;
-	}
-
-	int32 AddUnloadedReference(FCadId ReferenceId)
-	{
-		ensure(SceneGraphArchive.UnloadedReferences.Num() < SceneGraphArchive.UnloadedReferences.Max());
-
-		int32 Index = SceneGraphArchive.UnloadedReferences.Emplace(ReferenceId);
-		SceneGraphArchive.CADIdToUnloadedReferenceIndex.Add(ReferenceId, Index);
-		return Index;
-	}
-
-	FArchiveUnloadedReference& GetUnloadedReferenceAt(int32 Index)
-	{
-		return SceneGraphArchive.UnloadedReferences[Index];
-	}
-
 	FFileDescriptor& GetExternalReferences(int32 Index)
 	{
 		return SceneGraphArchive.ExternalReferenceFiles[Index];
diff --git a/Engine/Plugins/Enterprise/DatasmithCADImporter/Source/CADInterfaces/Public/CADSceneGraph.h b/Engine/Plugins/Enterprise/DatasmithCADImporter/Source/CADInterfaces/Public/CADSceneGraph.h
--- a/Engine/Plugins/Enterprise/DatasmithCADImporter/Source/CADInterfaces/Public/CADSceneGraph.h
+++ b/Engine/Plugins/Enterprise/DatasmithCADImporter/Source/CADInterfaces/Public/CADSceneGraph.h
@@ -10,11 +10,21 @@ class FArchive;
 namespace CADLibrary
 {
 
-class CADINTERFACES_API FArchiveCADObject
+class FArchiveUnloadedReference;
+
+class CADINTERFACES_API FArchiveCADObject : public FArchiveGraphicProperties
 {
 public:
-	FArchiveCADObject(FCadId InId = 0)
-		: Id(InId)
+	FArchiveCADObject()
+		: Id(0)
+		, Unit(1)
+	{
+	}
+
+	FArchiveCADObject(FCadId Id, const FArchiveCADObject& Parent)
+		: FArchiveGraphicProperties(Parent)
+		, Id(Id)
+		, Unit(Parent.Unit)
 	{
 	}
 
@@ -24,15 +34,28 @@ public:
 
 public:
 	uint32 Id;
+	FString Label;
 	TMap<FString, FString> MetaData;
 	FMatrix TransformMatrix = FMatrix::Identity;
+
+	// not serialized
+	double Unit = 1;
+
+	bool IsNameDefined() const
+	{
+		return !Label.IsEmpty();
+	}
+
+	bool SetNameWithAttributeValue(const TCHAR* Key);
+
 };
 
 class CADINTERFACES_API FArchiveInstance : public FArchiveCADObject
 {
 public:
-	FArchiveInstance(FCadId Id = 0)
-		: FArchiveCADObject(Id)
+	FArchiveInstance() = default;
+	FArchiveInstance(FCadId Id, const FArchiveCADObject& Parent)
+		: FArchiveCADObject(Id, Parent)
 	{
 	}
 
@@ -41,39 +64,56 @@ public:
 public:
 	FCadId ReferenceNodeId = 0;
 	bool bIsExternalReference = false;
-	FFileDescriptor ExternalReference;
 };
 
 class CADINTERFACES_API FArchiveReference : public FArchiveCADObject
 {
 public:
-	FArchiveReference(FCadId Id = 0)
-		: FArchiveCADObject(Id)
+	FArchiveReference() = default;
+	FArchiveReference(FCadId Id, const FArchiveCADObject& Reference)
+		: FArchiveCADObject(Id, Reference)
 	{
 	}
 
 	friend FArchive& operator<<(FArchive& Ar, FArchiveReference& C);
 
+	void AddChild(const FCadId ChildId);
+
+	int32 ChildrenCount();
+	void MoveTemp(FArchiveUnloadedReference& Reference);
+	void CopyMetaData(FArchiveUnloadedReference& Reference);
+
 public:
 	TArray<FCadId> Children;
 };
 
-class CADINTERFACES_API FArchiveUnloadedReference : public FArchiveReference
+class CADINTERFACES_API FArchiveUnloadedReference : public FArchiveCADObject
 {
 public:
-	FArchiveUnloadedReference(FCadId Id = 0)
-		: FArchiveReference(Id)
+	FArchiveUnloadedReference() = default;
+
+	FArchiveUnloadedReference(FCadId Id, const FArchiveCADObject& Parent)
+		: FArchiveCADObject(Id, Parent)
 	{
 	}
 
 	friend FArchive& operator<<(FArchive& Ar, FArchiveUnloadedReference& C);
+
+
+public:
+	FFileDescriptor ExternalFile;
+
+// not serialized
+	bool bIsUnloaded = true;
 };
 
 class CADINTERFACES_API FArchiveBody : public FArchiveCADObject
 {
 public:
-	FArchiveBody(FCadId Id = 0)
-		: FArchiveCADObject(Id)
+	FArchiveBody() = default;
+	FArchiveBody(FCadId Id, const FArchiveCADObject& Parent)
+		: FArchiveCADObject(Id, Parent)
+		, ParentId(Parent.Id)
 	{
 	}
 
@@ -87,6 +127,10 @@ public:
 	TSet<FMaterialUId> MaterialFaceSet;
 	TSet<FMaterialUId> ColorFaceSet;
 
+	void Delete();
+
+	// non serialized fields
+	bool bIsASolid = true;
 };
 
 class CADINTERFACES_API FArchiveColor
@@ -97,7 +141,7 @@ public:
 	{
 	}
 
-	friend FArchive& operator<<(FArchive& Ar, FArchiveColor& C);
+	friend FArchive& operator<<(FArchive& Ar, FArchiveColor& Color);
 
 public:
 	FMaterialUId Id;
@@ -113,7 +157,7 @@ public:
 	{
 	}
 
-	friend FArchive& operator<<(FArchive& Ar, FArchiveMaterial& C);
+	friend FArchive& operator<<(FArchive& Ar, FArchiveMaterial& Material);
 
 public:
 	FMaterialUId Id;
@@ -143,24 +187,57 @@ public:
 	TArray<FFileDescriptor> ExternalReferenceFiles;
 	TArray<FArchiveInstance> Instances;
 
-	TMap<FCadId, int32> CADIdToBodyIndex;
-	TMap<FCadId, int32> CADIdToReferenceIndex;
-	TMap<FCadId, int32> CADIdToUnloadedReferenceIndex;
-	TMap<FCadId, int32> CADIdToInstanceIndex;
+	TArray<int32> CADIdToIndex;
+	FCadId LastEntityId = 1;
+
+	void Reserve(uint32* ComponentCount)
+	{
+		// Must be call once
+		ensure(LastEntityId == 1);
+
+		Instances.Reserve(ComponentCount[EComponentType::Instance]);
+		References.Reserve(ComponentCount[EComponentType::Reference]);
+		UnloadedReferences.Reserve(ComponentCount[EComponentType::Reference]);
+		ExternalReferenceFiles.Reserve(ComponentCount[EComponentType::Reference]);
+		Bodies.Reserve(ComponentCount[EComponentType::Body]);
+
+		CADIdToIndex.Reserve(ComponentCount[EComponentType::Instance] + ComponentCount[EComponentType::Reference] + ComponentCount[EComponentType::Body] + 1);
+		CADIdToIndex.Add(0);
+	}
+
+	FArchiveInstance& AddInstance(const FArchiveCADObject& Parent);
+	FArchiveInstance& GetInstance(FCadId CadId);
+	void RemoveLastInstance();
+	bool IsAInstance(FCadId CadId) const;
+
+	FArchiveReference& AddReference(FArchiveUnloadedReference&);
+	FArchiveReference& AddReference(FArchiveInstance& Parent);
+	FArchiveReference& GetReference(FCadId CadId);
+	void RemoveLastReference();
+	bool IsAReference(FCadId CadId) const;
+
+	FArchiveReference& AddOccurence(FArchiveReference& Parent);
+	void RemoveLastOccurence();
+
+	FArchiveUnloadedReference& AddUnloadedReference(FArchiveInstance& Parent);
+	FArchiveUnloadedReference& GetUnloadedReference(FCadId CadId);
+	void RemoveLastUnloadedReference();
+	bool IsAUnloadedReference(FCadId CadId) const;
 
-	void Reserve(int32 InstanceNum, int32 ReferenceNum, int32 BodyNum)
+	FArchiveBody& AddBody(FArchiveReference& Parent);
+	FArchiveBody& GetBody(FCadId CadId);
+	void RemoveLastBody();
+	bool IsABody(FCadId CadId) const;
+
+	void AddExternalReferenceFile(const FArchiveUnloadedReference& Reference);
+
+	int32 ReferencesCount() const
 	{
-		Instances.Reserve(InstanceNum);
-		References.Reserve(ReferenceNum);
-		UnloadedReferences.Reserve(ReferenceNum);
-		ExternalReferenceFiles.Reserve(ReferenceNum);
-		Bodies.Reserve(BodyNum);
-
-		CADIdToInstanceIndex.Reserve(InstanceNum);
-		CADIdToReferenceIndex.Reserve(ReferenceNum);
-		CADIdToUnloadedReferenceIndex.Reserve(ReferenceNum);
-		CADIdToBodyIndex.Reserve(BodyNum);
+		return References.Num();
 	}
+
+private:
+
 };
 
 
diff --git a/Engine/Plugins/Enterprise/DatasmithCADImporter/Source/CADInterfaces/Public/TechSoftUtils.h b/Engine/Plugins/Enterprise/DatasmithCADImporter/Source/CADInterfaces/Public/TechSoftUtils.h
--- a/Engine/Plugins/Enterprise/DatasmithCADImporter/Source/CADInterfaces/Public/TechSoftUtils.h
+++ b/Engine/Plugins/Enterprise/DatasmithCADImporter/Source/CADInterfaces/Public/TechSoftUtils.h
@@ -30,17 +30,6 @@ namespace TechSoftInterface
 class FTechSoftInterface;
 }
 
-struct FEntityMetaData
-{
-	TMap<FString, FString> MetaData;
-	bool bRemoved = false;
-	bool bShow = true;
-	bool bUnloaded = false;
-	FCadUuid ColorName = 0;
-	FCadUuid MaterialName = 0;
-	FFileDescriptor ExternalFile;
-};
-
 namespace TechSoftUtils
 {
 
diff --git a/Engine/Plugins/Enterprise/DatasmithCADImporter/Source/CADTools/Public/CADData.h b/Engine/Plugins/Enterprise/DatasmithCADImporter/Source/CADTools/Public/CADData.h
--- a/Engine/Plugins/Enterprise/DatasmithCADImporter/Source/CADTools/Public/CADData.h
+++ b/Engine/Plugins/Enterprise/DatasmithCADImporter/Source/CADTools/Public/CADData.h
@@ -53,6 +53,21 @@ enum class ECADParsingResult : uint8
 	FileNotFound,
 };
 
+enum EComponentType : uint8
+{
+	Reference = 0,
+	Instance,
+	Body,
+	LastType
+};
+
+enum class ECADGraphicPropertyInheritance : uint8
+{
+	Unset,
+	FatherHerit,
+	ChildHerit,
+};
+
 // TODO: Remove from hear and replace by DatasmithUtils::GetCleanFilenameAndExtension... But need to remove DatasmithCore dependancies 
 CADTOOLS_API void GetCleanFilenameAndExtension(const FString& InFilePath, FString& OutFilename, FString& OutExtension);
 CADTOOLS_API FString GetExtension(const FString& InFilePath);
@@ -81,6 +96,73 @@ struct CADTOOLS_API FObjectDisplayDataId
 	FMaterialUId ColorUId = 0;
 };
 
+class CADTOOLS_API FArchiveGraphicProperties
+{
+public:
+	FArchiveGraphicProperties()
+		: ColorUId(0)
+		, MaterialUId(0)
+	{
+	}
+
+	FArchiveGraphicProperties(const FArchiveGraphicProperties& Parent)
+		: ColorUId(Parent.ColorUId)
+		, MaterialUId(Parent.MaterialUId)
+	{
+	}
+
+	virtual ~FArchiveGraphicProperties() = default;
+
+public:
+	FMaterialUId ColorUId = 0;
+	FMaterialUId MaterialUId = 0;
+	bool bIsRemoved = false;
+	bool bShow = true;
+	ECADGraphicPropertyInheritance Inheritance = ECADGraphicPropertyInheritance::Unset;
+
+	/**
+	 * If a graphic property is undefined, define it with Source property
+	 */
+	void DefineGraphicsPropertiesFromNoOverwrite(const FArchiveGraphicProperties& Source)
+	{
+		if (!ColorUId)
+		{
+			ColorUId = Source.ColorUId;
+		}
+
+		if (!MaterialUId)
+		{
+			MaterialUId = Source.MaterialUId;
+		}
+	}
+
+	/**
+	 * If a source property is defined, set the property with it
+	 */
+	void SetGraphicProperties(const FArchiveGraphicProperties& Source)
+	{
+		if (Source.ColorUId)
+		{
+			ColorUId = Source.ColorUId;
+		}
+
+		if (Source.MaterialUId)
+		{
+			MaterialUId = Source.MaterialUId;
+		}
+	}
+
+	bool IsDeleted() const
+	{
+		return bIsRemoved;
+	}
+
+	bool IsShown() const
+	{
+		return bShow;
+	}
+};
+
 class CADTOOLS_API FFileDescriptor
 {
 public:
@@ -211,7 +293,7 @@ private:
  * CoreTech mesh are defined surface by surface. The mesh is not connected
  * CADKernel mesh is connected.
  */
-struct CADTOOLS_API FTessellationData
+struct CADTOOLS_API FTessellationData : public FArchiveGraphicProperties
 {
 	friend CADTOOLS_API FArchive& operator<<(FArchive& Ar, FTessellationData& Tessellation);
 
@@ -230,9 +312,6 @@ struct CADTOOLS_API FTessellationData
 	/** UV coordinates of each vertex */
 	TArray<FVector2D> TexCoordArray;
 
-	FMaterialUId ColorUId = 0;
-	FMaterialUId MaterialUId = 0;
-
 	int32 PatchId;
 };
 
@@ -247,6 +326,19 @@ public:
 	friend FArchive& operator<<(FArchive& Ar, FBodyMesh& BodyMesh);
 
 public:
+
+	void AddGraphicPropertiesFrom(const FArchiveGraphicProperties& GraphicProperties)
+	{
+		if (GraphicProperties.ColorUId)
+		{
+			ColorSet.Add(GraphicProperties.ColorUId);
+		}
+		if (GraphicProperties.MaterialUId)
+		{
+			MaterialSet.Add(GraphicProperties.MaterialUId);
+		}
+	}
+
 	TArray<FVector> VertexArray; // set by CADKernel, filled by FillKioVertexPosition that merges coincident vertices (CoreTechHelper)
 	TArray<FTessellationData> Faces;
 	FBox BBox;
diff --git a/Engine/Plugins/Enterprise/DatasmithCADImporter/Source/DatasmithCADTranslator/Private/DatasmithSceneGraphBuilder.cpp b/Engine/Plugins/Enterprise/DatasmithCADImporter/Source/DatasmithCADTranslator/Private/DatasmithSceneGraphBuilder.cpp
--- a/Engine/Plugins/Enterprise/DatasmithCADImporter/Source/DatasmithCADTranslator/Private/DatasmithSceneGraphBuilder.cpp
+++ b/Engine/Plugins/Enterprise/DatasmithCADImporter/Source/DatasmithCADTranslator/Private/DatasmithSceneGraphBuilder.cpp
@@ -12,73 +12,111 @@
 
 namespace DatasmithSceneGraphBuilderImpl
 {
-	void GetMainMaterial(const TMap<FString, FString>& InNodeMetaDataMap, ActorData& OutNodeData, bool bMaterialPropagationIsTopDown)
+
+/**
+ * For more details:
+ * @see https://docs.techsoft3d.com/exchange/latest/build/managing_attribute_inheritance.html
+ */
+void SpreadGraphicProperties(const CADLibrary::FArchiveCADObject& Component, ActorData& ComponentData)
+{
+	if (!ComponentData.ColorUId)
+	{
+		ComponentData.ColorUId = Component.ColorUId;
+	}
+	if (!ComponentData.MaterialUId)
 	{
-		if (const FString* MaterialNameStr = InNodeMetaDataMap.Find(TEXT("MaterialName")))
+		ComponentData.MaterialUId = Component.MaterialUId;
+	}
+
+	if (Component.Inheritance == CADLibrary::ECADGraphicPropertyInheritance::ChildHerit)
+	{
+		if (Component.ColorUId)
 		{
-			if (!bMaterialPropagationIsTopDown || !OutNodeData.MaterialUuid)
-			{
-				OutNodeData.MaterialUuid = FCString::Atoi64(**MaterialNameStr);
-			}
+			ComponentData.ColorUId = Component.ColorUId;
 		}
-
-		if (const FString* ColorHashStr = InNodeMetaDataMap.Find(TEXT("ColorName")))
+		if (Component.MaterialUId)
 		{
-			if (bMaterialPropagationIsTopDown || !OutNodeData.ColorUuid)
-			{
-				OutNodeData.ColorUuid = FCString::Atoi64(**ColorHashStr);
-			}
+			ComponentData.MaterialUId = Component.MaterialUId;
 		}
+		ComponentData.Inheritance = CADLibrary::ECADGraphicPropertyInheritance::ChildHerit;
+	}
+	if (ComponentData.Inheritance == CADLibrary::ECADGraphicPropertyInheritance::ChildHerit)
+	{
+		return;
 	}
 
-	void AddTransformToActor(const CADLibrary::FArchiveCADObject& Object, TSharedPtr< IDatasmithActorElement > Actor, const CADLibrary::FImportParameters& ImportParameters)
+	if (Component.Inheritance == CADLibrary::ECADGraphicPropertyInheritance::FatherHerit)
 	{
-		if (!Actor.IsValid())
+		if (Component.ColorUId)
 		{
-			return;
+			ComponentData.ColorUId = Component.ColorUId;
 		}
+		if (Component.MaterialUId)
+		{
+			ComponentData.MaterialUId = Component.MaterialUId;
+		}
+		ComponentData.Inheritance = CADLibrary::ECADGraphicPropertyInheritance::FatherHerit;
+	}
+	if (ComponentData.Inheritance == CADLibrary::ECADGraphicPropertyInheritance::FatherHerit)
+	{
+		return;
+	}
 
-		FTransform LocalTransform(Object.TransformMatrix);
-		FTransform LocalUETransform = FDatasmithUtils::ConvertTransform(ImportParameters.GetModelCoordSys(), LocalTransform);
-
-		Actor->SetTranslation(LocalUETransform.GetTranslation() * ImportParameters.GetScaleFactor());
-		Actor->SetScale(LocalUETransform.GetScale3D());
-		Actor->SetRotation(LocalUETransform.GetRotation());
+	if (!ComponentData.ColorUId)
+	{
+		ComponentData.ColorUId = Component.ColorUId;
+	}
+	if (!ComponentData.MaterialUId)
+	{
+		ComponentData.MaterialUId = Component.MaterialUId;
 	}
+}
 
-	// Method to reduce the size of huge label. The length of the package path, based on label, cannot be bigger than ~256
-	void CleanName(FString& Label)
+void AddTransformToActor(const CADLibrary::FArchiveCADObject& Object, TSharedPtr<IDatasmithActorElement> Actor, const CADLibrary::FImportParameters& ImportParameters)
+{
+	if (!Actor.IsValid())
 	{
-		const int32 MaxLabelSize = 50; // If the label is smaller than this value, the label is not modified. This size of package name is "acceptable"
-		const int32 ReasonableLabelSize = 20; // If the label has to be cut, a label that is not too long is preferred. 
-		const int32 MinLabelSize = 5; // If the label is smaller than this value, the label is too much reduce. Therefore a ReasonableLabelSize is preferred 
+		return;
+	}
 
-		if (Label.Len() < MaxLabelSize)
-		{
-			return;
-		}
+	FTransform LocalTransform(Object.TransformMatrix);
+	FTransform LocalUETransform = FDatasmithUtils::ConvertTransform(ImportParameters.GetModelCoordSys(), LocalTransform);
 
-		FString NewLabel;
-		NewLabel = FPaths::GetCleanFilename(Label);
-		if ((NewLabel.Len() < MaxLabelSize) && (NewLabel.Len() > MinLabelSize))
-		{
-			Label = NewLabel;
-			return;
-		}
+	Actor->SetTranslation(LocalUETransform.GetTranslation() * ImportParameters.GetScaleFactor());
+	Actor->SetScale(LocalUETransform.GetScale3D());
+	Actor->SetRotation(LocalUETransform.GetRotation());
+}
 
-		Label = Label.Right(ReasonableLabelSize);
+void GetNodeUuidAndLabel(CADLibrary::FArchiveInstance& Instance, CADLibrary::FArchiveCADObject& Reference, const TCHAR* InParentUEUUID, FString& OutUEUUID, FString& OutLabel)
+{
+	if (Instance.IsNameDefined())
+	{
+		OutLabel = Instance.Label;
 	}
-}
+	else if (Reference.IsNameDefined())
+	{
+		OutLabel = Reference.Label;
+	}
+	else
+	{
+		OutLabel = TEXT("NoName");
+	}
+
+	FCadUuid UeUuid = HashCombine(GetTypeHash(InParentUEUUID), GetTypeHash(Instance.Id));
+	UeUuid = HashCombine(UeUuid, GetTypeHash(Reference.Id));
 
+	OutUEUUID = FString::Printf(TEXT("0x%08x"), UeUuid);
+}
+}
 
 FDatasmithSceneGraphBuilder::FDatasmithSceneGraphBuilder(
-	TMap<uint32, FString>& InCADFileToUnrealFileMap, 
-	const FString& InCachePath, 
-	TSharedRef<IDatasmithScene> InScene, 
-	const FDatasmithSceneSource& InSource, 
+	TMap<uint32, FString>& InCADFileToUnrealFileMap,
+	const FString& InCachePath,
+	TSharedRef<IDatasmithScene> InScene,
+	const FDatasmithSceneSource& InSource,
 	const CADLibrary::FImportParameters& InImportParameters)
-		: FDatasmithSceneBaseGraphBuilder(nullptr, InCachePath, InScene, InSource, InImportParameters)
-		, CADFileToSceneGraphDescriptionFile(InCADFileToUnrealFileMap)
+	: FDatasmithSceneBaseGraphBuilder(nullptr, InCachePath, InScene, InSource, InImportParameters)
+	, CADFileToSceneGraphDescriptionFile(InCADFileToUnrealFileMap)
 {
 }
 
@@ -113,7 +151,7 @@ void FDatasmithSceneGraphBuilder::LoadSceneGraphDescriptionFiles()
 
 		MockUpDescription.DeserializeMockUpFile(*MockUpDescriptionFile);
 
-		for(const auto& ColorPair : MockUpDescription.ColorHIdToColor)
+		for (const auto& ColorPair : MockUpDescription.ColorHIdToColor)
 		{
 			ColorUIdToColorArchive.Emplace(ColorPair.Value.UEMaterialUId, ColorPair.Value);
 		}
@@ -138,35 +176,28 @@ void FDatasmithSceneGraphBuilder::FillAnchorActor(const TSharedRef<IDatasmithAct
 		return;
 	}
 
-	FCadId RootId = 1;
-	const int32* Index = SceneGraph->CADIdToReferenceIndex.Find(RootId);
-	if (!Index)
-	{
-		return;
-	}
-
 	ActorData Data(TEXT(""));
 
-	// TODO: check ParentData and Index validity?
+	const FCadId RootId = 1;
 	ActorData ParentData(ActorElement->GetName());
-	CADLibrary::FArchiveReference& Reference = SceneGraph->References[*Index];
+	CADLibrary::FArchiveReference& Reference = SceneGraph->GetReference(RootId);
 
-	TMap<FString, FString> InstanceNodeMetaDataMap;
+	CADLibrary::FArchiveInstance EmptyInstance;
 	FString ActorUUID;
 	FString ActorLabel;
-	GetNodeUUIDAndName(InstanceNodeMetaDataMap, Reference.MetaData, Reference.Id, ParentData.Uuid, ActorUUID, ActorLabel);
+	DatasmithSceneGraphBuilderImpl::GetNodeUuidAndLabel(EmptyInstance, Reference, ParentData.Uuid, ActorUUID, ActorLabel);
 
-	AddMetaData(ActorElement, InstanceNodeMetaDataMap, Reference.MetaData);
+	AddMetaData(ActorElement, EmptyInstance, Reference);
 
-	ActorData ReferenceData(*ActorUUID, ParentData);
-	DatasmithSceneGraphBuilderImpl::GetMainMaterial(Reference.MetaData, ReferenceData, bMaterialPropagationIsTopDown);
+	ActorData ComponentData(*ActorUUID, ParentData);
+	//DatasmithSceneGraphBuilderImpl::GetMainMaterial(Reference.MetaData, ComponentData, bMaterialPropagationIsTopDown);
 
-	AddChildren(ActorElement, Reference, ReferenceData);
+	AddChildren(ActorElement, Reference, ComponentData);
 
 	ActorElement->SetLabel(*ActorLabel);
 }
 
-FDatasmithSceneBaseGraphBuilder::FDatasmithSceneBaseGraphBuilder(CADLibrary::FArchiveSceneGraph* InSceneGraph, const FString& InCachePath,  TSharedRef<IDatasmithScene> InScene, const FDatasmithSceneSource& InSource, const CADLibrary::FImportParameters& InImportParameters)
+FDatasmithSceneBaseGraphBuilder::FDatasmithSceneBaseGraphBuilder(CADLibrary::FArchiveSceneGraph* InSceneGraph, const FString& InCachePath, TSharedRef<IDatasmithScene> InScene, const FDatasmithSceneSource& InSource, const CADLibrary::FImportParameters& InImportParameters)
 	: SceneGraph(InSceneGraph)
 	, CachePath(InCachePath)
 	, DatasmithScene(InScene)
@@ -179,7 +210,7 @@ FDatasmithSceneBaseGraphBuilder::FDatasmithSceneBaseGraphBuilder(CADLibrary::FAr
 	if (InSceneGraph)
 	{
 		ColorUIdToColorArchive.Reserve(SceneGraph->ColorHIdToColor.Num());
-		for(const auto& ColorPair : SceneGraph->ColorHIdToColor)
+		for (const auto& ColorPair : SceneGraph->ColorHIdToColor)
 		{
 			ColorUIdToColorArchive.Emplace(ColorPair.Value.UEMaterialUId, ColorPair.Value);
 		}
@@ -194,16 +225,11 @@ FDatasmithSceneBaseGraphBuilder::FDatasmithSceneBaseGraphBuilder(CADLibrary::FAr
 
 bool FDatasmithSceneBaseGraphBuilder::Build()
 {
-	FCadId RootId = 1;
-	const int32* Index = SceneGraph->CADIdToReferenceIndex.Find(RootId);
-	if (!Index)
-	{
-		return false;
-	}
+	const FCadId RootId = 1;
 
 	ActorData Data(TEXT(""));
-	CADLibrary::FArchiveReference& Reference = SceneGraph->References[*Index];
-	TSharedPtr< IDatasmithActorElement > RootActor = BuildReference(Reference, Data);
+	CADLibrary::FArchiveReference& Reference = SceneGraph->GetReference(RootId);
+	TSharedPtr<IDatasmithActorElement> RootActor = BuildReference(Reference, Data);
 	DatasmithScene->AddActor(RootActor);
 
 	// Set ProductName, ProductVersion in DatasmithScene for Analytics purpose
@@ -217,7 +243,7 @@ bool FDatasmithSceneBaseGraphBuilder::Build()
 
 	FString ProductName;
 	const FString* ProductNamePtr = Reference.MetaData.Find(TEXT("Input_Format_and_Emitter"));
-	if(ProductNamePtr)
+	if (ProductNamePtr)
 	{
 		ProductName = *ProductNamePtr;
 		ProductName.TrimStartAndEndInline();
@@ -227,7 +253,7 @@ bool FDatasmithSceneBaseGraphBuilder::Build()
 		}
 	}
 
-	if(ProductName.IsEmpty())
+	if (ProductName.IsEmpty())
 	{
 		switch (RootFileDescription.GetFileFormat())
 		{
@@ -285,19 +311,21 @@ bool FDatasmithSceneBaseGraphBuilder::Build()
 	return true;
 }
 
-TSharedPtr< IDatasmithActorElement >  FDatasmithSceneBaseGraphBuilder::BuildInstance(int32 InstanceIndex, const ActorData& ParentData)
+TSharedPtr<IDatasmithActorElement>  FDatasmithSceneBaseGraphBuilder::BuildInstance(FCadId InstanceId, const ActorData& ParentData)
 {
 	CADLibrary::FArchiveReference* Reference = nullptr;
 	CADLibrary::FArchiveReference EmptyReference;
 
-	CADLibrary::FArchiveInstance& Instance = SceneGraph->Instances[InstanceIndex];
+	CADLibrary::FArchiveInstance& Instance = SceneGraph->GetInstance(InstanceId);
 
 	CADLibrary::FArchiveSceneGraph* InstanceSceneGraph = SceneGraph;
 	if (Instance.bIsExternalReference)
 	{
-		if (!Instance.ExternalReference.GetSourcePath().IsEmpty())
+		CADLibrary::FArchiveUnloadedReference& UnloadedReference = SceneGraph->GetUnloadedReference(Instance.ReferenceNodeId);
+
+		if (!UnloadedReference.ExternalFile.GetSourcePath().IsEmpty())
 		{
-			uint32 InstanceSceneGraphHash = Instance.ExternalReference.GetDescriptorHash();
+			uint32 InstanceSceneGraphHash = UnloadedReference.ExternalFile.GetDescriptorHash();
 			SceneGraph = CADFileToSceneGraphArchive.FindRef(InstanceSceneGraphHash);
 
 			if (SceneGraph)
@@ -307,32 +335,22 @@ TSharedPtr< IDatasmithActorElement >  FDatasmithSceneBaseGraphBuilder::BuildInst
 					AncestorSceneGraphHash.Add(InstanceSceneGraphHash);
 
 					FCadId RootId = 1;
-					const int32* Index = SceneGraph->CADIdToReferenceIndex.Find(RootId);
-					if (Index)
-					{
-						Reference = &(SceneGraph->References[*Index]);
-					}
+					Reference = &(SceneGraph->GetReference(RootId));
 				}
 			}
+			Reference->SetGraphicProperties(UnloadedReference);
 		}
 
-		if(!Reference)
+		if (!Reference)
 		{
 			SceneGraph = InstanceSceneGraph;
-			const int32* Index = SceneGraph->CADIdToUnloadedReferenceIndex.Find(Instance.ReferenceNodeId);
-			if (Index)
-			{
-				Reference = &(SceneGraph->UnloadedReferences[*Index]);
-			}
+			EmptyReference.CopyMetaData(UnloadedReference);
+			Reference = &EmptyReference;
 		}
 	}
-	else 
+	else
 	{
-		const int32* Index = SceneGraph->CADIdToReferenceIndex.Find(Instance.ReferenceNodeId);
-		if (Index)
-		{
-			Reference = &(SceneGraph->References[*Index]);
-		}
+		Reference = &(SceneGraph->GetReference(Instance.ReferenceNodeId));
 	}
 
 	if (!Reference) // Should never append
@@ -342,19 +360,19 @@ TSharedPtr< IDatasmithActorElement >  FDatasmithSceneBaseGraphBuilder::BuildInst
 
 	FString ActorUUID;
 	FString ActorLabel;
-	GetNodeUUIDAndName(Instance.MetaData, Reference->MetaData, Instance.Id, ParentData.Uuid, ActorUUID, ActorLabel);
+	DatasmithSceneGraphBuilderImpl::GetNodeUuidAndLabel(Instance, *Reference, ParentData.Uuid, ActorUUID, ActorLabel);
 
 	TSharedPtr<IDatasmithActorElement> Actor = CreateActor(*ActorUUID, *ActorLabel);
 	if (!Actor.IsValid())
 	{
-		return TSharedPtr< IDatasmithActorElement >();
+		return TSharedPtr<IDatasmithActorElement>();
 	}
-	AddMetaData(Actor, Instance.MetaData, Reference->MetaData);
+	AddMetaData(Actor, Instance, *Reference);
 
 	ActorData InstanceData(*ActorUUID, ParentData);
 
-	DatasmithSceneGraphBuilderImpl::GetMainMaterial(Instance.MetaData, InstanceData, bMaterialPropagationIsTopDown);
-	DatasmithSceneGraphBuilderImpl::GetMainMaterial(Reference->MetaData, InstanceData, bMaterialPropagationIsTopDown);
+	DatasmithSceneGraphBuilderImpl::SpreadGraphicProperties(Instance, InstanceData);
+	DatasmithSceneGraphBuilderImpl::SpreadGraphicProperties(*Reference, InstanceData);
 
 	AddChildren(Actor, *Reference, InstanceData);
 
@@ -368,9 +386,9 @@ TSharedPtr< IDatasmithActorElement >  FDatasmithSceneBaseGraphBuilder::BuildInst
 	return Actor;
 }
 
-TSharedPtr< IDatasmithActorElement >  FDatasmithSceneBaseGraphBuilder::CreateActor(const TCHAR* InEUUID, const TCHAR* InLabel)
+TSharedPtr<IDatasmithActorElement>  FDatasmithSceneBaseGraphBuilder::CreateActor(const TCHAR* InEUUID, const TCHAR* InLabel)
 {
-	TSharedPtr< IDatasmithActorElement > Actor = FDatasmithSceneFactory::CreateActor(InEUUID);
+	TSharedPtr<IDatasmithActorElement> Actor = FDatasmithSceneFactory::CreateActor(InEUUID);
 	if (Actor.IsValid())
 	{
 		Actor->SetLabel(InLabel);
@@ -379,95 +397,25 @@ TSharedPtr< IDatasmithActorElement >  FDatasmithSceneBaseGraphBuilder::CreateAct
 	return TSharedPtr<IDatasmithActorElement>();
 }
 
-void FDatasmithSceneBaseGraphBuilder::GetNodeUUIDAndName(
-	const TMap<FString, FString>& InInstanceNodeMetaDataMap,
-	const TMap<FString, FString>& InReferenceNodeMetaDataMap,
-	int32 InReferenceIndex,
-	const TCHAR* InParentUEUUID,
-	FString& OutUEUUID,
-	FString& OutName
-)
-{
-	const FString* InstanceSDKName = InInstanceNodeMetaDataMap.Find(TEXT("SDKName"));
-	const FString* InstanceCADName = InInstanceNodeMetaDataMap.Find(TEXT("Name"));
-	const FString* InstanceUUID = InInstanceNodeMetaDataMap.Find(TEXT("UUID"));
-
-	const FString* ReferenceSDKName = InReferenceNodeMetaDataMap.Find(TEXT("SDKName"));
-	const FString* ReferenceCADName = InReferenceNodeMetaDataMap.Find(TEXT("Name"));
-	const FString* ReferenceUUID = InReferenceNodeMetaDataMap.Find(TEXT("UUID"));
-
-	// Outname Name
-	// Instance SDK Name and Reference SDName are build Name. Original names (CAD system name / "Name") are preferred
-	if (InstanceCADName && !InstanceCADName->IsEmpty())
-	{
-		OutName = *InstanceCADName;
-	}
-	else if (ReferenceCADName && !ReferenceCADName->IsEmpty())
-	{
-		OutName = *ReferenceCADName;
-	}
-	else if (InstanceSDKName && !InstanceSDKName->IsEmpty())
-	{
-		OutName = *InstanceSDKName;
-	}
-	else if (ReferenceSDKName && !ReferenceSDKName->IsEmpty())
-	{
-		OutName = *ReferenceSDKName;
-	}
-	else
-	{
-		OutName = "NoName";
-	}
-
-	FCadUuid UEUUID = HashCombine(GetTypeHash(InParentUEUUID), GetTypeHash(InReferenceIndex));
-
-	if (InstanceUUID)
-	{
-		UEUUID = HashCombine(UEUUID, GetTypeHash(*InstanceUUID));
-	}
-	if (InstanceCADName)
-	{
-		UEUUID = HashCombine(UEUUID, GetTypeHash(*InstanceCADName));
-	}
-	if (InstanceSDKName)
-	{
-		UEUUID = HashCombine(UEUUID, GetTypeHash(*InstanceSDKName));
-	}
-
-	if (ReferenceUUID)
-	{
-		UEUUID = HashCombine(UEUUID, GetTypeHash(*ReferenceUUID));
-	}
-	if (ReferenceCADName)
-	{
-		UEUUID = HashCombine(UEUUID, GetTypeHash(*ReferenceCADName));
-	}
-	if (ReferenceSDKName)
-	{
-		UEUUID = HashCombine(UEUUID, GetTypeHash(*ReferenceSDKName));
-	}
-
-	OutUEUUID = FString::Printf(TEXT("0x%08x"), UEUUID);
-}
-
 TSharedPtr<IDatasmithActorElement> FDatasmithSceneBaseGraphBuilder::BuildReference(CADLibrary::FArchiveReference& Reference, const ActorData& ParentData)
 {
 	TMap<FString, FString> InstanceNodeMetaDataMap;
 
+	CADLibrary::FArchiveInstance EmptyInstance;
 	FString ActorUUID;
 	FString ActorLabel;
-	GetNodeUUIDAndName(InstanceNodeMetaDataMap, Reference.MetaData, Reference.Id, ParentData.Uuid, ActorUUID, ActorLabel);
+	DatasmithSceneGraphBuilderImpl::GetNodeUuidAndLabel(EmptyInstance, Reference, ParentData.Uuid, ActorUUID, ActorLabel);
 
-	TSharedPtr< IDatasmithActorElement > Actor = CreateActor(*ActorUUID, *ActorLabel);
+	TSharedPtr<IDatasmithActorElement> Actor = CreateActor(*ActorUUID, *ActorLabel);
 	if (!Actor.IsValid())
 	{
-		return TSharedPtr< IDatasmithActorElement >();
+		return TSharedPtr<IDatasmithActorElement>();
 	}
 
-	AddMetaData(Actor, InstanceNodeMetaDataMap, Reference.MetaData);
+	AddMetaData(Actor, EmptyInstance, Reference);
 
 	ActorData ReferenceData(*ActorUUID, ParentData);
-	DatasmithSceneGraphBuilderImpl::GetMainMaterial(Reference.MetaData, ReferenceData, bMaterialPropagationIsTopDown);
+	DatasmithSceneGraphBuilderImpl::SpreadGraphicProperties(Reference, ReferenceData);
 
 	AddChildren(Actor, Reference, ReferenceData);
 
@@ -476,41 +424,41 @@ TSharedPtr<IDatasmithActorElement> FDatasmithSceneBaseGraphBuilder::BuildReferen
 	return Actor;
 }
 
-TSharedPtr<IDatasmithActorElement> FDatasmithSceneBaseGraphBuilder::BuildBody(int32 BodyIndex, const ActorData& ParentData)
+TSharedPtr<IDatasmithActorElement> FDatasmithSceneBaseGraphBuilder::BuildBody(FCadId BodyId, const ActorData& ParentData)
 {
-	TMap<FString, FString> InstanceNodeMetaDataMap;
-
-	CADLibrary::FArchiveBody& Body = SceneGraph->Bodies[BodyIndex];
+	CADLibrary::FArchiveInstance EmptyInstance;
+	CADLibrary::FArchiveBody& Body = SceneGraph->GetBody(BodyId);
 
-	if (Body.ParentId == 0 || Body.MeshActorUId == 0)
+	if (Body.IsDeleted())
 	{
 		return TSharedPtr<IDatasmithActorElement>();
 	}
 
 	FString BodyUUID;
 	FString BodyLabel;
-	GetNodeUUIDAndName(InstanceNodeMetaDataMap, Body.MetaData, Body.Id, ParentData.Uuid, BodyUUID, BodyLabel);
+	DatasmithSceneGraphBuilderImpl::GetNodeUuidAndLabel(EmptyInstance, Body, ParentData.Uuid, BodyUUID, BodyLabel);
 
 	// Apply materials on the current part
-	uint32 MaterialUuid = 0;
-	MaterialUuid = ParentData.MaterialUuid ? ParentData.MaterialUuid : ParentData.ColorUuid;
+	ActorData BodyData(*BodyUUID, ParentData);
+	DatasmithSceneGraphBuilderImpl::SpreadGraphicProperties(Body, BodyData);
+	FMaterialUId MaterialUId = (BodyData.Inheritance == CADLibrary::ECADGraphicPropertyInheritance::Unset) ? -1 : BodyData.MaterialUId ? BodyData.MaterialUId : BodyData.ColorUId;
 
 	if (!(Body.ColorFaceSet.Num() + Body.MaterialFaceSet.Num()))
 	{
-		Body.ColorFaceSet.Add(MaterialUuid);
+		Body.ColorFaceSet.Add(MaterialUId);
 	}
 
-	TSharedPtr<IDatasmithMeshElement> MeshElement = FindOrAddMeshElement(Body, BodyLabel);
+	TSharedPtr<IDatasmithMeshElement> MeshElement = FindOrAddMeshElement(Body);
 	if (!MeshElement.IsValid())
 	{
-		return TSharedPtr< IDatasmithActorElement >();
+		return TSharedPtr<IDatasmithActorElement>();
 	}
 
 
 	TSharedPtr<IDatasmithMeshActorElement> ActorElement = FDatasmithSceneFactory::CreateMeshActor(*BodyUUID);
 	if (!ActorElement.IsValid())
 	{
-		return TSharedPtr< IDatasmithActorElement >();
+		return TSharedPtr<IDatasmithActorElement>();
 	}
 
 	ActorElement->SetLabel(*BodyLabel);
@@ -518,14 +466,14 @@ TSharedPtr<IDatasmithActorElement> FDatasmithSceneBaseGraphBuilder::BuildBody(in
 
 	DatasmithSceneGraphBuilderImpl::AddTransformToActor(Body, ActorElement, ImportParameters);
 
-	if (MaterialUuid && ImportParameters.GetPropagation() != CADLibrary::EDisplayDataPropagationMode::BodyOnly)
+	if (MaterialUId >= 0)
 	{
-		TSharedPtr< IDatasmithMaterialIDElement > PartMaterialIDElement = FindOrAddMaterial(MaterialUuid);
+		TSharedPtr<IDatasmithMaterialIDElement> PartMaterialIDElement = FindOrAddMaterial(MaterialUId);
 		const TCHAR* MaterialIDElementName = PartMaterialIDElement->GetName();
 
 		for (int32 Index = 0; Index < MeshElement->GetMaterialSlotCount(); ++Index)
 		{
-			TSharedPtr< IDatasmithMaterialIDElement > MaterialIDElement = FDatasmithSceneFactory::CreateMaterialId(MaterialIDElementName);
+			TSharedPtr<IDatasmithMaterialIDElement> MaterialIDElement = FDatasmithSceneFactory::CreateMaterialId(MaterialIDElementName);
 			MaterialIDElement->SetId(MeshElement->GetMaterialSlotAt(Index)->GetId());
 			ActorElement->AddMaterialOverride(MaterialIDElement);
 		}
@@ -533,19 +481,19 @@ TSharedPtr<IDatasmithActorElement> FDatasmithSceneBaseGraphBuilder::BuildBody(in
 	return ActorElement;
 }
 
-TSharedPtr< IDatasmithMeshElement > FDatasmithSceneBaseGraphBuilder::FindOrAddMeshElement(CADLibrary::FArchiveBody& Body, FString& BodyName)
+TSharedPtr<IDatasmithMeshElement> FDatasmithSceneBaseGraphBuilder::FindOrAddMeshElement(CADLibrary::FArchiveBody& Body)
 {
 	FString ShellUuidName = FString::Printf(TEXT("0x%012u"), Body.MeshActorUId);
 
 	// Look if geometry has not been already processed, return it if found
-	TSharedPtr< IDatasmithMeshElement >* MeshElementPtr = BodyUuidToMeshElement.Find(Body.MeshActorUId);
+	TSharedPtr<IDatasmithMeshElement>* MeshElementPtr = BodyUuidToMeshElement.Find(Body.MeshActorUId);
 	if (MeshElementPtr != nullptr)
 	{
 		return *MeshElementPtr;
 	}
 
-	TSharedPtr< IDatasmithMeshElement > MeshElement = FDatasmithSceneFactory::CreateMesh(*ShellUuidName);
-	MeshElement->SetLabel(*BodyName);
+	TSharedPtr<IDatasmithMeshElement> MeshElement = FDatasmithSceneFactory::CreateMesh(*ShellUuidName);
+	MeshElement->SetLabel(*Body.Label);
 	MeshElement->SetLightmapSourceUV(-1);
 
 	// Set MeshElement FileHash used for re-import task 
@@ -564,7 +512,7 @@ TSharedPtr< IDatasmithMeshElement > FDatasmithSceneBaseGraphBuilder::FindOrAddMe
 	{
 		for (FMaterialUId MaterialSlotId : MaterialSet)
 		{
-			TSharedPtr< IDatasmithMaterialIDElement > PartMaterialIDElement;
+			TSharedPtr<IDatasmithMaterialIDElement> PartMaterialIDElement;
 			PartMaterialIDElement = FindOrAddMaterial(MaterialSlotId);
 			MeshElement->SetMaterial(PartMaterialIDElement->GetName(), MaterialSlotId);
 		}
@@ -583,7 +531,7 @@ TSharedPtr< IDatasmithMeshElement > FDatasmithSceneBaseGraphBuilder::FindOrAddMe
 	return MeshElement;
 }
 
-TSharedPtr< IDatasmithUEPbrMaterialElement > FDatasmithSceneBaseGraphBuilder::GetDefaultMaterial()
+TSharedPtr<IDatasmithUEPbrMaterialElement> FDatasmithSceneBaseGraphBuilder::GetDefaultMaterial()
 {
 	if (!DefaultMaterial.IsValid())
 	{
@@ -594,11 +542,11 @@ TSharedPtr< IDatasmithUEPbrMaterialElement > FDatasmithSceneBaseGraphBuilder::Ge
 	return DefaultMaterial;
 }
 
-TSharedPtr<IDatasmithMaterialIDElement> FDatasmithSceneBaseGraphBuilder::FindOrAddMaterial(uint32 MaterialUuid)
+TSharedPtr<IDatasmithMaterialIDElement> FDatasmithSceneBaseGraphBuilder::FindOrAddMaterial(FMaterialUId MaterialUuid)
 {
-	TSharedPtr< IDatasmithUEPbrMaterialElement > MaterialElement;
+	TSharedPtr<IDatasmithUEPbrMaterialElement> MaterialElement;
 
-	TSharedPtr< IDatasmithUEPbrMaterialElement >* MaterialPtr = MaterialUuidMap.Find(MaterialUuid);
+	TSharedPtr<IDatasmithUEPbrMaterialElement>* MaterialPtr = MaterialUuidMap.Find(MaterialUuid);
 	if (MaterialPtr != nullptr)
 	{
 		MaterialElement = *MaterialPtr;
@@ -627,120 +575,58 @@ TSharedPtr<IDatasmithMaterialIDElement> FDatasmithSceneBaseGraphBuilder::FindOrA
 		MaterialUuidMap.Add(MaterialUuid, MaterialElement);
 	}
 
-	TSharedPtr< IDatasmithMaterialIDElement > MaterialIDElement = FDatasmithSceneFactory::CreateMaterialId(MaterialElement->GetName());
+	TSharedPtr<IDatasmithMaterialIDElement> MaterialIDElement = FDatasmithSceneFactory::CreateMaterialId(MaterialElement->GetName());
 
 	return MaterialIDElement;
 }
 
-void FDatasmithSceneBaseGraphBuilder::AddMetaData(TSharedPtr< IDatasmithActorElement > ActorElement, TMap<FString, FString>& InstanceNodeAttributeSetMap, TMap<FString, FString>& ReferenceNodeAttributeSetMap)
+void FDatasmithSceneBaseGraphBuilder::AddMetaData(TSharedPtr<IDatasmithActorElement> ActorElement, const CADLibrary::FArchiveCADObject& Instance, const CADLibrary::FArchiveCADObject& Reference)
 {
-	// Initialize list of attributes not to pass as meta-data
-	auto GetUnwantedAttributes = []() -> TSet<FString>
-	{
-		TSet<FString> UnwantedAttributes;
-
-		UnwantedAttributes.Add(TEXT("SDKName"));
-
-		// CoreTech
-		UnwantedAttributes.Add(TEXT("CTName"));
-		UnwantedAttributes.Add(TEXT("LayerId"));
-		UnwantedAttributes.Add(TEXT("LayerName"));
-		UnwantedAttributes.Add(TEXT("LayerFlag"));
-		UnwantedAttributes.Add(TEXT("OriginalUnitsMass"));
-		UnwantedAttributes.Add(TEXT("OriginalUnitsLength"));
-		UnwantedAttributes.Add(TEXT("OriginalUnitsDuration"));
-		UnwantedAttributes.Add(TEXT("OriginalIdStr"));
-		UnwantedAttributes.Add(TEXT("ShowAttribute"));
-		UnwantedAttributes.Add(TEXT("Identification"));
-		UnwantedAttributes.Add(TEXT("MaterialId"));
-		UnwantedAttributes.Add(TEXT("ColorUEId"));
-		UnwantedAttributes.Add(TEXT("ColorId"));
-		UnwantedAttributes.Add(TEXT("KernelIOVersion"));
-		return UnwantedAttributes;
-	};
 
-	static const TSet<FString> UnwantedAttributes = GetUnwantedAttributes();
+	const TMap<FString, FString>& InstanceNodeAttributeSetMap = Instance.MetaData;
+	const TMap<FString, FString>& ReferenceNodeAttributeSetMap = Reference.MetaData;
 
-	TSharedRef< IDatasmithMetaDataElement > MetaDataRefElement = FDatasmithSceneFactory::CreateMetaData(ActorElement->GetName());
+	TSharedRef<IDatasmithMetaDataElement> MetaDataRefElement = FDatasmithSceneFactory::CreateMetaData(ActorElement->GetName());
 	MetaDataRefElement->SetAssociatedElement(ActorElement);
 
-	for (auto& Attribute : ReferenceNodeAttributeSetMap)
+	TFunction<void(const CADLibrary::FArchiveCADObject&, const TCHAR*)> AddMetaData = [&](const CADLibrary::FArchiveCADObject& Component, const TCHAR* PostName)
 	{
-		if (UnwantedAttributes.Contains(Attribute.Key))
-		{
-			continue;
-		}
-
-		if (Attribute.Value.IsEmpty())
-		{
-			continue;
-		}
-
-		// If file information are attached to object, make sure to set a workable and full path
-		if (Attribute.Key == TEXT("FileName"))
+		for (auto& Attribute : Component.MetaData)
 		{
-			FString OFilePath = Attribute.Value;
-			if (FPaths::FileExists(OFilePath))
+			if (Attribute.Value.IsEmpty())
 			{
-				OFilePath = *FPaths::ConvertRelativePathToFull(OFilePath);
+				continue;
 			}
-			else
-			{
-				FString FileDir = RootFileDescription.GetRootFolder();
-				FString FilePath = FPaths::Combine(FileDir, OFilePath);
 
-				if (FPaths::FileExists(FilePath))
-				{
-					OFilePath = *FPaths::ConvertRelativePathToFull(OFilePath);
-				}
-				else // No workable file path to store. Skip
-				{
-					continue;
-				}
-			}
-
-			// Beautifying the attributes name
-			Attribute.Key = TEXT("FilePath");
-			Attribute.Value = OFilePath;
-		}
-
-		FString MetaName = TEXT("Reference ");
-		MetaName += Attribute.Key;
-		TSharedRef< IDatasmithKeyValueProperty > KeyValueProperty = FDatasmithSceneFactory::CreateKeyValueProperty(*MetaName);
+			FString MetaName = PostName;
+			MetaName += Attribute.Key;
+			TSharedRef<IDatasmithKeyValueProperty> KeyValueProperty = FDatasmithSceneFactory::CreateKeyValueProperty(*MetaName);
 
-		KeyValueProperty->SetValue(*Attribute.Value);
-		KeyValueProperty->SetPropertyType(EDatasmithKeyValuePropertyType::String);
-
-		MetaDataRefElement->AddProperty(KeyValueProperty);
-	}
+			KeyValueProperty->SetValue(*Attribute.Value);
+			KeyValueProperty->SetPropertyType(EDatasmithKeyValuePropertyType::String);
 
-	for (const auto& Attribute : InstanceNodeAttributeSetMap)
-	{
-		if (UnwantedAttributes.Contains(*Attribute.Key))
-		{
-			continue;
+			MetaDataRefElement->AddProperty(KeyValueProperty);
 		}
 
-		if (Attribute.Value.IsEmpty())
+		// Add name
 		{
-			continue;
+			FString MetaName = PostName;
+			MetaName += TEXT(" Name");
+			TSharedRef<IDatasmithKeyValueProperty> KeyValueProperty = FDatasmithSceneFactory::CreateKeyValueProperty(*MetaName);
+			KeyValueProperty->SetValue(*Component.Label);
+			KeyValueProperty->SetPropertyType(EDatasmithKeyValuePropertyType::String);
+			MetaDataRefElement->AddProperty(KeyValueProperty);
 		}
+	};
 
-		FString MetaName = TEXT("Instance ");
-		MetaName += Attribute.Key;
-		TSharedRef< IDatasmithKeyValueProperty > KeyValueProperty = FDatasmithSceneFactory::CreateKeyValueProperty(*MetaName);
-
-		KeyValueProperty->SetValue(*Attribute.Value);
-		KeyValueProperty->SetPropertyType(EDatasmithKeyValuePropertyType::String);
-
-		MetaDataRefElement->AddProperty(KeyValueProperty);
-	}
+	AddMetaData(Reference, TEXT("Reference "));
+	AddMetaData(Instance, TEXT("Instance "));
 
 	DatasmithScene->AddMetaData(MetaDataRefElement);
 
 }
 
-bool FDatasmithSceneBaseGraphBuilder::DoesActorHaveChildrenOrIsAStaticMesh(const TSharedPtr< IDatasmithActorElement >& ActorElement)
+bool FDatasmithSceneBaseGraphBuilder::DoesActorHaveChildrenOrIsAStaticMesh(const TSharedPtr<IDatasmithActorElement>& ActorElement)
 {
 	if (ActorElement != nullptr)
 	{
@@ -750,29 +636,28 @@ bool FDatasmithSceneBaseGraphBuilder::DoesActorHaveChildrenOrIsAStaticMesh(const
 		}
 		else if (ActorElement->IsA(EDatasmithElementType::StaticMeshActor))
 		{
-			const TSharedPtr< IDatasmithMeshActorElement >& MeshActorElement = StaticCastSharedPtr< IDatasmithMeshActorElement >(ActorElement);
+			const TSharedPtr<IDatasmithMeshActorElement>& MeshActorElement = StaticCastSharedPtr<IDatasmithMeshActorElement>(ActorElement);
 			return FCString::Strlen(MeshActorElement->GetStaticMeshPathName()) > 0;
 		}
 	}
 	return false;
 }
 
-
-void FDatasmithSceneBaseGraphBuilder::AddChildren(TSharedPtr< IDatasmithActorElement > Actor, const CADLibrary::FArchiveReference& Reference, const ActorData& ParentData)
+void FDatasmithSceneBaseGraphBuilder::AddChildren(TSharedPtr<IDatasmithActorElement> Actor, const CADLibrary::FArchiveReference& Reference, const ActorData& ParentData)
 {
-	for (const int32 ChildId : Reference.Children)
+	for (const FCadId& ChildId : Reference.Children)
 	{
-		if (int32* ChildNodeIndex = SceneGraph->CADIdToInstanceIndex.Find(ChildId))
+		if (SceneGraph->IsAInstance(ChildId))
 		{
-			TSharedPtr< IDatasmithActorElement > ChildActor = BuildInstance(*ChildNodeIndex, ParentData);
+			TSharedPtr<IDatasmithActorElement> ChildActor = BuildInstance(ChildId, ParentData);
 			if (ChildActor.IsValid() && DoesActorHaveChildrenOrIsAStaticMesh(ChildActor))
 			{
 				Actor->AddChild(ChildActor);
 			}
 		}
-		if (int32* ChildNodeIndex = SceneGraph->CADIdToBodyIndex.Find(ChildId))
+		else
 		{
-			TSharedPtr< IDatasmithActorElement > ChildActor = BuildBody(*ChildNodeIndex, ParentData);
+			TSharedPtr<IDatasmithActorElement> ChildActor = BuildBody(ChildId, ParentData);
 			if (ChildActor.IsValid() && DoesActorHaveChildrenOrIsAStaticMesh(ChildActor))
 			{
 				Actor->AddChild(ChildActor);
diff --git a/Engine/Plugins/Enterprise/DatasmithCADImporter/Source/DatasmithCADTranslator/Public/DatasmithSceneGraphBuilder.h b/Engine/Plugins/Enterprise/DatasmithCADImporter/Source/DatasmithCADTranslator/Public/DatasmithSceneGraphBuilder.h
--- a/Engine/Plugins/Enterprise/DatasmithCADImporter/Source/DatasmithCADTranslator/Public/DatasmithSceneGraphBuilder.h
+++ b/Engine/Plugins/Enterprise/DatasmithCADImporter/Source/DatasmithCADTranslator/Public/DatasmithSceneGraphBuilder.h
@@ -12,6 +12,11 @@
 #include "Containers/Queue.h"
 #include "Misc/Paths.h"
 
+namespace CADLibrary
+{
+class FArchiveCADObject;
+}
+
 class FDatasmithSceneSource;
 class IDatasmithActorElement;
 class IDatasmithMeshElement;
@@ -23,27 +28,24 @@ class ActorData  //#ueent_CAD
 public:
 	ActorData(const TCHAR* NodeUuid, const ActorData& ParentData)
 		: Uuid(NodeUuid)
-		, Material(ParentData.Material)
-		, MaterialUuid(ParentData.MaterialUuid)
-		, Color(ParentData.Color)
-		, ColorUuid(ParentData.ColorUuid)
+		, Inheritance(ParentData.Inheritance)
+		, MaterialUId(ParentData.MaterialUId)
+		, ColorUId(ParentData.ColorUId)
 	{
 	}
 
 	ActorData(const TCHAR* NodeUuid)
 		: Uuid(NodeUuid)
-		, MaterialUuid(0)
-		, ColorUuid(0)
+		, Inheritance(CADLibrary::ECADGraphicPropertyInheritance::Unset)
+		, MaterialUId(0)
+		, ColorUId(0)
 	{
 	}
 
 	const TCHAR* Uuid;
-
-	CADLibrary::FCADMaterial Material;
-	FMaterialUId MaterialUuid;
-
-	FColor Color;
-	FMaterialUId ColorUuid;
+	CADLibrary::ECADGraphicPropertyInheritance Inheritance = CADLibrary::ECADGraphicPropertyInheritance::Unset;
+	FMaterialUId MaterialUId;
+	FMaterialUId ColorUId;
 };
 
 
@@ -64,21 +66,19 @@ public:
 	virtual bool Build();
 
 protected:
-	TSharedPtr< IDatasmithActorElement > BuildInstance(int32 InstanceIndex, const ActorData& ParentData);
-	TSharedPtr< IDatasmithActorElement > BuildReference(CADLibrary::FArchiveReference& Component, const ActorData& ParentData);
-	TSharedPtr< IDatasmithActorElement > BuildBody(int32 BodyIndex, const ActorData& ParentData);
+	TSharedPtr<IDatasmithActorElement> BuildInstance(FCadId InstanceId, const ActorData& ParentData);
+	TSharedPtr<IDatasmithActorElement> BuildReference(CADLibrary::FArchiveReference& Reference, const ActorData& ParentData);
+	TSharedPtr<IDatasmithActorElement> BuildBody(FCadId BodyId, const ActorData& ParentData);
 
-	void AddMetaData(TSharedPtr<IDatasmithActorElement> ActorElement, TMap<FString, FString>& InstanceNodeAttributeSetMap, TMap<FString, FString>& ReferenceNodeAttributeSetMap);
-	void AddChildren(TSharedPtr<IDatasmithActorElement> Actor, const CADLibrary::FArchiveReference& Component, const ActorData& ParentData);
+	void AddMetaData(TSharedPtr<IDatasmithActorElement> ActorElement, const CADLibrary::FArchiveCADObject& Instance, const CADLibrary::FArchiveCADObject& Reference);
+	void AddChildren(TSharedPtr<IDatasmithActorElement> Actor, const CADLibrary::FArchiveReference& Reference, const ActorData& ParentData);
 	bool DoesActorHaveChildrenOrIsAStaticMesh(const TSharedPtr< IDatasmithActorElement >& ActorElement);
 
-	TSharedPtr< IDatasmithUEPbrMaterialElement > GetDefaultMaterial();
-	TSharedPtr<IDatasmithMaterialIDElement> FindOrAddMaterial(uint32 MaterialUuid);
-
-	TSharedPtr< IDatasmithActorElement > CreateActor(const TCHAR* ActorUUID, const TCHAR* ActorLabel);
-	virtual TSharedPtr< IDatasmithMeshElement > FindOrAddMeshElement(CADLibrary::FArchiveBody& Body, FString& InLabel);
+	TSharedPtr<IDatasmithUEPbrMaterialElement> GetDefaultMaterial();
+	TSharedPtr<IDatasmithMaterialIDElement> FindOrAddMaterial(FMaterialUId MaterialUuid);
 
-	void GetNodeUUIDAndName(const TMap<FString, FString>& InInstanceNodeMetaDataMap, const TMap<FString, FString>& InReferenceNodeMetaDataMap, int32 InComponentIndex, const TCHAR* InParentUEUUID, FString& OutUEUUID, FString& OutName);
+	TSharedPtr<IDatasmithActorElement> CreateActor(const TCHAR* ActorUUID, const TCHAR* ActorLabel);
+	virtual TSharedPtr<IDatasmithMeshElement> FindOrAddMeshElement(CADLibrary::FArchiveBody& Body);
 
 	CADLibrary::FArchiveSceneGraph* FindSceneGraphArchive(const CADLibrary::FFileDescriptor& File, uint32& FileHash) const;
 
