[Rhino] Rollback to OpenNurbs6

diff --git a/Engine/Plugins/Enterprise/DatasmithCADImporter/Source/DatasmithOpenNurbsTranslator/DatasmithOpenNurbsTranslator.Build.cs b/Engine/Plugins/Enterprise/DatasmithCADImporter/Source/DatasmithOpenNurbsTranslator/DatasmithOpenNurbsTranslator.Build.cs
--- a/Engine/Plugins/Enterprise/DatasmithCADImporter/Source/DatasmithOpenNurbsTranslator/DatasmithOpenNurbsTranslator.Build.cs
+++ b/Engine/Plugins/Enterprise/DatasmithCADImporter/Source/DatasmithOpenNurbsTranslator/DatasmithOpenNurbsTranslator.Build.cs
@@ -35,9 +35,9 @@ namespace UnrealBuildTool.Rules
 				}
 			);
 
-			if (System.Type.GetType("OpenNurbs7") != null)
+			if (System.Type.GetType("OpenNurbs6") != null)
 			{
-				PrivateDependencyModuleNames.Add("OpenNurbs7");
+				PrivateDependencyModuleNames.Add("OpenNurbs6");
 			}
 
             if (Target.Type == TargetType.Editor)
