[CAD Import] Fix IFC orientation

diff --git a/Engine/Plugins/Enterprise/DatasmithCADImporter/Source/CADTools/Private/CADData.cpp b/Engine/Plugins/Enterprise/DatasmithCADImporter/Source/CADTools/Private/CADData.cpp
--- a/Engine/Plugins/Enterprise/DatasmithCADImporter/Source/CADTools/Private/CADData.cpp
+++ b/Engine/Plugins/Enterprise/DatasmithCADImporter/Source/CADTools/Private/CADData.cpp
@@ -266,11 +266,15 @@ ECADFormat FileFormat(const FString& Extension)
 	{
 		return ECADFormat::AUTOCAD;
 	}
+	else if (Extension == TEXT("ifc"))
+	{
+		return ECADFormat::IFC;
+	}
 	else if (Extension == TEXT("dgn"))
 	{
 		return ECADFormat::MICROSTATION;
 	}
-	else if (Extension == TEXT("hsf"))
+	else if (Extension == TEXT("hsf") || Extension == TEXT("prc"))
 	{
 		return ECADFormat::TECHSOFT;
 	}
diff --git a/Engine/Plugins/Enterprise/DatasmithCADImporter/Source/DatasmithCADTranslator/Private/DatasmithCADTranslator.cpp b/Engine/Plugins/Enterprise/DatasmithCADImporter/Source/DatasmithCADTranslator/Private/DatasmithCADTranslator.cpp
--- a/Engine/Plugins/Enterprise/DatasmithCADImporter/Source/DatasmithCADTranslator/Private/DatasmithCADTranslator.cpp
+++ b/Engine/Plugins/Enterprise/DatasmithCADImporter/Source/DatasmithCADTranslator/Private/DatasmithCADTranslator.cpp
@@ -165,6 +165,12 @@ bool FDatasmithCADTranslator::LoadScene(TSharedRef<IDatasmithScene> DatasmithSce
 		break;
 	}
 
+	case CADLibrary::ECADFormat::IFC:
+	{
+		ImportParameters.SetModelCoordinateSystem(FDatasmithUtils::EModelCoordSystem::ZUp_RightHanded_FBXLegacy);
+		break;
+	}
+
 	default:
 		break;
 	}
