Remove unused Import Parameters

diff --git a/Engine/Plugins/Enterprise/DatasmithCADImporter/Source/CADInterfaces/Private/TechSoftFileParserCADKernelTessellator.cpp b/Engine/Plugins/Enterprise/DatasmithCADImporter/Source/CADInterfaces/Private/TechSoftFileParserCADKernelTessellator.cpp
--- a/Engine/Plugins/Enterprise/DatasmithCADImporter/Source/CADInterfaces/Private/TechSoftFileParserCADKernelTessellator.cpp
+++ b/Engine/Plugins/Enterprise/DatasmithCADImporter/Source/CADInterfaces/Private/TechSoftFileParserCADKernelTessellator.cpp
@@ -72,7 +72,7 @@ void FTechSoftFileParserCADKernelTessellator::SewAndGenerateBodyMeshes()
 
 void FTechSoftFileParserCADKernelTessellator::SewAndMesh(TArray<A3DRiRepresentationItem*>& Representations)
 {
-	double GeometricTolerance = CADFileData.GetImportParameters().ConvertMMToImportUnit(0.01); // mm
+	const double GeometricTolerance = 0.01; // mm
 
 	CADKernel::FSession CADKernelSession(GeometricTolerance);
 	CADKernelSession.SetFirstNewHostId(LastHostIdUsed);
@@ -185,7 +185,7 @@ void FTechSoftFileParserCADKernelTessellator::SewAndMesh(TArray<A3DRiRepresentat
 
 void FTechSoftFileParserCADKernelTessellator::GenerateBodyMesh(A3DRiRepresentationItem* Representation, FArchiveBody& ArchiveBody)
 {
-	double GeometricTolerance = CADFileData.GetImportParameters().ConvertMMToImportUnit(0.01); // mm
+	const double GeometricTolerance = 0.01; // mm
 
 	CADKernel::FSession CADKernelSession(GeometricTolerance);
 	CADKernelSession.SetFirstNewHostId(LastHostIdUsed);
@@ -198,7 +198,7 @@ void FTechSoftFileParserCADKernelTessellator::GenerateBodyMesh(A3DRiRepresentati
 
 	if (CADFileData.GetImportParameters().GetStitchingTechnique() == StitchingHeal)
 	{
-		double Tolerance = CADFileData.GetImportParameters().ConvertMMToImportUnit(0.1); // mm
+		const double Tolerance = 0.1; // mm
 		CADKernel::FTopomaker Topomaker(CADKernelSession, Tolerance);
 		Topomaker.Sew();
 		Topomaker.OrientShells();
diff --git a/Engine/Plugins/Enterprise/DatasmithCADImporter/Source/CADKernelSurface/Private/CADKernelSurfaceExtension.cpp b/Engine/Plugins/Enterprise/DatasmithCADImporter/Source/CADKernelSurface/Private/CADKernelSurfaceExtension.cpp
--- a/Engine/Plugins/Enterprise/DatasmithCADImporter/Source/CADKernelSurface/Private/CADKernelSurfaceExtension.cpp
+++ b/Engine/Plugins/Enterprise/DatasmithCADImporter/Source/CADKernelSurface/Private/CADKernelSurfaceExtension.cpp
@@ -34,7 +34,7 @@ bool UCADKernelParametricSurfaceData::Tessellate(UStaticMesh& StaticMesh, const
 	bool bSuccessfulTessellation = false;
 
 #if WITH_EDITOR
-	CADLibrary::FImportParameters ImportParameters((double) SceneParameters.MetricUnit, (double) SceneParameters.ScaleFactor);
+	CADLibrary::FImportParameters ImportParameters;
 	ImportParameters.SetModelCoordinateSystem((FDatasmithUtils::EModelCoordSystem) SceneParameters.ModelCoordSys);
 	ImportParameters.SetTesselationParameters(RetessellateOptions.ChordTolerance, RetessellateOptions.MaxEdgeLength, RetessellateOptions.NormalTolerance, (CADLibrary::EStitchingTechnique) RetessellateOptions.StitchingTechnique);
 
@@ -56,7 +56,8 @@ bool UCADKernelParametricSurfaceData::Tessellate(UStaticMesh& StaticMesh, const
 			CADLibrary::CopyPatchGroups(*DestinationMeshDescription, MeshDescription);
 		}
 
-		TSharedRef<CADKernel::FSession> CADKernelSession = MakeShared<CADKernel::FSession>(0.00001 / ImportParameters.GetMetricUnit());
+		const double GeometricTolerance = 0.01; // mm
+		TSharedRef<CADKernel::FSession> CADKernelSession = MakeShared<CADKernel::FSession>(GeometricTolerance);
 		CADKernelSession->AddDatabase(RawData);
 
 		CADKernel::FModel& CADKernelModel = CADKernelSession->GetModel();
diff --git a/Engine/Plugins/Enterprise/DatasmithCADImporter/Source/CADKernelSurface/Public/CADModelToCADKernelConverterBase.h b/Engine/Plugins/Enterprise/DatasmithCADImporter/Source/CADKernelSurface/Public/CADModelToCADKernelConverterBase.h
--- a/Engine/Plugins/Enterprise/DatasmithCADImporter/Source/CADKernelSurface/Public/CADModelToCADKernelConverterBase.h
+++ b/Engine/Plugins/Enterprise/DatasmithCADImporter/Source/CADKernelSurface/Public/CADModelToCADKernelConverterBase.h
@@ -22,14 +22,14 @@ class FCADModelToCADKernelConverterBase : public CADLibrary::ICADModelConverter
 public:
 
 	FCADModelToCADKernelConverterBase(CADLibrary::FImportParameters InImportParameters)
-		: CADKernelSession(0.00001 / InImportParameters.GetMetricUnit())
+		: CADKernelSession(0.01)
 		, ImportParameters(InImportParameters)
-		, GeometricTolerance(0.00001 / InImportParameters.GetMetricUnit())
+		, GeometricTolerance(0.01)
 		, SquareTolerance(GeometricTolerance* GeometricTolerance)
 	{
 	}
 
-	virtual void InitializeProcess(double InMetricUnit) override
+	virtual void InitializeProcess() override
 	{
 		CADKernelSession.Clear();
 	}
@@ -40,7 +40,7 @@ public:
 		if(ImportParameters.GetStitchingTechnique() != CADLibrary::StitchingNone)
 		{
 			// the joining tolerance is set to 0.1 mm until the user can specify it
-			double JoiningTolerance = ImportParameters.ConvertMMToImportUnit(0.1);
+			const double JoiningTolerance = 0.1;
 			CADKernel::FTopomaker Topomaker(CADKernelSession, JoiningTolerance);
 			Topomaker.Sew();
 			Topomaker.OrientShells();
@@ -68,21 +68,6 @@ public:
 		ImportParameters.SetTesselationParameters(ChordTolerance, MaxEdgeLength, NormalTolerance, StitchingTechnique);
 	}
 
-	virtual void SetMetricUnit(double NewMetricUnit) override
-	{
-		ImportParameters.SetMetricUnit(NewMetricUnit);
-	}
-
-	virtual double GetScaleFactor() const override
-	{
-		return ImportParameters.GetScaleFactor();
-	}
-
-	virtual double GetMetricUnit() const override
-	{
-		return ImportParameters.GetMetricUnit();
-	}
-
 	virtual bool IsSessionValid() override
 	{
 		return true;
diff --git a/Engine/Plugins/Enterprise/DatasmithCADImporter/Source/CADLibrary/Private/CADKernelTools.cpp b/Engine/Plugins/Enterprise/DatasmithCADImporter/Source/CADLibrary/Private/CADKernelTools.cpp
--- a/Engine/Plugins/Enterprise/DatasmithCADImporter/Source/CADLibrary/Private/CADKernelTools.cpp
+++ b/Engine/Plugins/Enterprise/DatasmithCADImporter/Source/CADLibrary/Private/CADKernelTools.cpp
@@ -407,13 +407,13 @@ namespace CADLibrary
 
 		if (ImportParameters.GetMaxEdgeLength() > SMALL_NUMBER)
 		{
-			TSharedPtr<CADKernel::FCriterion> MaxSizeCriterion = CADKernel::FCriterion::CreateCriterion(CADKernel::ECriterion::MaxSize, ImportParameters.GetMaxEdgeLength() / ImportParameters.GetScaleFactor());
+			TSharedPtr<CADKernel::FCriterion> MaxSizeCriterion = CADKernel::FCriterion::CreateCriterion(CADKernel::ECriterion::MaxSize, ImportParameters.GetMaxEdgeLength());
 			MeshModel.AddCriterion(MaxSizeCriterion);
 		}
 
 		if (ImportParameters.GetChordTolerance() > SMALL_NUMBER)
 		{
-			TSharedPtr<CADKernel::FCriterion> ChordCriterion = CADKernel::FCriterion::CreateCriterion(CADKernel::ECriterion::Sag, ImportParameters.GetChordTolerance() / ImportParameters.GetScaleFactor());
+			TSharedPtr<CADKernel::FCriterion> ChordCriterion = CADKernel::FCriterion::CreateCriterion(CADKernel::ECriterion::Sag, ImportParameters.GetChordTolerance());
 			MeshModel.AddCriterion(ChordCriterion);
 		}
 
diff --git a/Engine/Plugins/Enterprise/DatasmithCADImporter/Source/CADLibrary/Private/MeshDescriptionHelper.cpp b/Engine/Plugins/Enterprise/DatasmithCADImporter/Source/CADLibrary/Private/MeshDescriptionHelper.cpp
--- a/Engine/Plugins/Enterprise/DatasmithCADImporter/Source/CADLibrary/Private/MeshDescriptionHelper.cpp
+++ b/Engine/Plugins/Enterprise/DatasmithCADImporter/Source/CADLibrary/Private/MeshDescriptionHelper.cpp
@@ -80,7 +80,7 @@ namespace CADLibrary
 		{
 			for (const FVector& Vertex : CTTessellation.PositionArray)
 			{
-				VertexDataSet.Emplace(GlobalVertexCount, Vertex * ImportParams.GetScaleFactor());
+				VertexDataSet.Emplace(GlobalVertexCount, Vertex);
 				++GlobalVertexCount;
 			}
 		}
@@ -263,12 +263,6 @@ namespace CADLibrary
 		TVertexAttributesRef<FVector3f> VertexPositions = MeshDescription.GetVertexPositions();
 
 		TArray<FVector>& VertexArray = Body.VertexArray;
-
-		for (FVector& Vertex : VertexArray)
-		{
-			Vertex *= ImportParams.GetScaleFactor();
-		}
-
 		TArray<int32>& VertexIdSet = Body.VertexIds;
 		int32 VertexCount = VertexArray.Num();
 		VertexIdSet.SetNumZeroed(VertexCount);
diff --git a/Engine/Plugins/Enterprise/DatasmithCADImporter/Source/CADTools/Private/CADOption.cpp b/Engine/Plugins/Enterprise/DatasmithCADImporter/Source/CADTools/Private/CADOption.cpp
--- a/Engine/Plugins/Enterprise/DatasmithCADImporter/Source/CADTools/Private/CADOption.cpp
+++ b/Engine/Plugins/Enterprise/DatasmithCADImporter/Source/CADTools/Private/CADOption.cpp
@@ -76,8 +76,6 @@ uint32 GetTypeHash(const FImportParameters& ImportParameters)
 	ParametersHash = HashCombine(ParametersHash, ::GetTypeHash(ImportParameters.ChordTolerance));
 	ParametersHash = HashCombine(ParametersHash, ::GetTypeHash(ImportParameters.MaxEdgeLength));
 	ParametersHash = HashCombine(ParametersHash, ::GetTypeHash(ImportParameters.MaxNormalAngle));
-	ParametersHash = HashCombine(ParametersHash, ::GetTypeHash(ImportParameters.MetricUnit));
-	ParametersHash = HashCombine(ParametersHash, ::GetTypeHash(ImportParameters.ScaleFactor));
 	ParametersHash = HashCombine(ParametersHash, ::GetTypeHash(ImportParameters.StitchingTechnique));
 	return ParametersHash;
 }
diff --git a/Engine/Plugins/Enterprise/DatasmithCADImporter/Source/CADTools/Public/CADOptions.h b/Engine/Plugins/Enterprise/DatasmithCADImporter/Source/CADTools/Public/CADOptions.h
--- a/Engine/Plugins/Enterprise/DatasmithCADImporter/Source/CADTools/Public/CADOptions.h
+++ b/Engine/Plugins/Enterprise/DatasmithCADImporter/Source/CADTools/Public/CADOptions.h
@@ -19,34 +19,14 @@ namespace CADLibrary
 		StitchingSew,
 	};
 
-	enum class EDisplayPreference : uint8
-	{
-		ColorPrefered,
-		MaterialPrefered,
-		ColorOnly,
-		MaterialOnly,
-	};
-
-	enum class EDisplayDataPropagationMode : uint8
-	{
-		TopDown,
-		BottomUp,
-		BodyOnly,
-	};
-
 	class FImportParameters
 	{
 	private:
-		double MetricUnit;
-		double ScaleFactor;
 		double ChordTolerance;
 		double MaxEdgeLength;
 		double MaxNormalAngle;
 		EStitchingTechnique StitchingTechnique;
 		FDatasmithUtils::EModelCoordSystem ModelCoordSys;
-		EDisplayDataPropagationMode Propagation;
-		EDisplayPreference DisplayPreference;
-		bool bScaleUVMap;
 		
 	public:
 		CADTOOLS_API static bool bGDisableCADKernelTessellation;
@@ -57,28 +37,15 @@ namespace CADLibrary
 		CADTOOLS_API static float GStitchingTolerance;
 
 	public:
-		FImportParameters(double InMetricUnit = 0.001, double InScaleFactor = 1., FDatasmithUtils::EModelCoordSystem NewCoordinateSystem = FDatasmithUtils::EModelCoordSystem::ZUp_RightHanded)
-			: MetricUnit(InMetricUnit)
-			, ScaleFactor(InScaleFactor)
-			, ChordTolerance(0.2)
+		FImportParameters(FDatasmithUtils::EModelCoordSystem NewCoordinateSystem = FDatasmithUtils::EModelCoordSystem::ZUp_RightHanded)
+			: ChordTolerance(0.2)
 			, MaxEdgeLength(0.0)
 			, MaxNormalAngle(20.0)
 			, StitchingTechnique(EStitchingTechnique::StitchingNone)
 			, ModelCoordSys(NewCoordinateSystem)
-			, Propagation(EDisplayDataPropagationMode::TopDown)
-			, DisplayPreference(EDisplayPreference::MaterialPrefered)
-			, bScaleUVMap(true)
 		{
 		}
 	
-		/**
-		 * @param ValueInMM, the value in millimeter to convert into the scene metric unit 
-		 */
-		double ConvertMMToImportUnit(double ValueInMM) const
-		{
-			return ValueInMM * 0.001 / MetricUnit;
-		}
-
 		void SetTesselationParameters(double InChordTolerance, double InMaxEdgeLength, double InMaxNormalAngle, CADLibrary::EStitchingTechnique InStitchingTechnique)
 		{
 			ChordTolerance = InChordTolerance;
@@ -87,41 +54,30 @@ namespace CADLibrary
 			StitchingTechnique = InStitchingTechnique;
 		}
 
-		void SetMetricUnit(double NewMetricUnit)
-		{
-			MetricUnit = NewMetricUnit;
-			ScaleFactor = NewMetricUnit / 0.01;
-		}
-
 		uint32 GetHash() const
 		{
 			uint32 Hash = 0;
-			for (double Param : {MetricUnit, ScaleFactor, ChordTolerance, MaxEdgeLength, MaxNormalAngle})
+			for (double Param : {ChordTolerance, MaxEdgeLength, MaxNormalAngle})
 			{
 				Hash = HashCombine(Hash, ::GetTypeHash(Param));
 			}
-			for (uint32 Param : {uint32(StitchingTechnique), uint32(ModelCoordSys), uint32(Propagation), uint32(DisplayPreference)})
+			for (uint32 Param : {uint32(StitchingTechnique), uint32(ModelCoordSys)})
 			{
 				Hash = HashCombine(Hash, ::GetTypeHash(Param));
 			}
-			Hash = HashCombine(Hash, ::GetTypeHash(bScaleUVMap));
 			Hash = HashCombine(Hash, ::GetTypeHash(bGPreferJtFileEmbeddedTessellation));
 			return Hash;
 		}
 
 		friend FArchive& operator<<(FArchive& Ar, FImportParameters& ImportParameters)
 		{
-			Ar << ImportParameters.MetricUnit;
-			Ar << ImportParameters.ScaleFactor;
 			Ar << ImportParameters.ChordTolerance;
 			Ar << ImportParameters.MaxEdgeLength;
 			Ar << ImportParameters.MaxNormalAngle;
 			Ar << (uint32&) ImportParameters.StitchingTechnique;
 			Ar << (uint8&) ImportParameters.ModelCoordSys;
-			Ar << (uint8&) ImportParameters.Propagation;
-			Ar << (uint8&) ImportParameters.DisplayPreference;
-			Ar << ImportParameters.bScaleUVMap;
-			// these 5 static variables have to be serialized to be transmitted to CADWorkers
+
+			// these static variables have to be serialized to be transmitted to CADWorkers
 			// because CADWorker has not access to CVars
 			Ar << ImportParameters.bGOverwriteCache;
 			Ar << ImportParameters.bGDisableCADKernelTessellation;
@@ -132,16 +88,6 @@ namespace CADLibrary
 			return Ar;
 		}
 
-		double GetMetricUnit() const
-		{
-			return MetricUnit;
-		}
-
-		double GetScaleFactor() const
-		{
-			return ScaleFactor;
-		}
-
 		double GetChordTolerance() const
 		{
 			return ChordTolerance;
@@ -167,43 +113,12 @@ namespace CADLibrary
 			return ModelCoordSys;
 		}
 
-		EDisplayDataPropagationMode GetPropagation() const
-		{
-			return Propagation;
-		}
-
-		EDisplayPreference GetDisplayPreference() const
-		{
-			return DisplayPreference;
-		}
-
-		bool NeedScaleUVMap() const
-		{	
-			return bScaleUVMap;
-		}
-
-		void SwitchOffUVMapScaling()
-		{
-			bScaleUVMap = false;
-		}
-
 		void SetModelCoordinateSystem(FDatasmithUtils::EModelCoordSystem NewCoordinateSystem)
 		{
 			ModelCoordSys = NewCoordinateSystem;
 		}
 
-		void SetDisplayPreference(CADLibrary::EDisplayPreference InDisplayPreference)
-		{
-			DisplayPreference = InDisplayPreference;
-		}
-
-		void SetPropagationMode(CADLibrary::EDisplayDataPropagationMode InMode)
-		{
-			Propagation = InMode;
-		}
-
 		CADTOOLS_API friend uint32 GetTypeHash(const FImportParameters& ImportParameters);
-
 	};
 
 	inline FString BuildCacheFilePath(const TCHAR* CachePath, const TCHAR* Folder, uint32 BodyHash)
diff --git a/Engine/Plugins/Enterprise/DatasmithCADImporter/Source/DatasmithCADTranslator/Private/DatasmithCADTranslator.cpp b/Engine/Plugins/Enterprise/DatasmithCADImporter/Source/DatasmithCADTranslator/Private/DatasmithCADTranslator.cpp
--- a/Engine/Plugins/Enterprise/DatasmithCADImporter/Source/DatasmithCADTranslator/Private/DatasmithCADTranslator.cpp
+++ b/Engine/Plugins/Enterprise/DatasmithCADImporter/Source/DatasmithCADTranslator/Private/DatasmithCADTranslator.cpp
@@ -137,15 +137,12 @@ bool FDatasmithCADTranslator::LoadScene(TSharedRef<IDatasmithScene> DatasmithSce
 	{
 	case CADLibrary::ECADFormat::NX:
 	{
-		ImportParameters.SetDisplayPreference(CADLibrary::EDisplayPreference::ColorOnly);
-		ImportParameters.SetPropagationMode(CADLibrary::EDisplayDataPropagationMode::BodyOnly);
 		break;
 	}
 
 	case CADLibrary::ECADFormat::SOLIDWORKS:
 	{
 		ImportParameters.SetModelCoordinateSystem(FDatasmithUtils::EModelCoordSystem::YUp_RightHanded);
-		ImportParameters.SetDisplayPreference(CADLibrary::EDisplayPreference::ColorOnly);
 		break;
 	}
 
@@ -153,15 +150,11 @@ bool FDatasmithCADTranslator::LoadScene(TSharedRef<IDatasmithScene> DatasmithSce
 	case CADLibrary::ECADFormat::CREO:
 	{
 		ImportParameters.SetModelCoordinateSystem(FDatasmithUtils::EModelCoordSystem::YUp_RightHanded);
-		ImportParameters.SetDisplayPreference(CADLibrary::EDisplayPreference::ColorOnly);
-		ImportParameters.SetPropagationMode(CADLibrary::EDisplayDataPropagationMode::BodyOnly);
 		break;
 	}
 
 	case CADLibrary::ECADFormat::DWG:
 	{
-		ImportParameters.SetDisplayPreference(CADLibrary::EDisplayPreference::ColorOnly);
-		ImportParameters.SetPropagationMode(CADLibrary::EDisplayDataPropagationMode::BodyOnly);
 		break;
 	}
 
diff --git a/Engine/Plugins/Enterprise/DatasmithCADImporter/Source/DatasmithCADTranslator/Private/DatasmithSceneGraphBuilder.cpp b/Engine/Plugins/Enterprise/DatasmithCADImporter/Source/DatasmithCADTranslator/Private/DatasmithSceneGraphBuilder.cpp
--- a/Engine/Plugins/Enterprise/DatasmithCADImporter/Source/DatasmithCADTranslator/Private/DatasmithSceneGraphBuilder.cpp
+++ b/Engine/Plugins/Enterprise/DatasmithCADImporter/Source/DatasmithCADTranslator/Private/DatasmithSceneGraphBuilder.cpp
@@ -82,7 +82,7 @@ void AddTransformToActor(const CADLibrary::FArchiveCADObject& Object, TSharedPtr
 	FTransform LocalTransform(Object.TransformMatrix);
 	FTransform LocalUETransform = FDatasmithUtils::ConvertTransform(ImportParameters.GetModelCoordSys(), LocalTransform);
 
-	Actor->SetTranslation(LocalUETransform.GetTranslation() * ImportParameters.GetScaleFactor());
+	Actor->SetTranslation(LocalUETransform.GetTranslation());
 	Actor->SetScale(LocalUETransform.GetScale3D());
 	Actor->SetRotation(LocalUETransform.GetRotation());
 }
@@ -204,8 +204,6 @@ FDatasmithSceneBaseGraphBuilder::FDatasmithSceneBaseGraphBuilder(CADLibrary::FAr
 	, ImportParameters(InImportParameters)
 	, ImportParametersHash(ImportParameters.GetHash())
 	, RootFileDescription(*InSource.GetSourceFile())
-	, bPreferMaterial(false)
-	, bMaterialPropagationIsTopDown(ImportParameters.GetPropagation() == CADLibrary::EDisplayDataPropagationMode::TopDown)
 {
 	if (InSceneGraph)
 	{
diff --git a/Engine/Plugins/Enterprise/DatasmithCADImporter/Source/DatasmithCADTranslator/Public/DatasmithSceneGraphBuilder.h b/Engine/Plugins/Enterprise/DatasmithCADImporter/Source/DatasmithCADTranslator/Public/DatasmithSceneGraphBuilder.h
--- a/Engine/Plugins/Enterprise/DatasmithCADImporter/Source/DatasmithCADTranslator/Public/DatasmithSceneGraphBuilder.h
+++ b/Engine/Plugins/Enterprise/DatasmithCADImporter/Source/DatasmithCADTranslator/Public/DatasmithSceneGraphBuilder.h
@@ -102,9 +102,6 @@ protected:
 	TMap<FMaterialUId, CADLibrary::FArchiveMaterial> MaterialUIdToMaterialArchive; 
 
 	TArray<uint32> AncestorSceneGraphHash;
-
-	bool bPreferMaterial;
-	bool bMaterialPropagationIsTopDown;
 };
 
 class DATASMITHCADTRANSLATOR_API FDatasmithSceneGraphBuilder : public FDatasmithSceneBaseGraphBuilder
diff --git a/Engine/Plugins/Enterprise/DatasmithCADImporter/Source/DatasmithOpenNurbsTranslator/Private/DatasmithOpenNurbsTranslator.cpp b/Engine/Plugins/Enterprise/DatasmithCADImporter/Source/DatasmithOpenNurbsTranslator/Private/DatasmithOpenNurbsTranslator.cpp
--- a/Engine/Plugins/Enterprise/DatasmithCADImporter/Source/DatasmithOpenNurbsTranslator/Private/DatasmithOpenNurbsTranslator.cpp
+++ b/Engine/Plugins/Enterprise/DatasmithCADImporter/Source/DatasmithOpenNurbsTranslator/Private/DatasmithOpenNurbsTranslator.cpp
@@ -644,9 +644,7 @@ public:
 			TranslationCache = MakeShared<FTranslationCache>();
 		}
 
-		const double IPMetricUnit = 0.001; // CAD Tools works in mm
-		const double IPScaleFactor = 1.;
-		CADLibrary::FImportParameters ImportParameters(IPMetricUnit, IPScaleFactor, FDatasmithUtils::EModelCoordSystem::ZUp_RightHanded_FBXLegacy);
+		CADLibrary::FImportParameters ImportParameters(FDatasmithUtils::EModelCoordSystem::ZUp_RightHanded_FBXLegacy);
 
 		if (CADLibrary::FImportParameters::bGDisableCADKernelTessellation)
 		{
@@ -2982,7 +2980,7 @@ bool FOpenNurbsTranslatorImpl::TranslateBRep(ON_Brep* Brep, const ON_3dmObjectAt
 		// Ref. visitBRep
 		CADModelConverter->SetImportParameters(OpenNurbsOptions.ChordTolerance, OpenNurbsOptions.MaxEdgeLength, OpenNurbsOptions.NormalTolerance, (CADLibrary::EStitchingTechnique)OpenNurbsOptions.StitchingTechnique);
 
-		CADModelConverter->InitializeProcess(CADModelConverter->GetMetricUnit());
+		CADModelConverter->InitializeProcess();
 
 		OpenNurbsBRepConverter->AddBRep(*Brep, Offset);
 		CADModelConverter->RepairTopology();
diff --git a/Engine/Plugins/Enterprise/DatasmithCADImporter/Source/DatasmithWireTranslator/Private/DatasmithWireTranslator.cpp b/Engine/Plugins/Enterprise/DatasmithCADImporter/Source/DatasmithWireTranslator/Private/DatasmithWireTranslator.cpp
--- a/Engine/Plugins/Enterprise/DatasmithCADImporter/Source/DatasmithWireTranslator/Private/DatasmithWireTranslator.cpp
+++ b/Engine/Plugins/Enterprise/DatasmithCADImporter/Source/DatasmithWireTranslator/Private/DatasmithWireTranslator.cpp
@@ -305,9 +305,7 @@ public:
 			DatasmithScene->SetProductVersion(TEXT("Alias 2022.2"));
 		}
 
-		const double MetricUnit = 0.001; // ImportParameters MetricUnit is defined in meter
-		const double ScaleFactor = 1;    // MetricUnit to Modeler unit i.e. mm
-		CADLibrary::FImportParameters ImportParameters(MetricUnit, ScaleFactor);
+		CADLibrary::FImportParameters ImportParameters;
 		if(CADLibrary::FImportParameters::bGDisableCADKernelTessellation)
 		{
 			TSharedRef<FAliasModelToTechSoftConverter> AliasToTechSoftConverter = MakeShared<FAliasModelToTechSoftConverter>(ImportParameters);
@@ -1918,7 +1916,7 @@ void FWireTranslatorImpl::DagForLeavesNoMerge(const TSharedPtr<AlDagNode>& DagNo
 TOptional<FMeshDescription> FWireTranslatorImpl::MeshDagNodeWithExternalMesher(AlDagNode& DagNode, TSharedRef<IDatasmithMeshElement> MeshElement, CADLibrary::FMeshParameters& MeshParameters)
 {
 	// Wire unit is cm
-	CADModelConverter->InitializeProcess(0.01);
+	CADModelConverter->InitializeProcess();
 
 	EAliasObjectReference ObjectReference = EAliasObjectReference::LocalReference;
 
@@ -1962,7 +1960,7 @@ TOptional<FMeshDescription> FWireTranslatorImpl::MeshDagNodeWithExternalMesher(A
 TOptional<FMeshDescription> FWireTranslatorImpl::MeshDagNodeWithExternalMesher(TSharedRef<BodyData> Body, TSharedRef<IDatasmithMeshElement> MeshElement, CADLibrary::FMeshParameters& MeshParameters)
 {
 	// Wire unit is cm
-	CADModelConverter->InitializeProcess(0.01);
+	CADModelConverter->InitializeProcess();
 
 	EAliasObjectReference ObjectReference = EAliasObjectReference::LocalReference;
 	if (MeshParameters.bIsSymmetric)
diff --git a/Engine/Plugins/Enterprise/DatasmithCADImporter/Source/ParametricSurface/Private/ParametricSurfaceData.cpp b/Engine/Plugins/Enterprise/DatasmithCADImporter/Source/ParametricSurface/Private/ParametricSurfaceData.cpp
--- a/Engine/Plugins/Enterprise/DatasmithCADImporter/Source/ParametricSurface/Private/ParametricSurfaceData.cpp
+++ b/Engine/Plugins/Enterprise/DatasmithCADImporter/Source/ParametricSurface/Private/ParametricSurfaceData.cpp
@@ -26,8 +26,6 @@ bool UParametricSurfaceData::SetFile(const TCHAR* FilePath)
 void UParametricSurfaceData::SetImportParameters(const CADLibrary::FImportParameters& InSceneParameters)
 {
 	SceneParameters.ModelCoordSys = uint8(InSceneParameters.GetModelCoordSys());
-	SceneParameters.MetricUnit = InSceneParameters.GetMetricUnit();
-	SceneParameters.ScaleFactor = InSceneParameters.GetScaleFactor();
 }
 
 void UParametricSurfaceData::SetMeshParameters(const CADLibrary::FMeshParameters& InMeshParameters)
diff --git a/Engine/Plugins/Enterprise/DatasmithCADImporter/Source/ParametricSurface/Private/TechSoft/CADModelToTechSoftConverterBase.cpp b/Engine/Plugins/Enterprise/DatasmithCADImporter/Source/ParametricSurface/Private/TechSoft/CADModelToTechSoftConverterBase.cpp
--- a/Engine/Plugins/Enterprise/DatasmithCADImporter/Source/ParametricSurface/Private/TechSoft/CADModelToTechSoftConverterBase.cpp
+++ b/Engine/Plugins/Enterprise/DatasmithCADImporter/Source/ParametricSurface/Private/TechSoft/CADModelToTechSoftConverterBase.cpp
@@ -24,7 +24,7 @@ bool FCADModelToTechSoftConverterBase::RepairTopology()
 		A3DRiBrepModel** OutNewBReps;
 		uint32 OutNewBRepCount;
 		const double SewTolerance = CADLibrary::FImportParameters::GStitchingTolerance;
-		const double FileUnit = ImportParameters.GetMetricUnit() * 1000; // ImportParameters MetricUnit is defined in meter, CADLibrary::TechSoftInterface::SewBReps expectes it in mm
+		const double FileUnit = 1.; // BReps from Rhino and Alias are converted in mm
 		CADLibrary::TechSoftInterface::SewBReps((A3DRiBrepModel**) RiRepresentationItems.GetData(), RiRepresentationItems.Num(), SewTolerance, FileUnit, SewOptionsData.GetPtr(), &OutNewBReps, OutNewBRepCount);
 
 		RiRepresentationItems.Empty(OutNewBRepCount);
@@ -37,7 +37,7 @@ bool FCADModelToTechSoftConverterBase::RepairTopology()
 	return true;
 }
 
-void FCADModelToTechSoftConverterBase::InitializeProcess(double InMetricUnit)
+void FCADModelToTechSoftConverterBase::InitializeProcess()
 {
 	RiRepresentationItems.Empty();
 	ModelFile.Reset();
diff --git a/Engine/Plugins/Enterprise/DatasmithCADImporter/Source/ParametricSurface/Private/TechSoft/TechSoftParametricSurface.cpp b/Engine/Plugins/Enterprise/DatasmithCADImporter/Source/ParametricSurface/Private/TechSoft/TechSoftParametricSurface.cpp
--- a/Engine/Plugins/Enterprise/DatasmithCADImporter/Source/ParametricSurface/Private/TechSoft/TechSoftParametricSurface.cpp
+++ b/Engine/Plugins/Enterprise/DatasmithCADImporter/Source/ParametricSurface/Private/TechSoft/TechSoftParametricSurface.cpp
@@ -44,7 +44,7 @@ bool UTechSoftParametricSurfaceData::Tessellate(UStaticMesh& StaticMesh, const F
 			CADLibrary::CopyPatchGroups(*DestinationMeshDescription, MeshDescription);
 		}
 
-		CADLibrary::FImportParameters ImportParameters(SceneParameters.MetricUnit, SceneParameters.ScaleFactor, (FDatasmithUtils::EModelCoordSystem)SceneParameters.ModelCoordSys);
+		CADLibrary::FImportParameters ImportParameters((FDatasmithUtils::EModelCoordSystem)SceneParameters.ModelCoordSys);
 		ImportParameters.SetTesselationParameters(RetessellateOptions.ChordTolerance, RetessellateOptions.MaxEdgeLength, RetessellateOptions.NormalTolerance, (CADLibrary::EStitchingTechnique)RetessellateOptions.StitchingTechnique);
 
 		CADLibrary::FTechSoftInterface& TechSoftInterface = CADLibrary::FTechSoftInterface::Get();
diff --git a/Engine/Plugins/Enterprise/DatasmithCADImporter/Source/ParametricSurface/Public/CADModelConverter.h b/Engine/Plugins/Enterprise/DatasmithCADImporter/Source/ParametricSurface/Public/CADModelConverter.h
--- a/Engine/Plugins/Enterprise/DatasmithCADImporter/Source/ParametricSurface/Public/CADModelConverter.h
+++ b/Engine/Plugins/Enterprise/DatasmithCADImporter/Source/ParametricSurface/Public/CADModelConverter.h
@@ -37,11 +37,7 @@ class ICADModelConverter
 public:
 	virtual ~ICADModelConverter() = default;
 
-	/**	 
-	 * @param FileMetricUnit: Handle input file unit number of meters per file unit.
-	 * eg. For a file in inches, arg should be 0.0254
-	 */
-	virtual void InitializeProcess(double InMetricUnit) = 0;
+	virtual void InitializeProcess() = 0;
 
 	virtual bool RepairTopology() = 0;
 	virtual bool SaveModel(const TCHAR* OutputPath, TSharedRef<IDatasmithMeshElement>& MeshElement) = 0;
@@ -58,10 +54,6 @@ public:
 	 */
 	virtual void SetImportParameters(double ChordTolerance, double MaxEdgeLength, double NormalTolerance, CADLibrary::EStitchingTechnique StitchingTechnique) = 0;
 
-	virtual void SetMetricUnit(double NewMetricUnit) = 0;
-	virtual double GetScaleFactor() const = 0;
-	virtual double GetMetricUnit() const = 0;
-
 	virtual bool IsSessionValid() = 0;
 
 	virtual void AddSurfaceDataForMesh(const TCHAR* InFilePath, const FMeshParameters& InMeshParameters, const FDatasmithTessellationOptions& InTessellationOptions, FDatasmithMeshElementPayload& OutMeshPayload) const = 0;
diff --git a/Engine/Plugins/Enterprise/DatasmithCADImporter/Source/ParametricSurface/Public/CADModelToTechSoftConverterBase.h b/Engine/Plugins/Enterprise/DatasmithCADImporter/Source/ParametricSurface/Public/CADModelToTechSoftConverterBase.h
--- a/Engine/Plugins/Enterprise/DatasmithCADImporter/Source/ParametricSurface/Public/CADModelToTechSoftConverterBase.h
+++ b/Engine/Plugins/Enterprise/DatasmithCADImporter/Source/ParametricSurface/Public/CADModelToTechSoftConverterBase.h
@@ -29,7 +29,7 @@ public:
 	{
 	}
 
-	virtual void InitializeProcess(double InMetricUnit) override;
+	virtual void InitializeProcess() override;
 
 	virtual bool RepairTopology() override;
 
@@ -42,21 +42,6 @@ public:
 		ImportParameters.SetTesselationParameters(ChordTolerance, MaxEdgeLength, NormalTolerance, StitchingTechnique);
 	}
 
-	virtual void SetMetricUnit(double NewMetricUnit) override
-	{
-		ImportParameters.SetMetricUnit(NewMetricUnit);
-	}
-
-	virtual double GetScaleFactor() const override
-	{
-		return ImportParameters.GetScaleFactor();
-	}
-
-	virtual double GetMetricUnit() const override
-	{
-		return ImportParameters.GetMetricUnit();
-	}
-
 	virtual bool IsSessionValid() override
 	{
 		return true;
