Improve  CADImport naming

diff --git a/Engine/Plugins/Enterprise/DatasmithCADImporter/Source/CADInterfaces/Private/TechSoftFileParser.cpp b/Engine/Plugins/Enterprise/DatasmithCADImporter/Source/CADInterfaces/Private/TechSoftFileParser.cpp
--- a/Engine/Plugins/Enterprise/DatasmithCADImporter/Source/CADInterfaces/Private/TechSoftFileParser.cpp
+++ b/Engine/Plugins/Enterprise/DatasmithCADImporter/Source/CADInterfaces/Private/TechSoftFileParser.cpp
@@ -1252,7 +1252,7 @@ void FTechSoftFileParser::ExtractMetaData(const A3DEntity* Entity, FArchiveCADOb
 			FString Name = UTF8_TO_TCHAR(MetaData->m_pcName);
 			if(Name != TEXT("unnamed"))  // "unnamed" is create by Techsoft. This name is ignored 
 			{
-				Name = TechSoftUtils::CleanSdkName(Name);
+				Name = TechSoftUtils::CleanLabel(Name);
 				OutObject.Label = Name;
 			}
 		}
@@ -1290,7 +1290,7 @@ void FTechSoftFileParser::BuildReferenceName(FArchiveCADObject& ReferenceData)
 	{
 		if (Format == ECADFormat::CATIA)
 		{
-			ReferenceData.Label = TechSoftUtils::CleanCatiaReferenceName(*NamePtr);
+			ReferenceData.Label = TechSoftUtils::CleanCatiaReferenceLabel(*NamePtr);
 		}
 		else
 		{
@@ -1302,11 +1302,11 @@ void FTechSoftFileParser::BuildReferenceName(FArchiveCADObject& ReferenceData)
 	switch (Format)
 	{
 	case ECADFormat::CATIA_3DXML:
-		ReferenceData.Label = TechSoftUtils::Clean3dxmlReferenceSdkName(ReferenceData.Label);
+		ReferenceData.Label = TechSoftUtils::Clean3dxmlReferenceLabel(ReferenceData.Label);
 		break;
 
 	case ECADFormat::SOLIDWORKS:
-		ReferenceData.Label = TechSoftUtils::CleanSwReferenceSdkName(ReferenceData.Label);
+		ReferenceData.Label = TechSoftUtils::CleanSwReferenceLabel(ReferenceData.Label);
 		break;
 
 	default:
@@ -1328,15 +1328,15 @@ void FTechSoftFileParser::BuildInstanceName(FArchiveInstance& InstanceData, cons
 		switch (Format)
 		{
 		case ECADFormat::CATIA:
-			InstanceData.Label = TechSoftUtils::CleanCatiaInstanceSdkName(InstanceData.Label);
+			InstanceData.Label = TechSoftUtils::CleanCatiaInstanceLabel(InstanceData.Label);
 			break;
 		
 		case ECADFormat::CATIA_3DXML:
-			InstanceData.Label = TechSoftUtils::Clean3dxmlInstanceSdkName(InstanceData.Label);
+			InstanceData.Label = TechSoftUtils::Clean3dxmlInstanceLabel(InstanceData.Label);
 			break;
 
 		case ECADFormat::SOLIDWORKS:
-			InstanceData.Label = TechSoftUtils::CleanSwInstanceSdkName(InstanceData.Label);
+			InstanceData.Label = TechSoftUtils::CleanSwInstanceLabel(InstanceData.Label);
 			break;
 
 		default:
@@ -1368,7 +1368,7 @@ void FTechSoftFileParser::BuildBodyName(FArchiveBody& Body, const FArchiveRefere
 {
 	if (Format == ECADFormat::CREO)
 	{
-		Body.Label = TechSoftUtils::CleanCreoName(Body.Label);
+		Body.Label = TechSoftUtils::CleanCreoLabel(Body.Label);
 	}
 
 	if(Body.IsNameDefined())
diff --git a/Engine/Plugins/Enterprise/DatasmithCADImporter/Source/CADInterfaces/Private/TechSoftUtils.cpp b/Engine/Plugins/Enterprise/DatasmithCADImporter/Source/CADInterfaces/Private/TechSoftUtils.cpp
--- a/Engine/Plugins/Enterprise/DatasmithCADImporter/Source/CADInterfaces/Private/TechSoftUtils.cpp
+++ b/Engine/Plugins/Enterprise/DatasmithCADImporter/Source/CADInterfaces/Private/TechSoftUtils.cpp
@@ -663,7 +663,7 @@ void ExtractAttribute(const A3DMiscAttributeData& AttributeData, TMap<FString, F
 }
 #endif
 
-FString CleanSdkName(const FString& Name)
+FString CleanLabel(const FString& Name)
 {
 	int32 Index;
 	if (Name.FindLastChar(TEXT('['), Index))
@@ -673,7 +673,7 @@ FString CleanSdkName(const FString& Name)
 	return Name;
 }
 
-FString CleanCatiaInstanceSdkName(const FString& Name)
+FString CleanCatiaInstanceLabel(const FString& Name)
 {
 	int32 Index;
 	if (Name.FindChar(TEXT('('), Index))
@@ -688,9 +688,9 @@ FString CleanCatiaInstanceSdkName(const FString& Name)
 	return Name;
 }
 
-FString Clean3dxmlInstanceSdkName(const FString& Name)
+FString Clean3dxmlInstanceLabel(const FString& Name)
 {
-	FString NewName = CleanCatiaInstanceSdkName(Name);
+	FString NewName = CleanCatiaInstanceLabel(Name);
 
 	int32 Index = NewName.Find(TEXT("_InstanceRep"));
 	if (Index != INDEX_NONE)
@@ -701,7 +701,7 @@ FString Clean3dxmlInstanceSdkName(const FString& Name)
 	return NewName;
 }
 
-FString Clean3dxmlReferenceSdkName(const FString& Name)
+FString Clean3dxmlReferenceLabel(const FString& Name)
 {
 	int32 Index;
 	if (Name.FindChar(TEXT('('), Index))
@@ -720,7 +720,7 @@ FString Clean3dxmlReferenceSdkName(const FString& Name)
 	return Name;
 }
 
-FString CleanSwInstanceSdkName(const FString& Name)
+FString CleanSwInstanceLabel(const FString& Name)
 {
 	int32 Position;
 	if (Name.FindLastChar(TEXT('-'), Position))
@@ -731,7 +731,7 @@ FString CleanSwInstanceSdkName(const FString& Name)
 	return Name;
 }
 
-FString CleanSwReferenceSdkName(const FString& Name)
+FString CleanSwReferenceLabel(const FString& Name)
 {
 	int32 Position;
 	if (Name.FindLastChar(TEXT('-'), Position))
@@ -742,7 +742,7 @@ FString CleanSwReferenceSdkName(const FString& Name)
 	return Name;
 }
 
-FString CleanCatiaReferenceName(const FString& Name)
+FString CleanCatiaReferenceLabel(const FString& Name)
 {
 	int32 Position;
 	if (Name.FindLastChar(TEXT('.'), Position))
@@ -756,7 +756,7 @@ FString CleanCatiaReferenceName(const FString& Name)
 	}
 	return Name;
 }
-FString CleanCreoName(const FString& Name)
+FString CleanCreoLabel(const FString& Name)
 {
 	int32 Position;
 	if (Name.FindLastChar(TEXT('.'), Position))
diff --git a/Engine/Plugins/Enterprise/DatasmithCADImporter/Source/CADInterfaces/Private/TechSoftUtilsPrivate.h b/Engine/Plugins/Enterprise/DatasmithCADImporter/Source/CADInterfaces/Private/TechSoftUtilsPrivate.h
--- a/Engine/Plugins/Enterprise/DatasmithCADImporter/Source/CADInterfaces/Private/TechSoftUtilsPrivate.h
+++ b/Engine/Plugins/Enterprise/DatasmithCADImporter/Source/CADInterfaces/Private/TechSoftUtilsPrivate.h
@@ -1,6 +1,8 @@
 // Copyright Epic Games, Inc. All Rights Reserved.
 #pragma once
 
+#include "CoreMinimal.h"
+
 #include "TechSoftInterface.h"
 
 namespace CADLibrary
@@ -13,14 +15,14 @@ namespace TechSoftUtils
 void ExtractAttribute(const A3DMiscAttributeData& AttributeData, TMap<FString, FString>& OutMetaData);
 #endif
 
-FString CleanSdkName(const FString& Name);
-FString CleanCatiaInstanceSdkName(const FString& Name);
-FString CleanCatiaReferenceName(const FString& Name);
-FString Clean3dxmlInstanceSdkName(const FString& Name);
-FString Clean3dxmlReferenceSdkName(const FString& Name);
-FString CleanSwInstanceSdkName(const FString& Name);
-FString CleanSwReferenceSdkName(const FString& Name);
-FString CleanCreoName(const FString& Name);
+FString CleanLabel(const FString& Name);
+FString CleanCatiaInstanceLabel(const FString& Name);
+FString Clean3dxmlInstanceLabel(const FString& Name);
+FString CleanCatiaReferenceLabel(const FString& Name);
+FString Clean3dxmlReferenceLabel(const FString& Name);
+FString CleanSwInstanceLabel(const FString& Name);
+FString CleanSwReferenceLabel(const FString& Name);
+FString CleanCreoLabel(const FString& Name);
 bool CheckIfNameExists(TMap<FString, FString>& MetaData);
 bool ReplaceOrAddNameValue(TMap<FString, FString>& MetaData, const TCHAR* Key);
 
