[NO NON-PII COMMIT MESSAGE]

diff --git a/Engine/Plugins/Enterprise/DatasmithCADImporter/Source/CADTools/Private/CADData.cpp b/Engine/Plugins/Enterprise/DatasmithCADImporter/Source/CADTools/Private/CADData.cpp
--- a/Engine/Plugins/Enterprise/DatasmithCADImporter/Source/CADTools/Private/CADData.cpp
+++ b/Engine/Plugins/Enterprise/DatasmithCADImporter/Source/CADTools/Private/CADData.cpp
@@ -19,12 +19,12 @@ FMaterialUId BuildColorFastUId(uint32 ColorId, uint8 Alpha)
 		Alpha = 1;
 	}
 	uint32 FastColorId =  ColorId | Alpha << 24;
-	return FMath::Abs((int32)GetTypeHash(FastColorId));
+	return FMath::Abs((int32) ::GetTypeHash(FastColorId));
 }
 
 FMaterialUId BuildColorUId(const FColor& Color)
 {
-	return FMath::Abs((int32)GetTypeHash(Color));
+	return FMath::Abs((int32) ::GetTypeHash(Color));
 }
 
 FMaterialUId BuildMaterialUId(const FCADMaterial& Material)
