Add CAD extension (CATShape (Catia), exp, session, dlv (CATIA V4 ) sab (ACIS))

diff --git a/Engine/Plugins/Enterprise/DatasmithCADImporter/Source/CADTools/Private/CADData.cpp b/Engine/Plugins/Enterprise/DatasmithCADImporter/Source/CADTools/Private/CADData.cpp
--- a/Engine/Plugins/Enterprise/DatasmithCADImporter/Source/CADTools/Private/CADData.cpp
+++ b/Engine/Plugins/Enterprise/DatasmithCADImporter/Source/CADTools/Private/CADData.cpp
@@ -227,7 +227,7 @@ ECADFormat FileFormat(const FString& Extension)
 	{
 		return ECADFormat::JT;
 	}
-	else if (Extension == TEXT("model"))
+	else if (Extension == TEXT("model") || Extension == TEXT("exp") || Extension == TEXT("session") || Extension == TEXT("dlv"))
 	{
 		return ECADFormat::CATIAV4;
 	}
@@ -242,7 +242,7 @@ ECADFormat FileFormat(const FString& Extension)
 	{
 		return ECADFormat::NX;
 	}
-	else if (Extension == TEXT("sat"))
+	else if (Extension == TEXT("sat") || Extension == TEXT("sab"))
 	{
 		return ECADFormat::ACIS;
 	}
diff --git a/Engine/Plugins/Enterprise/DatasmithCADImporter/Source/DatasmithCADTranslator/Private/DatasmithCADTranslator.cpp b/Engine/Plugins/Enterprise/DatasmithCADImporter/Source/DatasmithCADTranslator/Private/DatasmithCADTranslator.cpp
--- a/Engine/Plugins/Enterprise/DatasmithCADImporter/Source/DatasmithCADTranslator/Private/DatasmithCADTranslator.cpp
+++ b/Engine/Plugins/Enterprise/DatasmithCADImporter/Source/DatasmithCADTranslator/Private/DatasmithCADTranslator.cpp
@@ -54,10 +54,16 @@ void FDatasmithCADTranslator::Initialize(FDatasmithTranslatorCapabilities& OutCa
 #endif
 	OutCapabilities.SupportedFileFormats.Add(FFileFormatInfo{ TEXT("CATPart"), TEXT("CATIA Part files") });
 	OutCapabilities.SupportedFileFormats.Add(FFileFormatInfo{ TEXT("CATProduct"), TEXT("CATIA Product files") });
+	OutCapabilities.SupportedFileFormats.Add(FFileFormatInfo{ TEXT("CATShape"), TEXT("CATIA Shape files") });
 	OutCapabilities.SupportedFileFormats.Add(FFileFormatInfo{ TEXT("cgr"), TEXT("CATIA Graphical Representation V5 files") });
-	OutCapabilities.SupportedFileFormats.Add(FFileFormatInfo{ TEXT("3dxml"), TEXT("CATIA files") });
-	OutCapabilities.SupportedFileFormats.Add(FFileFormatInfo{ TEXT("3drep"), TEXT("CATIA files") });
+
+	OutCapabilities.SupportedFileFormats.Add(FFileFormatInfo{ TEXT("3dxml"), TEXT("CATIA 3D xml files") });
+	OutCapabilities.SupportedFileFormats.Add(FFileFormatInfo{ TEXT("3drep"), TEXT("CATIA 3D xml files") });
+
 	OutCapabilities.SupportedFileFormats.Add(FFileFormatInfo{ TEXT("model"), TEXT("CATIA V4 files") });
+	OutCapabilities.SupportedFileFormats.Add(FFileFormatInfo{ TEXT("session"), TEXT("CATIA V4 files") });
+	OutCapabilities.SupportedFileFormats.Add(FFileFormatInfo{ TEXT("exp"), TEXT("CATIA V4 files") });
+	OutCapabilities.SupportedFileFormats.Add(FFileFormatInfo{ TEXT("dlv"), TEXT("CATIA V4 files") });
 
 	OutCapabilities.SupportedFileFormats.Add(FFileFormatInfo{ TEXT("asm.*"), TEXT("Creo Assembly files") });
 	OutCapabilities.SupportedFileFormats.Add(FFileFormatInfo{ TEXT("creo.*"), TEXT("Creo Assembly files") });
@@ -77,6 +83,7 @@ void FDatasmithCADTranslator::Initialize(FDatasmithTranslatorCapabilities& OutCa
 	OutCapabilities.SupportedFileFormats.Add(FFileFormatInfo{ TEXT("jt"), TEXT("JT Open files") });
 
 	OutCapabilities.SupportedFileFormats.Add(FFileFormatInfo{ TEXT("sat"), TEXT("3D ACIS model files") });
+	OutCapabilities.SupportedFileFormats.Add(FFileFormatInfo{ TEXT("sab"), TEXT("3D ACIS model files") });
 
 	OutCapabilities.SupportedFileFormats.Add(FFileFormatInfo{ TEXT("SLDASM"), TEXT("SolidWorks Product files") });
 	OutCapabilities.SupportedFileFormats.Add(FFileFormatInfo{ TEXT("SLDPRT"), TEXT("SolidWorks Part files") });
