[NO NON-PII COMMIT MESSAGE]

diff --git a/Engine/Plugins/Enterprise/DatasmithCADImporter/Source/DatasmithWireTranslator/Private/OpenModelUtils.cpp b/Engine/Plugins/Enterprise/DatasmithCADImporter/Source/DatasmithWireTranslator/Private/OpenModelUtils.cpp
--- a/Engine/Plugins/Enterprise/DatasmithCADImporter/Source/DatasmithWireTranslator/Private/OpenModelUtils.cpp
+++ b/Engine/Plugins/Enterprise/DatasmithCADImporter/Source/DatasmithWireTranslator/Private/OpenModelUtils.cpp
@@ -25,1539 +25,6 @@
 namespace UE_DATASMITHWIRETRANSLATOR_NAMESPACE
 {
 
-const TCHAR* OpenModelUtils::AlObjectTypeToString(AlObjectType type)
-{
-	switch (type)
-	{
-	case kAmbientLightType:
-		return TEXT("kAmbientLightType");
-	case kAreaLightType:
-		return TEXT("kAreaLightType");
-	case kBoxLightType:
-		return TEXT("kBoxLightType");
-	case kCameraEyeType:
-		return TEXT("kCameraEyeType");
-	case kCameraType:
-		return TEXT("kCameraType");
-	case kCameraUpType:
-		return TEXT("kCameraUpType");
-	case kCameraViewType:
-		return TEXT("kCameraViewType");
-	case kClusterNodeType:
-		return TEXT("kClusterNodeType");
-	case kClusterType:
-		return TEXT("kClusterType");
-	case kClusterMemberType:
-		return TEXT("kClusterMemberType");
-	case kConeLightType:
-		return TEXT("kConeLightType");
-	case kCurveNodeType:
-		return TEXT("kCurveNodeType");
-	case kCurveOnSurfaceType:
-		return TEXT("kCurveOnSurfaceType");
-	case kCurveType:
-		return TEXT("kCurveType");
-	case kCurveCVType:
-		return TEXT("kCurveCVType");
-	case kCylinderLightType:
-		return TEXT("kCylinderLightType");
-	case kDagNodeType:
-		return TEXT("kDagNodeType");
-	case kDirectionLightType:
-		return TEXT("kDirectionLightType");
-	case kFaceNodeType:
-		return TEXT("kFaceNodeType");
-	case kFaceType:
-		return TEXT("kFaceType");
-	case kGroupNodeType:
-		return TEXT("kGroupNodeType");
-	case kImagePlaneType:
-		return TEXT("kImagePlaneType");
-	case kLightLookAtNodeType:
-		return TEXT("kLightLookAtNodeType");
-	case kLightNodeType:
-		return TEXT("kLightNodeType");
-	case kLightType:
-		return TEXT("kLightType");
-	case kLightUpNodeType:
-		return TEXT("kLightUpNodeType");
-	case kLinearLightType:
-		return TEXT("kLinearLightType");
-	case kNonAmbientLightType:
-		return TEXT("kNonAmbientLightType");
-	case kOrthographicCameraType:
-		return TEXT("kOrthographicCameraType");
-	case kPerspectiveCameraType:
-		return TEXT("kPerspectiveCameraType");
-	case kPointLightType:
-		return TEXT("kPointLightType");
-	case kSetType:
-		return TEXT("kSetType");
-	case kSetMemberType:
-		return TEXT("kSetMemberType");
-	case kSphereLightType:
-		return TEXT("kSphereLightType");
-	case kSpotLightType:
-		return TEXT("kSpotLightType");
-	case kSurfaceNodeType:
-		return TEXT("kSurfaceNodeType");
-	case kSurfaceType:
-		return TEXT("kSurfaceType");
-	case kSurfaceCurveType:
-		return TEXT("kSurfaceCurveType");
-	case kSurfaceCVType:
-		return TEXT("kSurfaceCVType");
-	case kTorusLightType:
-		return TEXT("kTorusLightType");
-	case kVolumeLightType:
-		return TEXT("kVolumeLightType");
-	case kWindowType:
-		return TEXT("kWindowType");
-	case kChannelType:
-		return TEXT("kChannelType");
-	case kActionType:
-		return TEXT("kActionType");
-	case kParamActionType:
-		return TEXT("kParamActionType");
-	case kMotionActionType:
-		return TEXT("kMotionActionType");
-	case kKeyframeType:
-		return TEXT("kKeyframeType");
-	case kStreamType:
-		return TEXT("kStreamType");
-	case kEnvironmentType:
-		return TEXT("kEnvironmentType");
-	case kShaderType:
-		return TEXT("kShaderType");
-	case kTextureType:
-		return TEXT("kTextureType");
-	case kPolysetNodeType:
-		return TEXT("kPolysetNodeType");
-	case kPolysetType:
-		return TEXT("kPolysetType");
-	case kPolygonType:
-		return TEXT("kPolygonType");
-	case kPolysetVertexType:
-		return TEXT("kPolysetVertexType");
-	case kAttributeType:
-		return TEXT("kAttributeType");
-	case kArcAttributeType:
-		return TEXT("kArcAttributeType");
-	case kLineAttributeType:
-		return TEXT("kLineAttributeType");
-	case kCurveAttributeType:
-		return TEXT("kCurveAttributeType");
-	case kPlaneAttributeType:
-		return TEXT("kPlaneAttributeType");
-	case kConicAttributeType:
-		return TEXT("kConicAttributeType");
-	case kRevSurfAttributeType:
-		return TEXT("kRevSurfAttributeType");
-	case kJointType:
-		return TEXT("kJointType");
-	case kConstraintType:
-		return TEXT("kConstraintType");
-	case kPointConstraintType:
-		return TEXT("kPointConstraintType");
-	case kOrientationConstraintType:
-		return TEXT("kOrientationConstraintType");
-	case kAimConstraintType:
-		return TEXT("kAimConstraintType");
-	case kTextureNodeType:
-		return TEXT("kTextureNodeType");
-	case kShellNodeType:
-		return TEXT("kShellNodeType");
-	case kShellType:
-		return TEXT("kShellType");
-	case kTrimRegionType:
-		return TEXT("kTrimRegionType");
-	case kTrimBoundaryType:
-		return TEXT("kTrimBoundaryType");
-	case kTrimCurveType:
-		return TEXT("kTrimCurveType");
-	case kCommandType:
-		return TEXT("kCommandType");
-	case kCommandRefType:
-		return TEXT("kCommandRefType");
-	case kContactType:
-		return TEXT("kContactType");
-	case kCharacterSpaceType:
-		return TEXT("kCharacterSpaceType");
-	case kCharacterType:
-		return TEXT("kCharacterType");
-	case kCharSnippetType:
-		return TEXT("kCharSnippetType");
-	case kCharTransitionType:
-		return TEXT("kCharTransitionType");
-	case kIKHandleType:
-		return TEXT("kIKHandleType");
-	case kIKHandleNodeType:
-		return TEXT("kIKHandleNodeType");
-	case kLocatorType:
-		return TEXT("kLocatorType");
-	case kAnnotationLocatorType:
-		return TEXT("kAnnotationLocatorType");
-	case kDistanceLocatorType:
-		return TEXT("kDistanceLocatorType");
-	case kAngleLocatorType:
-		return TEXT("kAngleLocatorType");
-	case kRadialLocatorType:
-		return TEXT("kRadialLocatorType");
-	case kDeviationLocatorType:
-		return TEXT("kDeviationLocatorType");
-	case kMinmaxLocatorType:
-		return TEXT("kMinmaxLocatorType");
-	case kConstructionEntityType:
-		return TEXT("kConstructionEntityType");
-	case kConstructionVectorType:
-		return TEXT("kConstructionVectorType");
-	case kConstructionPlaneType:
-		return TEXT("kConstructionPlaneType");
-	case kConstructionFutureType1:
-		return TEXT("kConstructionFutureType1");
-	case kConstructionFutureType2:
-		return TEXT("kConstructionFutureType2");
-	case kPointType:
-		return TEXT("kPointType");
-	case kSpacePointType:
-		return TEXT("kSpacePointType");
-	case kCurvePointType:
-		return TEXT("kCurvePointType");
-	case kSurfacePointType:
-		return TEXT("kSurfacePointType");
-	case kCurveOnSurfacePointType:
-		return TEXT("kCurveOnSurfacePointType");
-	case kFuturePointType1:
-		return TEXT("kFuturePointType1");
-	case kFuturePointType2:
-		return TEXT("kFuturePointType2");
-	case kFuturePointType3:
-		return TEXT("kFuturePointType3");
-	case kFuturePointType4:
-		return TEXT("kFuturePointType4");
-	case kFuturePointType5:
-		return TEXT("kFuturePointType5");
-	case kLayerType:
-		return TEXT("kLayerType");
-	case kCloudType:
-		return TEXT("kCloudType");
-	case kBlendCurveType:
-		return TEXT("kBlendCurveType");
-	case kBlendPointType:
-		return TEXT("kBlendPointType");
-	case kCategoryType:
-		return TEXT("kCategoryType");
-	case kMeshType:
-		return TEXT("kMeshType");
-	case kMeshNodeType:
-		return TEXT("kMeshNodeType");
-	case kEvaluateType:
-		return TEXT("kEvaluateType");
-	case kReferenceFileType:
-		return TEXT("kReferenceFileType");
-	case kReferenceFileSetType:
-		return TEXT("kReferenceFileSetType");
-	}
-	return TEXT("kUndefined");
-}
-
-const TCHAR* OpenModelUtils::AlShadingFieldToString(AlShadingFields field)
-{
-	switch (field)
-	{
-
-	case kFLD_SHADING_COMMON_COLOR_R:
-		return TEXT("kFLD_SHADING_COMMON_COLOR_R");
-	case kFLD_SHADING_COMMON_COLOR_G:
-		return TEXT("kFLD_SHADING_COMMON_COLOR_G");
-	case kFLD_SHADING_COMMON_COLOR_B:
-		return TEXT("kFLD_SHADING_COMMON_COLOR_B");
-	case kFLD_SHADING_COMMON_INCANDESCENCE_R:
-		return TEXT("kFLD_SHADING_COMMON_INCANDESCENCE_R");
-	case kFLD_SHADING_COMMON_INCANDESCENCE_G:
-		return TEXT("kFLD_SHADING_COMMON_INCANDESCENCE_G");
-	case kFLD_SHADING_COMMON_INCANDESCENCE_B:
-		return TEXT("kFLD_SHADING_COMMON_INCANDESCENCE_B");
-	case kFLD_SHADING_COMMON_TRANSPARENCY_R:
-		return TEXT("kFLD_SHADING_COMMON_TRANSPARENCY_R");
-	case kFLD_SHADING_COMMON_TRANSPARENCY_G:
-		return TEXT("kFLD_SHADING_COMMON_TRANSPARENCY_G");
-	case kFLD_SHADING_COMMON_TRANSPARENCY_B:
-		return TEXT("kFLD_SHADING_COMMON_TRANSPARENCY_B");
-	case kFLD_SHADING_COMMON_USE_BACKGROUND_COLOR:
-		return TEXT("kFLD_SHADING_COMMON_USE_BACKGROUND_COLOR");
-	case kFLD_SHADING_COMMON_REFRACTIVE_INDEX:
-		return TEXT("kFLD_SHADING_COMMON_REFRACTIVE_INDEX");
-	case kFLD_SHADING_COMMON_REFLECTION_LIMIT:
-		return TEXT("kFLD_SHADING_COMMON_REFLECTION_LIMIT");
-	case kFLD_SHADING_COMMON_REFRACTION_LIMIT:
-		return TEXT("kFLD_SHADING_COMMON_REFRACTION_LIMIT");
-	case kFLD_SHADING_COMMON_SHADOW_LEVEL_LIMIT:
-		return TEXT("kFLD_SHADING_COMMON_SHADOW_LEVEL_LIMIT");
-	case kFLD_SHADING_COMMON_TRANSPARENCY_SHADE:
-		return TEXT("kFLD_SHADING_COMMON_TRANSPARENCY_SHADE");
-	case kFLD_SHADING_COMMON_RESPECT_REFLECTION_MAP:
-		return TEXT("kFLD_SHADING_COMMON_RESPECT_REFLECTION_MAP");
-	case kFLD_SHADING_LAMBERT_DIFFUSE:
-		return TEXT("kFLD_SHADING_LAMBERT_DIFFUSE");
-	case kFLD_SHADING_PHONG_DIFFUSE:
-		return TEXT("kFLD_SHADING_PHONG_DIFFUSE");
-	case kFLD_SHADING_PHONG_SPECULAR_R:
-		return TEXT("kFLD_SHADING_PHONG_SPECULAR_R");
-	case kFLD_SHADING_PHONG_SPECULAR_G:
-		return TEXT("kFLD_SHADING_PHONG_SPECULAR_G");
-	case kFLD_SHADING_PHONG_SPECULAR_B:
-		return TEXT("kFLD_SHADING_PHONG_SPECULAR_B");
-	case kFLD_SHADING_PHONG_SHINYNESS:
-		return TEXT("kFLD_SHADING_PHONG_SHINYNESS");
-	case kFLD_SHADING_PHONG_REFLECTIVITY:
-		return TEXT("kFLD_SHADING_PHONG_REFLECTIVITY");
-	case kFLD_SHADING_BLINN_DIFFUSE:
-		return TEXT("kFLD_SHADING_BLINN_DIFFUSE");
-	case kFLD_SHADING_BLINN_SPECULAR_R:
-		return TEXT("kFLD_SHADING_BLINN_SPECULAR_R");
-	case kFLD_SHADING_BLINN_SPECULAR_G:
-		return TEXT("kFLD_SHADING_BLINN_SPECULAR_G");
-	case kFLD_SHADING_BLINN_SPECULAR_B:
-		return TEXT("kFLD_SHADING_BLINN_SPECULAR_B");
-	case kFLD_SHADING_BLINN_SPECULAR_ROLLOFF:
-		return TEXT("kFLD_SHADING_BLINN_SPECULAR_ROLLOFF");
-	case kFLD_SHADING_BLINN_ECCENTRICITY:
-		return TEXT("kFLD_SHADING_BLINN_ECCENTRICITY");
-	case kFLD_SHADING_BLINN_REFLECTIVITY:
-		return TEXT("kFLD_SHADING_BLINN_REFLECTIVITY");
-	case kFLD_SHADING_BACKGROUND_COLOR_R:
-		return TEXT("kFLD_SHADING_BACKGROUND_COLOR_R");
-	case kFLD_SHADING_BACKGROUND_COLOR_G:
-		return TEXT("kFLD_SHADING_BACKGROUND_COLOR_G");
-	case kFLD_SHADING_BACKGROUND_COLOR_B:
-		return TEXT("kFLD_SHADING_BACKGROUND_COLOR_B");
-	case kFLD_SHADING_FOG_COLOR_R:
-		return TEXT("kFLD_SHADING_FOG_COLOR_R");
-	case kFLD_SHADING_FOG_COLOR_G:
-		return TEXT("kFLD_SHADING_FOG_COLOR_G");
-	case kFLD_SHADING_FOG_COLOR_B:
-		return TEXT("kFLD_SHADING_FOG_COLOR_B");
-	case kFLD_SHADING_FOG_DEPTH_R:
-		return TEXT("kFLD_SHADING_FOG_DEPTH_R");
-	case kFLD_SHADING_FOG_DEPTH_G:
-		return TEXT("kFLD_SHADING_FOG_DEPTH_G");
-	case kFLD_SHADING_FOG_DEPTH_B:
-		return TEXT("kFLD_SHADING_FOG_DEPTH_B");
-	case kFLD_SHADING_FOG_MIN_ALTITUDE:
-		return TEXT("kFLD_SHADING_FOG_MIN_ALTITUDE");
-	case kFLD_SHADING_FOG_MAX_ALTITUDE:
-		return TEXT("kFLD_SHADING_FOG_MAX_ALTITUDE");
-	case kFLD_SHADING_FOG_AMULT:
-		return TEXT("kFLD_SHADING_FOG_AMULT");
-	case kFLD_SHADING_FOG_AOFFSET:
-		return TEXT("kFLD_SHADING_FOG_AOFFSET");
-	case kFLD_SHADING_FOG_FOGTYPE:
-		return TEXT("kFLD_SHADING_FOG_FOGType");
-	case kFLD_SHADING_FOG_COLUMN:
-		return TEXT("kFLD_SHADING_FOG_COLUMN");
-	case kFLD_SHADING_TEXTURE_ACTIVE:
-		return TEXT("kFLD_SHADING_TEXTURE_ACTIVE");
-	case kFLD_SHADING_COMMON_TEXTURE_RGBMULT_R:
-		return TEXT("kFLD_SHADING_COMMON_TEXTURE_RGBMULT_R");
-	case kFLD_SHADING_COMMON_TEXTURE_RGBMULT_G:
-		return TEXT("kFLD_SHADING_COMMON_TEXTURE_RGBMULT_G");
-	case kFLD_SHADING_COMMON_TEXTURE_RGBMULT_B:
-		return TEXT("kFLD_SHADING_COMMON_TEXTURE_RGBMULT_B");
-	case kFLD_SHADING_COMMON_TEXTURE_RGBOFFSET_R:
-		return TEXT("kFLD_SHADING_COMMON_TEXTURE_RGBOFFSET_R");
-	case kFLD_SHADING_COMMON_TEXTURE_RGBOFFSET_G:
-		return TEXT("kFLD_SHADING_COMMON_TEXTURE_RGBOFFSET_G");
-	case kFLD_SHADING_COMMON_TEXTURE_RGBOFFSET_B:
-		return TEXT("kFLD_SHADING_COMMON_TEXTURE_RGBOFFSET_B");
-	case kFLD_SHADING_COMMON_TEXTURE_AMULT:
-		return TEXT("kFLD_SHADING_COMMON_TEXTURE_AMULT");
-	case kFLD_SHADING_COMMON_TEXTURE_AOFFSET:
-		return TEXT("kFLD_SHADING_COMMON_TEXTURE_AOFFSET");
-	case kFLD_SHADING_COMMON_TEXTURE_BLURMULT:
-		return TEXT("kFLD_SHADING_COMMON_TEXTURE_BLURMULT");
-	case kFLD_SHADING_COMMON_TEXTURE_BLUROFFSET:
-		return TEXT("kFLD_SHADING_COMMON_TEXTURE_BLUROFFSET");
-	case kFLD_SHADING_COMMON_TEXTURE_ROTATE:
-		return TEXT("kFLD_SHADING_COMMON_TEXTURE_ROTATE");
-	case kFLD_SHADING_COMMON_TEXTURE_UREPEAT:
-		return TEXT("kFLD_SHADING_COMMON_TEXTURE_UREPEAT");
-	case kFLD_SHADING_COMMON_TEXTURE_VREPEAT:
-		return TEXT("kFLD_SHADING_COMMON_TEXTURE_VREPEAT");
-	case kFLD_SHADING_COMMON_TEXTURE_UOFFSET:
-		return TEXT("kFLD_SHADING_COMMON_TEXTURE_UOFFSET");
-	case kFLD_SHADING_COMMON_TEXTURE_VOFFSET:
-		return TEXT("kFLD_SHADING_COMMON_TEXTURE_VOFFSET");
-	case kFLD_SHADING_COMMON_TEXTURE_MIRROR:
-		return TEXT("kFLD_SHADING_COMMON_TEXTURE_MIRROR");
-	case kFLD_SHADING_COMMON_TEXTURE_UNUSED:
-		return TEXT("kFLD_SHADING_COMMON_TEXTURE_UNUSED");
-	case kFLD_SHADING_COMMON_TEXTURE_UCOVERAGE:
-		return TEXT("kFLD_SHADING_COMMON_TEXTURE_UCOVERAGE");
-	case kFLD_SHADING_COMMON_TEXTURE_VCOVERAGE:
-		return TEXT("kFLD_SHADING_COMMON_TEXTURE_VCOVERAGE");
-	case kFLD_SHADING_COMMON_TEXTURE_UTRANSLATE:
-		return TEXT("kFLD_SHADING_COMMON_TEXTURE_UTRANSLATE");
-	case kFLD_SHADING_COMMON_TEXTURE_VTRANSLATE:
-		return TEXT("kFLD_SHADING_COMMON_TEXTURE_VTRANSLATE");
-	case kFLD_SHADING_COMMON_TEXTURE_UWRAP:
-		return TEXT("kFLD_SHADING_COMMON_TEXTURE_UWRAP");
-	case kFLD_SHADING_COMMON_TEXTURE_VWRAP:
-		return TEXT("kFLD_SHADING_COMMON_TEXTURE_VWRAP");
-	case kFLD_SHADING_BULGE_UWIDTH:
-		return TEXT("kFLD_SHADING_BULGE_UWIDTH");
-	case kFLD_SHADING_BULGE_VWIDTH:
-		return TEXT("kFLD_SHADING_BULGE_VWIDTH");
-	case kFLD_SHADING_BALL_INCLINATION:
-		return TEXT("kFLD_SHADING_BALL_INCLINATION");
-	case kFLD_SHADING_BALL_ELEVATION:
-		return TEXT("kFLD_SHADING_BALL_ELEVATION");
-	case kFLD_SHADING_BALL_EYESPACE:
-		return TEXT("kFLD_SHADING_BALL_EYESPACE");
-	case kFLD_SHADING_BALL_REFLECT:
-		return TEXT("kFLD_SHADING_BALL_REFLECT");
-	case kFLD_SHADING_BALL_SKY_RADIUS:
-		return TEXT("kFLD_SHADING_BALL_SKY_RADIUS");
-	case kFLD_SHADING_BALL_BOTTOM:
-		return TEXT("kFLD_SHADING_BALL_BOTTOM");
-	case kFLD_SHADING_BALL_TOP:
-		return TEXT("kFLD_SHADING_BALL_TOP");
-	case kFLD_SHADING_BALL_LEFT:
-		return TEXT("kFLD_SHADING_BALL_LEFT");
-	case kFLD_SHADING_BALL_RIGHT:
-		return TEXT("kFLD_SHADING_BALL_RIGHT");
-	case kFLD_SHADING_BALL_FRONT:
-		return TEXT("kFLD_SHADING_BALL_FRONT");
-	case kFLD_SHADING_BALL_BACK:
-		return TEXT("kFLD_SHADING_BALL_BACK");
-	case kFLD_SHADING_CHECKER_COLOR1_R:
-		return TEXT("kFLD_SHADING_CHECKER_COLOR1_R");
-	case kFLD_SHADING_CHECKER_COLOR1_G:
-		return TEXT("kFLD_SHADING_CHECKER_COLOR1_G");
-	case kFLD_SHADING_CHECKER_COLOR1_B:
-		return TEXT("kFLD_SHADING_CHECKER_COLOR1_B");
-	case kFLD_SHADING_CHECKER_COLOR2_R:
-		return TEXT("kFLD_SHADING_CHECKER_COLOR2_R");
-	case kFLD_SHADING_CHECKER_COLOR2_G:
-		return TEXT("kFLD_SHADING_CHECKER_COLOR2_G");
-	case kFLD_SHADING_CHECKER_COLOR2_B:
-		return TEXT("kFLD_SHADING_CHECKER_COLOR2_B");
-	case kFLD_SHADING_CHECKER_CONTRAST:
-		return TEXT("kFLD_SHADING_CHECKER_CONTRAST");
-	case kFLD_SHADING_CLOTH_U_THREAD_COLOR_R:
-		return TEXT("kFLD_SHADING_CLOTH_U_THREAD_COLOR_R");
-	case kFLD_SHADING_CLOTH_U_THREAD_COLOR_G:
-		return TEXT("kFLD_SHADING_CLOTH_U_THREAD_COLOR_G");
-	case kFLD_SHADING_CLOTH_U_THREAD_COLOR_B:
-		return TEXT("kFLD_SHADING_CLOTH_U_THREAD_COLOR_B");
-	case kFLD_SHADING_CLOTH_V_THREAD_COLOR_R:
-		return TEXT("kFLD_SHADING_CLOTH_V_THREAD_COLOR_R");
-	case kFLD_SHADING_CLOTH_V_THREAD_COLOR_G:
-		return TEXT("kFLD_SHADING_CLOTH_V_THREAD_COLOR_G");
-	case kFLD_SHADING_CLOTH_V_THREAD_COLOR_B:
-		return TEXT("kFLD_SHADING_CLOTH_V_THREAD_COLOR_B");
-	case kFLD_SHADING_CLOTH_GAP_COLOR_R:
-		return TEXT("kFLD_SHADING_CLOTH_GAP_COLOR_R");
-	case kFLD_SHADING_CLOTH_GAP_COLOR_G:
-		return TEXT("kFLD_SHADING_CLOTH_GAP_COLOR_G");
-	case kFLD_SHADING_CLOTH_GAP_COLOR_B:
-		return TEXT("kFLD_SHADING_CLOTH_GAP_COLOR_B");
-	case kFLD_SHADING_CLOTH_U_THREAD_WIDTH:
-		return TEXT("kFLD_SHADING_CLOTH_U_THREAD_WIDTH");
-	case kFLD_SHADING_CLOTH_V_THREAD_WIDTH:
-		return TEXT("kFLD_SHADING_CLOTH_V_THREAD_WIDTH");
-	case kFLD_SHADING_CLOTH_U_WAVE:
-		return TEXT("kFLD_SHADING_CLOTH_U_WAVE");
-	case kFLD_SHADING_CLOTH_V_WAVE:
-		return TEXT("kFLD_SHADING_CLOTH_V_WAVE");
-	case kFLD_SHADING_CLOTH_RANDOMNESS:
-		return TEXT("kFLD_SHADING_CLOTH_RANDOMNESS");
-	case kFLD_SHADING_CLOTH_WIDTH_SPREAD:
-		return TEXT("kFLD_SHADING_CLOTH_WIDTH_SPREAD");
-	case kFLD_SHADING_CLOTH_BRIGHTNESS_SPREAD:
-		return TEXT("kFLD_SHADING_CLOTH_BRIGHTNESS_SPREAD");
-	case kFLD_SHADING_CHROME_SKY_COLOR_R:
-		return TEXT("kFLD_SHADING_CHROME_SKY_COLOR_R");
-	case kFLD_SHADING_CHROME_SKY_COLOR_G:
-		return TEXT("kFLD_SHADING_CHROME_SKY_COLOR_G");
-	case kFLD_SHADING_CHROME_SKY_COLOR_B:
-		return TEXT("kFLD_SHADING_CHROME_SKY_COLOR_B");
-	case kFLD_SHADING_CHROME_ZENITH_COLOR_R:
-		return TEXT("kFLD_SHADING_CHROME_ZENITH_COLOR_R");
-	case kFLD_SHADING_CHROME_ZENITH_COLOR_G:
-		return TEXT("kFLD_SHADING_CHROME_ZENITH_COLOR_G");
-	case kFLD_SHADING_CHROME_ZENITH_COLOR_B:
-		return TEXT("kFLD_SHADING_CHROME_ZENITH_COLOR_B");
-	case kFLD_SHADING_CHROME_LIGHT_COLOR_R:
-		return TEXT("kFLD_SHADING_CHROME_LIGHT_COLOR_R");
-	case kFLD_SHADING_CHROME_LIGHT_COLOR_G:
-		return TEXT("kFLD_SHADING_CHROME_LIGHT_COLOR_G");
-	case kFLD_SHADING_CHROME_LIGHT_COLOR_B:
-		return TEXT("kFLD_SHADING_CHROME_LIGHT_COLOR_B");
-	case kFLD_SHADING_CHROME_LIGHT_WIDTH:
-		return TEXT("kFLD_SHADING_CHROME_LIGHT_WIDTH");
-	case kFLD_SHADING_CHROME_LIGHT_WIDTH_MULT:
-		return TEXT("kFLD_SHADING_CHROME_LIGHT_WIDTH_MULT");
-	case kFLD_SHADING_CHROME_LIGHT_WIDTH_OFFSET:
-		return TEXT("kFLD_SHADING_CHROME_LIGHT_WIDTH_OFFSET");
-	case kFLD_SHADING_CHROME_LIGHT_DEPTH:
-		return TEXT("kFLD_SHADING_CHROME_LIGHT_DEPTH");
-	case kFLD_SHADING_CHROME_LIGHT_DEPTH_MULT:
-		return TEXT("kFLD_SHADING_CHROME_LIGHT_DEPTH_MULT");
-	case kFLD_SHADING_CHROME_LIGHT_DEPTH_OFFSET:
-		return TEXT("kFLD_SHADING_CHROME_LIGHT_DEPTH_OFFSET");
-	case kFLD_SHADING_CHROME_FLOOR_COLOR_R:
-		return TEXT("kFLD_SHADING_CHROME_FLOOR_COLOR_R");
-	case kFLD_SHADING_CHROME_FLOOR_COLOR_G:
-		return TEXT("kFLD_SHADING_CHROME_FLOOR_COLOR_G");
-	case kFLD_SHADING_CHROME_FLOOR_COLOR_B:
-		return TEXT("kFLD_SHADING_CHROME_FLOOR_COLOR_B");
-	case kFLD_SHADING_CHROME_HORIZON_COLOR_R:
-		return TEXT("kFLD_SHADING_CHROME_HORIZON_COLOR_R");
-	case kFLD_SHADING_CHROME_HORIZON_COLOR_G:
-		return TEXT("kFLD_SHADING_CHROME_HORIZON_COLOR_G");
-	case kFLD_SHADING_CHROME_HORIZON_COLOR_B:
-		return TEXT("kFLD_SHADING_CHROME_HORIZON_COLOR_B");
-	case kFLD_SHADING_CHROME_GRID_COLOR_R:
-		return TEXT("kFLD_SHADING_CHROME_GRID_COLOR_R");
-	case kFLD_SHADING_CHROME_GRID_COLOR_G:
-		return TEXT("kFLD_SHADING_CHROME_GRID_COLOR_G");
-	case kFLD_SHADING_CHROME_GRID_COLOR_B:
-		return TEXT("kFLD_SHADING_CHROME_GRID_COLOR_B");
-	case kFLD_SHADING_CHROME_REAL_FLOOR:
-		return TEXT("kFLD_SHADING_CHROME_REAL_FLOOR");
-	case kFLD_SHADING_CHROME_GRID_WIDTH:
-		return TEXT("kFLD_SHADING_CHROME_GRID_WIDTH");
-	case kFLD_SHADING_CHROME_GRID_WIDTH_MULT:
-		return TEXT("kFLD_SHADING_CHROME_GRID_WIDTH_MULT");
-	case kFLD_SHADING_CHROME_GRID_WIDTH_OFFSET:
-		return TEXT("kFLD_SHADING_CHROME_GRID_WIDTH_OFFSET");
-	case kFLD_SHADING_CHROME_GRID_DEPTH:
-		return TEXT("kFLD_SHADING_CHROME_GRID_DEPTH");
-	case kFLD_SHADING_CHROME_GRID_DEPTH_MULT:
-		return TEXT("kFLD_SHADING_CHROME_GRID_DEPTH_MULT");
-	case kFLD_SHADING_CHROME_GRID_DEPTH_OFFSET:
-		return TEXT("kFLD_SHADING_CHROME_GRID_DEPTH_OFFSET");
-	case kFLD_SHADING_CURVATURE_MIN:
-		return TEXT("kFLD_SHADING_CURVATURE_MIN");
-	case kFLD_SHADING_CURVATURE_MAX:
-		return TEXT("kFLD_SHADING_CURVATURE_MAX");
-	case kFLD_SHADING_FRACTAL_AMPLITUDE:
-		return TEXT("kFLD_SHADING_FRACTAL_AMPLITUDE");
-	case kFLD_SHADING_FRACTAL_THRESHOLD:
-		return TEXT("kFLD_SHADING_FRACTAL_THRESHOLD");
-	case kFLD_SHADING_FRACTAL_RATIO:
-		return TEXT("kFLD_SHADING_FRACTAL_RATIO");
-	case kFLD_SHADING_FRACTAL_LEVEL_MIN:
-		return TEXT("kFLD_SHADING_FRACTAL_LEVEL_MIN");
-	case kFLD_SHADING_FRACTAL_LEVEL_MAX:
-		return TEXT("kFLD_SHADING_FRACTAL_LEVEL_MAX");
-	case kFLD_SHADING_FRACTAL_ANIMATED:
-		return TEXT("kFLD_SHADING_FRACTAL_ANIMATED");
-	case kFLD_SHADING_FRACTAL_TIME:
-		return TEXT("kFLD_SHADING_FRACTAL_TIME");
-	case kFLD_SHADING_GRID_LINE_COLOR_R:
-		return TEXT("kFLD_SHADING_GRID_LINE_COLOR_R");
-	case kFLD_SHADING_GRID_LINE_COLOR_G:
-		return TEXT("kFLD_SHADING_GRID_LINE_COLOR_G");
-	case kFLD_SHADING_GRID_LINE_COLOR_B:
-		return TEXT("kFLD_SHADING_GRID_LINE_COLOR_B");
-	case kFLD_SHADING_GRID_FILLER_COLOR_R:
-		return TEXT("kFLD_SHADING_GRID_FILLER_COLOR_R");
-	case kFLD_SHADING_GRID_FILLER_COLOR_G:
-		return TEXT("kFLD_SHADING_GRID_FILLER_COLOR_G");
-	case kFLD_SHADING_GRID_FILLER_COLOR_B:
-		return TEXT("kFLD_SHADING_GRID_FILLER_COLOR_B");
-	case kFLD_SHADING_GRID_UWIDTH:
-		return TEXT("kFLD_SHADING_GRID_UWIDTH");
-	case kFLD_SHADING_GRID_VWIDTH:
-		return TEXT("kFLD_SHADING_GRID_VWIDTH");
-	case kFLD_SHADING_GRID_CONTRAST:
-		return TEXT("kFLD_SHADING_GRID_CONTRAST");
-	case kFLD_SHADING_HIGHLIGHT_MIN:
-		return TEXT("kFLD_SHADING_HIGHLIGHT_MIN");
-	case kFLD_SHADING_HIGHLIGHT_MAX:
-		return TEXT("kFLD_SHADING_HIGHLIGHT_MAX");
-	case kFLD_SHADING_HIGHLIGHT_HLDIRECTION_X:
-		return TEXT("kFLD_SHADING_HIGHLIGHT_HLDIRECTION_X");
-	case kFLD_SHADING_HIGHLIGHT_HLDIRECTION_Y:
-		return TEXT("kFLD_SHADING_HIGHLIGHT_HLDIRECTION_Y");
-	case kFLD_SHADING_HIGHLIGHT_HLDIRECTION_Z:
-		return TEXT("kFLD_SHADING_HIGHLIGHT_HLDIRECTION_Z");
-	case kFLD_SHADING_MOUNTAIN_SNOW_COLOR_R:
-		return TEXT("kFLD_SHADING_MOUNTAIN_SNOW_COLOR_R");
-	case kFLD_SHADING_MOUNTAIN_SNOW_COLOR_G:
-		return TEXT("kFLD_SHADING_MOUNTAIN_SNOW_COLOR_G");
-	case kFLD_SHADING_MOUNTAIN_SNOW_COLOR_B:
-		return TEXT("kFLD_SHADING_MOUNTAIN_SNOW_COLOR_B");
-	case kFLD_SHADING_MOUNTAIN_ROCK_COLOR_R:
-		return TEXT("kFLD_SHADING_MOUNTAIN_ROCK_COLOR_R");
-	case kFLD_SHADING_MOUNTAIN_ROCK_COLOR_G:
-		return TEXT("kFLD_SHADING_MOUNTAIN_ROCK_COLOR_G");
-	case kFLD_SHADING_MOUNTAIN_ROCK_COLOR_B:
-		return TEXT("kFLD_SHADING_MOUNTAIN_ROCK_COLOR_B");
-	case kFLD_SHADING_MOUNTAIN_AMPLITUDE:
-		return TEXT("kFLD_SHADING_MOUNTAIN_AMPLITUDE");
-	case kFLD_SHADING_MOUNTAIN_SNOW_ROUGHNESS:
-		return TEXT("kFLD_SHADING_MOUNTAIN_SNOW_ROUGHNESS");
-	case kFLD_SHADING_MOUNTAIN_ROCK_ROUGHNESS:
-		return TEXT("kFLD_SHADING_MOUNTAIN_ROCK_ROUGHNESS");
-	case kFLD_SHADING_MOUNTAIN_BOUNDARY_ROUGHNESS:
-		return TEXT("kFLD_SHADING_MOUNTAIN_BOUNDARY_ROUGHNESS");
-	case kFLD_SHADING_MOUNTAIN_SNOW_ALTITUDE:
-		return TEXT("kFLD_SHADING_MOUNTAIN_SNOW_ALTITUDE");
-	case kFLD_SHADING_MOUNTAIN_SNOW_DROPOFF:
-		return TEXT("kFLD_SHADING_MOUNTAIN_SNOW_DROPOFF");
-	case kFLD_SHADING_MOUNTAIN_SNOW_MAX_SLOPE:
-		return TEXT("kFLD_SHADING_MOUNTAIN_SNOW_MAX_SLOPE");
-	case kFLD_SHADING_MOUNTAIN_LEVEL_MAX:
-		return TEXT("kFLD_SHADING_MOUNTAIN_LEVEL_MAX");
-	case kFLD_SHADING_NOISE_AMPLITUDE:
-		return TEXT("kFLD_SHADING_NOISE_AMPLITUDE");
-	case kFLD_SHADING_NOISE_THRESHOLD:
-		return TEXT("kFLD_SHADING_NOISE_THRESHOLD");
-	case kFLD_SHADING_NOISE_ANIMATED:
-		return TEXT("kFLD_SHADING_NOISE_ANIMATED");
-	case kFLD_SHADING_NOISE_TIME:
-		return TEXT("kFLD_SHADING_NOISE_TIME");
-	case kFLD_SHADING_RAMP_U_WAVE:
-		return TEXT("kFLD_SHADING_RAMP_U_WAVE");
-	case kFLD_SHADING_RAMP_V_WAVE:
-		return TEXT("kFLD_SHADING_RAMP_V_WAVE");
-	case kFLD_SHADING_RAMP_NOISE:
-		return TEXT("kFLD_SHADING_RAMP_NOISE");
-	case kFLD_SHADING_RAMP_NOISE_FREQUENCY:
-		return TEXT("kFLD_SHADING_RAMP_NOISE_FREQUENCY");
-	case kFLD_SHADING_RAMP_HUE_NOISE:
-		return TEXT("kFLD_SHADING_RAMP_HUE_NOISE");
-	case kFLD_SHADING_RAMP_SAT_NOISE:
-		return TEXT("kFLD_SHADING_RAMP_SAT_NOISE");
-	case kFLD_SHADING_RAMP_VAL_NOISE:
-		return TEXT("kFLD_SHADING_RAMP_VAL_NOISE");
-	case kFLD_SHADING_RAMP_HUE_NOISE_FREQUENCY:
-		return TEXT("kFLD_SHADING_RAMP_HUE_NOISE_FREQUENCY");
-	case kFLD_SHADING_RAMP_SAT_NOISE_FREQUENCY:
-		return TEXT("kFLD_SHADING_RAMP_SAT_NOISE_FREQUENCY");
-	case kFLD_SHADING_RAMP_VAL_NOISE_FREQUENCY:
-		return TEXT("kFLD_SHADING_RAMP_VAL_NOISE_FREQUENCY");
-	case kFLD_SHADING_SCLOUD_COLOR1_R:
-		return TEXT("kFLD_SHADING_SCLOUD_COLOR1_R");
-	case kFLD_SHADING_SCLOUD_COLOR1_G:
-		return TEXT("kFLD_SHADING_SCLOUD_COLOR1_G");
-	case kFLD_SHADING_SCLOUD_COLOR1_B:
-		return TEXT("kFLD_SHADING_SCLOUD_COLOR1_B");
-	case kFLD_SHADING_SCLOUD_COLOR2_R:
-		return TEXT("kFLD_SHADING_SCLOUD_COLOR2_R");
-	case kFLD_SHADING_SCLOUD_COLOR2_G:
-		return TEXT("kFLD_SHADING_SCLOUD_COLOR2_G");
-	case kFLD_SHADING_SCLOUD_COLOR2_B:
-		return TEXT("kFLD_SHADING_SCLOUD_COLOR2_B");
-	case kFLD_SHADING_SCLOUD_CONTRAST:
-		return TEXT("kFLD_SHADING_SCLOUD_CONTRAST");
-	case kFLD_SHADING_SCLOUD_SOFT_EDGES:
-		return TEXT("kFLD_SHADING_SCLOUD_SOFT_EDGES");
-	case kFLD_SHADING_SCLOUD_TRANSPARENCY_RANGE:
-		return TEXT("kFLD_SHADING_SCLOUD_TRANSPARENCY_RANGE");
-	case kFLD_SHADING_SCLOUD_CENTER_THRESHOLD:
-		return TEXT("kFLD_SHADING_SCLOUD_CENTER_THRESHOLD");
-	case kFLD_SHADING_SCLOUD_EDGE_THRESHOLD:
-		return TEXT("kFLD_SHADING_SCLOUD_EDGE_THRESHOLD");
-	case kFLD_SHADING_SCLOUD_AMPLITUDE:
-		return TEXT("kFLD_SHADING_SCLOUD_AMPLITUDE");
-	case kFLD_SHADING_SCLOUD_RATIO:
-		return TEXT("kFLD_SHADING_SCLOUD_RATIO");
-	case kFLD_SHADING_SCLOUD_LEVEL_MIN:
-		return TEXT("kFLD_SHADING_SCLOUD_LEVEL_MIN");
-	case kFLD_SHADING_SCLOUD_LEVEL_MAX:
-		return TEXT("kFLD_SHADING_SCLOUD_LEVEL_MAX");
-	case kFLD_SHADING_SCLOUD_XRIPPLES:
-		return TEXT("kFLD_SHADING_SCLOUD_XRIPPLES");
-	case kFLD_SHADING_SCLOUD_YRIPPLES:
-		return TEXT("kFLD_SHADING_SCLOUD_YRIPPLES");
-	case kFLD_SHADING_SCLOUD_ZRIPPLES:
-		return TEXT("kFLD_SHADING_SCLOUD_ZRIPPLES");
-	case kFLD_SHADING_SFRACTAL_THRESHOLD:
-		return TEXT("kFLD_SHADING_SFRACTAL_THRESHOLD");
-	case kFLD_SHADING_SFRACTAL_AMPLITUDE:
-		return TEXT("kFLD_SHADING_SFRACTAL_AMPLITUDE");
-	case kFLD_SHADING_SFRACTAL_RATIO:
-		return TEXT("kFLD_SHADING_SFRACTAL_RATIO");
-	case kFLD_SHADING_SFRACTAL_LEVEL_MIN:
-		return TEXT("kFLD_SHADING_SFRACTAL_LEVEL_MIN");
-	case kFLD_SHADING_SFRACTAL_LEVEL_MAX:
-		return TEXT("kFLD_SHADING_SFRACTAL_LEVEL_MAX");
-	case kFLD_SHADING_SFRACTAL_XRIPPLES:
-		return TEXT("kFLD_SHADING_SFRACTAL_XRIPPLES");
-	case kFLD_SHADING_SFRACTAL_YRIPPLES:
-		return TEXT("kFLD_SHADING_SFRACTAL_YRIPPLES");
-	case kFLD_SHADING_SFRACTAL_ZRIPPLES:
-		return TEXT("kFLD_SHADING_SFRACTAL_ZRIPPLES");
-	case kFLD_SHADING_SKY_TOTAL_BRIGHTNESS:
-		return TEXT("kFLD_SHADING_SKY_TOTAL_BRIGHTNESS");
-	case kFLD_SHADING_SKY_SUN_BRIGHTNESS_R:
-		return TEXT("kFLD_SHADING_SKY_SUN_BRIGHTNESS_R");
-	case kFLD_SHADING_SKY_SUN_BRIGHTNESS_G:
-		return TEXT("kFLD_SHADING_SKY_SUN_BRIGHTNESS_G");
-	case kFLD_SHADING_SKY_SUN_BRIGHTNESS_B:
-		return TEXT("kFLD_SHADING_SKY_SUN_BRIGHTNESS_B");
-	case kFLD_SHADING_SKY_SUN_HALO_BRIGHTNESS_R:
-		return TEXT("kFLD_SHADING_SKY_SUN_HALO_BRIGHTNESS_R");
-	case kFLD_SHADING_SKY_SUN_HALO_BRIGHTNESS_G:
-		return TEXT("kFLD_SHADING_SKY_SUN_HALO_BRIGHTNESS_G");
-	case kFLD_SHADING_SKY_SUN_HALO_BRIGHTNESS_B:
-		return TEXT("kFLD_SHADING_SKY_SUN_HALO_BRIGHTNESS_B");
-	case kFLD_SHADING_SKY_SUN_ELEVATION:
-		return TEXT("kFLD_SHADING_SKY_SUN_ELEVATION");
-	case kFLD_SHADING_SKY_SUN_AZIMUTH:
-		return TEXT("kFLD_SHADING_SKY_SUN_AZIMUTH");
-	case kFLD_SHADING_SKY_SUN_SIZE:
-		return TEXT("kFLD_SHADING_SKY_SUN_SIZE");
-	case kFLD_SHADING_SKY_SUN_BLUR:
-		return TEXT("kFLD_SHADING_SKY_SUN_BLUR");
-	case kFLD_SHADING_SKY_SKY_BRIGHTNESS_R:
-		return TEXT("kFLD_SHADING_SKY_SKY_BRIGHTNESS_R");
-	case kFLD_SHADING_SKY_SKY_BRIGHTNESS_G:
-		return TEXT("kFLD_SHADING_SKY_SKY_BRIGHTNESS_G");
-	case kFLD_SHADING_SKY_SKY_BRIGHTNESS_B:
-		return TEXT("kFLD_SHADING_SKY_SKY_BRIGHTNESS_B");
-	case kFLD_SHADING_SKY_AIR_DENSITY:
-		return TEXT("kFLD_SHADING_SKY_AIR_DENSITY");
-	case kFLD_SHADING_SKY_DUST_DENSITY:
-		return TEXT("kFLD_SHADING_SKY_DUST_DENSITY");
-	case kFLD_SHADING_SKY_SKY_THICKNESS:
-		return TEXT("kFLD_SHADING_SKY_SKY_THICKNESS");
-	case kFLD_SHADING_SKY_SKY_RADIUS:
-		return TEXT("kFLD_SHADING_SKY_SKY_RADIUS");
-	case kFLD_SHADING_SKY_HAS_FLOOR:
-		return TEXT("kFLD_SHADING_SKY_HAS_FLOOR");
-	case kFLD_SHADING_SKY_FLOOR_TEXTURE_R:
-		return TEXT("kFLD_SHADING_SKY_FLOOR_TEXTURE_R");
-	case kFLD_SHADING_SKY_FLOOR_TEXTURE_G:
-		return TEXT("kFLD_SHADING_SKY_FLOOR_TEXTURE_G");
-	case kFLD_SHADING_SKY_FLOOR_TEXTURE_B:
-		return TEXT("kFLD_SHADING_SKY_FLOOR_TEXTURE_B");
-	case kFLD_SHADING_SKY_CLOUD_BRIGHTNESS_R:
-		return TEXT("kFLD_SHADING_SKY_CLOUD_BRIGHTNESS_R");
-	case kFLD_SHADING_SKY_CLOUD_BRIGHTNESS_G:
-		return TEXT("kFLD_SHADING_SKY_CLOUD_BRIGHTNESS_G");
-	case kFLD_SHADING_SKY_CLOUD_BRIGHTNESS_B:
-		return TEXT("kFLD_SHADING_SKY_CLOUD_BRIGHTNESS_B");
-	case kFLD_SHADING_SKY_SUNSET_BRIGHTNESS_R:
-		return TEXT("kFLD_SHADING_SKY_SUNSET_BRIGHTNESS_R");
-	case kFLD_SHADING_SKY_SUNSET_BRIGHTNESS_G:
-		return TEXT("kFLD_SHADING_SKY_SUNSET_BRIGHTNESS_G");
-	case kFLD_SHADING_SKY_SUNSET_BRIGHTNESS_B:
-		return TEXT("kFLD_SHADING_SKY_SUNSET_BRIGHTNESS_B");
-	case kFLD_SHADING_SKY_CLOUD_DENSITY:
-		return TEXT("kFLD_SHADING_SKY_CLOUD_DENSITY");
-	case kFLD_SHADING_SKY_CLOUD_THRESHOLD:
-		return TEXT("kFLD_SHADING_SKY_CLOUD_THRESHOLD");
-	case kFLD_SHADING_SKY_CLOUD_POWER:
-		return TEXT("kFLD_SHADING_SKY_CLOUD_POWER");
-	case kFLD_SHADING_SKY_CLOUD_ALTITUDE:
-		return TEXT("kFLD_SHADING_SKY_CLOUD_ALTITUDE");
-	case kFLD_SHADING_SKY_CLOUD_HALO_SIZE:
-		return TEXT("kFLD_SHADING_SKY_CLOUD_HALO_SIZE");
-	case kFLD_SHADING_SKY_SKY_SAMPLES:
-		return TEXT("kFLD_SHADING_SKY_SKY_SAMPLES");
-	case kFLD_SHADING_SKY_FLOOR_SAMPLES:
-		return TEXT("kFLD_SHADING_SKY_FLOOR_SAMPLES");
-	case kFLD_SHADING_SKY_CLOUD_SAMPLES:
-		return TEXT("kFLD_SHADING_SKY_CLOUD_SAMPLES");
-	case kFLD_SHADING_STENCIL_EDGE_BLEND:
-		return TEXT("kFLD_SHADING_STENCIL_EDGE_BLEND");
-	case kFLD_SHADING_STENCIL_MASK_LEVEL:
-		return TEXT("kFLD_SHADING_STENCIL_MASK_LEVEL");
-	case kFLD_SHADING_STENCIL_MASK_BLUR:
-		return TEXT("kFLD_SHADING_STENCIL_MASK_BLUR");
-	case kFLD_SHADING_STENCIL_KEY_MASKING:
-		return TEXT("kFLD_SHADING_STENCIL_KEY_MASKING");
-	case kFLD_SHADING_STENCIL_POSITIVE_KEY:
-		return TEXT("kFLD_SHADING_STENCIL_POSITIVE_KEY");
-	case kFLD_SHADING_STENCIL_COLOR_KEY_R:
-		return TEXT("kFLD_SHADING_STENCIL_COLOR_KEY_R");
-	case kFLD_SHADING_STENCIL_COLOR_KEY_G:
-		return TEXT("kFLD_SHADING_STENCIL_COLOR_KEY_G");
-	case kFLD_SHADING_STENCIL_COLOR_KEY_B:
-		return TEXT("kFLD_SHADING_STENCIL_COLOR_KEY_B");
-	case kFLD_SHADING_STENCIL_HUE_RANGE:
-		return TEXT("kFLD_SHADING_STENCIL_HUE_RANGE");
-	case kFLD_SHADING_STENCIL_SAT_RANGE:
-		return TEXT("kFLD_SHADING_STENCIL_SAT_RANGE");
-	case kFLD_SHADING_STENCIL_VAL_RANGE:
-		return TEXT("kFLD_SHADING_STENCIL_VAL_RANGE");
-	case kFLD_SHADING_STENCIL_THRESHOLD:
-		return TEXT("kFLD_SHADING_STENCIL_THRESHOLD");
-	case kFLD_SHADING_SFILE_XAMPLITUDE:
-		return TEXT("kFLD_SHADING_SFILE_XAMPLITUDE");
-	case kFLD_SHADING_SFILE_YAMPLITUDE:
-		return TEXT("kFLD_SHADING_SFILE_YAMPLITUDE");
-	case kFLD_SHADING_SFILE_RATIO:
-		return TEXT("kFLD_SHADING_SFILE_RATIO");
-	case kFLD_SHADING_SFILE_LEVEL_MIN:
-		return TEXT("kFLD_SHADING_SFILE_LEVEL_MIN");
-	case kFLD_SHADING_SFILE_LEVEL_MAX:
-		return TEXT("kFLD_SHADING_SFILE_LEVEL_MAX");
-	case kFLD_SHADING_SFILE_XRIPPLES:
-		return TEXT("kFLD_SHADING_SFILE_XRIPPLES");
-	case kFLD_SHADING_SFILE_YRIPPLES:
-		return TEXT("kFLD_SHADING_SFILE_YRIPPLES");
-	case kFLD_SHADING_SFILE_ZRIPPLES:
-		return TEXT("kFLD_SHADING_SFILE_ZRIPPLES");
-	case kFLD_SHADING_SMARBLE_FILLER_COLOR_R:
-		return TEXT("kFLD_SHADING_SMARBLE_FILLER_COLOR_R");
-	case kFLD_SHADING_SMARBLE_FILLER_COLOR_G:
-		return TEXT("kFLD_SHADING_SMARBLE_FILLER_COLOR_G");
-	case kFLD_SHADING_SMARBLE_FILLER_COLOR_B:
-		return TEXT("kFLD_SHADING_SMARBLE_FILLER_COLOR_B");
-	case kFLD_SHADING_SMARBLE_VEIN_COLOR_R:
-		return TEXT("kFLD_SHADING_SMARBLE_VEIN_COLOR_R");
-	case kFLD_SHADING_SMARBLE_VEIN_COLOR_G:
-		return TEXT("kFLD_SHADING_SMARBLE_VEIN_COLOR_G");
-	case kFLD_SHADING_SMARBLE_VEIN_COLOR_B:
-		return TEXT("kFLD_SHADING_SMARBLE_VEIN_COLOR_B");
-	case kFLD_SHADING_SMARBLE_VEIN_WIDTH:
-		return TEXT("kFLD_SHADING_SMARBLE_VEIN_WIDTH");
-	case kFLD_SHADING_SMARBLE_DIFFUSION:
-		return TEXT("kFLD_SHADING_SMARBLE_DIFFUSION");
-	case kFLD_SHADING_SMARBLE_CONTRAST:
-		return TEXT("kFLD_SHADING_SMARBLE_CONTRAST");
-	case kFLD_SHADING_SMARBLE_AMPLITUDE:
-		return TEXT("kFLD_SHADING_SMARBLE_AMPLITUDE");
-	case kFLD_SHADING_SMARBLE_RATIO:
-		return TEXT("kFLD_SHADING_SMARBLE_RATIO");
-	case kFLD_SHADING_SMARBLE_LEVEL_MIN:
-		return TEXT("kFLD_SHADING_SMARBLE_LEVEL_MIN");
-	case kFLD_SHADING_SMARBLE_LEVEL_MAX:
-		return TEXT("kFLD_SHADING_SMARBLE_LEVEL_MAX");
-	case kFLD_SHADING_SMARBLE_XRIPPLES:
-		return TEXT("kFLD_SHADING_SMARBLE_XRIPPLES");
-	case kFLD_SHADING_SMARBLE_YRIPPLES:
-		return TEXT("kFLD_SHADING_SMARBLE_YRIPPLES");
-	case kFLD_SHADING_SMARBLE_ZRIPPLES:
-		return TEXT("kFLD_SHADING_SMARBLE_ZRIPPLES");
-	case kFLD_SHADING_SNOW_COLOR_R:
-		return TEXT("kFLD_SHADING_SNOW_COLOR_R");
-	case kFLD_SHADING_SNOW_COLOR_G:
-		return TEXT("kFLD_SHADING_SNOW_COLOR_G");
-	case kFLD_SHADING_SNOW_COLOR_B:
-		return TEXT("kFLD_SHADING_SNOW_COLOR_B");
-	case kFLD_SHADING_SNOW_SURFACE_COLOR_R:
-		return TEXT("kFLD_SHADING_SNOW_SURFACE_COLOR_R");
-	case kFLD_SHADING_SNOW_SURFACE_COLOR_G:
-		return TEXT("kFLD_SHADING_SNOW_SURFACE_COLOR_G");
-	case kFLD_SHADING_SNOW_SURFACE_COLOR_B:
-		return TEXT("kFLD_SHADING_SNOW_SURFACE_COLOR_B");
-	case kFLD_SHADING_SNOW_THRESHOLD:
-		return TEXT("kFLD_SHADING_SNOW_THRESHOLD");
-	case kFLD_SHADING_SNOW_DEPTH_DECAY:
-		return TEXT("kFLD_SHADING_SNOW_DEPTH_DECAY");
-	case kFLD_SHADING_SNOW_THICKNESS:
-		return TEXT("kFLD_SHADING_SNOW_THICKNESS");
-	case kFLD_SHADING_SPHERE_SHEAR_U:
-		return TEXT("kFLD_SHADING_SPHERE_SHEAR_U");
-	case kFLD_SHADING_SPHERE_SHEAR_V:
-		return TEXT("kFLD_SHADING_SPHERE_SHEAR_V");
-	case kFLD_SHADING_SPHERE_FLIP:
-		return TEXT("kFLD_SHADING_SPHERE_FLIP");
-	case kFLD_SHADING_SROCK_COLOR1_R:
-		return TEXT("kFLD_SHADING_SROCK_COLOR1_R");
-	case kFLD_SHADING_SROCK_COLOR1_G:
-		return TEXT("kFLD_SHADING_SROCK_COLOR1_G");
-	case kFLD_SHADING_SROCK_COLOR1_B:
-		return TEXT("kFLD_SHADING_SROCK_COLOR1_B");
-	case kFLD_SHADING_SROCK_COLOR2_R:
-		return TEXT("kFLD_SHADING_SROCK_COLOR2_R");
-	case kFLD_SHADING_SROCK_COLOR2_G:
-		return TEXT("kFLD_SHADING_SROCK_COLOR2_G");
-	case kFLD_SHADING_SROCK_COLOR2_B:
-		return TEXT("kFLD_SHADING_SROCK_COLOR2_B");
-	case kFLD_SHADING_SROCK_GRAIN_SIZE:
-		return TEXT("kFLD_SHADING_SROCK_GRAIN_SIZE");
-	case kFLD_SHADING_SROCK_DIFFUSION:
-		return TEXT("kFLD_SHADING_SROCK_DIFFUSION");
-	case kFLD_SHADING_SROCK_MIX_RATIO:
-		return TEXT("kFLD_SHADING_SROCK_MIX_RATIO");
-	case kFLD_SHADING_LEATHER_CELL_COLOR_R:
-		return TEXT("kFLD_SHADING_LEATHER_CELL_COLOR_R");
-	case kFLD_SHADING_LEATHER_CELL_COLOR_G:
-		return TEXT("kFLD_SHADING_LEATHER_CELL_COLOR_G");
-	case kFLD_SHADING_LEATHER_CELL_COLOR_B:
-		return TEXT("kFLD_SHADING_LEATHER_CELL_COLOR_B");
-	case kFLD_SHADING_LEATHER_CREASE_COLOR_R:
-		return TEXT("kFLD_SHADING_LEATHER_CREASE_COLOR_R");
-	case kFLD_SHADING_LEATHER_CREASE_COLOR_G:
-		return TEXT("kFLD_SHADING_LEATHER_CREASE_COLOR_G");
-	case kFLD_SHADING_LEATHER_CREASE_COLOR_B:
-		return TEXT("kFLD_SHADING_LEATHER_CREASE_COLOR_B");
-	case kFLD_SHADING_LEATHER_CELL_SIZE:
-		return TEXT("kFLD_SHADING_LEATHER_CELL_SIZE");
-	case kFLD_SHADING_LEATHER_DENSITY:
-		return TEXT("kFLD_SHADING_LEATHER_DENSITY");
-	case kFLD_SHADING_LEATHER_SPOTTYNESS:
-		return TEXT("kFLD_SHADING_LEATHER_SPOTTYNESS");
-	case kFLD_SHADING_LEATHER_RANDOMNESS:
-		return TEXT("kFLD_SHADING_LEATHER_RANDOMNESS");
-	case kFLD_SHADING_LEATHER_THRESHOLD:
-		return TEXT("kFLD_SHADING_LEATHER_THRESHOLD");
-	case kFLD_SHADING_LEATHER_CREASES:
-		return TEXT("kFLD_SHADING_LEATHER_CREASES");
-	case kFLD_SHADING_GRANITE_COLOR1_R:
-		return TEXT("kFLD_SHADING_GRANITE_COLOR1_R");
-	case kFLD_SHADING_GRANITE_COLOR1_G:
-		return TEXT("kFLD_SHADING_GRANITE_COLOR1_G");
-	case kFLD_SHADING_GRANITE_COLOR1_B:
-		return TEXT("kFLD_SHADING_GRANITE_COLOR1_B");
-	case kFLD_SHADING_GRANITE_COLOR2_R:
-		return TEXT("kFLD_SHADING_GRANITE_COLOR2_R");
-	case kFLD_SHADING_GRANITE_COLOR2_G:
-		return TEXT("kFLD_SHADING_GRANITE_COLOR2_G");
-	case kFLD_SHADING_GRANITE_COLOR2_B:
-		return TEXT("kFLD_SHADING_GRANITE_COLOR2_B");
-	case kFLD_SHADING_GRANITE_COLOR3_R:
-		return TEXT("kFLD_SHADING_GRANITE_COLOR3_R");
-	case kFLD_SHADING_GRANITE_COLOR3_G:
-		return TEXT("kFLD_SHADING_GRANITE_COLOR3_G");
-	case kFLD_SHADING_GRANITE_COLOR3_B:
-		return TEXT("kFLD_SHADING_GRANITE_COLOR3_B");
-	case kFLD_SHADING_GRANITE_FILLER_COLOR_R:
-		return TEXT("kFLD_SHADING_GRANITE_FILLER_COLOR_R");
-	case kFLD_SHADING_GRANITE_FILLER_COLOR_G:
-		return TEXT("kFLD_SHADING_GRANITE_FILLER_COLOR_G");
-	case kFLD_SHADING_GRANITE_FILLER_COLOR_B:
-		return TEXT("kFLD_SHADING_GRANITE_FILLER_COLOR_B");
-	case kFLD_SHADING_GRANITE_CELL_SIZE:
-		return TEXT("kFLD_SHADING_GRANITE_CELL_SIZE");
-	case kFLD_SHADING_GRANITE_DENSITY:
-		return TEXT("kFLD_SHADING_GRANITE_DENSITY");
-	case kFLD_SHADING_GRANITE_MIX_RATIO:
-		return TEXT("kFLD_SHADING_GRANITE_MIX_RATIO");
-	case kFLD_SHADING_GRANITE_SPOTTYNESS:
-		return TEXT("kFLD_SHADING_GRANITE_SPOTTYNESS");
-	case kFLD_SHADING_GRANITE_RANDOMNESS:
-		return TEXT("kFLD_SHADING_GRANITE_RANDOMNESS");
-	case kFLD_SHADING_GRANITE_THRESHOLD:
-		return TEXT("kFLD_SHADING_GRANITE_THRESHOLD");
-	case kFLD_SHADING_GRANITE_CREASES:
-		return TEXT("kFLD_SHADING_GRANITE_CREASES");
-	case kFLD_SHADING_SWOOD_FILLER_COLOR_R:
-		return TEXT("kFLD_SHADING_SWOOD_FILLER_COLOR_R");
-	case kFLD_SHADING_SWOOD_FILLER_COLOR_G:
-		return TEXT("kFLD_SHADING_SWOOD_FILLER_COLOR_G");
-	case kFLD_SHADING_SWOOD_FILLER_COLOR_B:
-		return TEXT("kFLD_SHADING_SWOOD_FILLER_COLOR_B");
-	case kFLD_SHADING_SWOOD_VEIN_COLOR_R:
-		return TEXT("kFLD_SHADING_SWOOD_VEIN_COLOR_R");
-	case kFLD_SHADING_SWOOD_VEIN_COLOR_G:
-		return TEXT("kFLD_SHADING_SWOOD_VEIN_COLOR_G");
-	case kFLD_SHADING_SWOOD_VEIN_COLOR_B:
-		return TEXT("kFLD_SHADING_SWOOD_VEIN_COLOR_B");
-	case kFLD_SHADING_SWOOD_VEIN_SPREAD:
-		return TEXT("kFLD_SHADING_SWOOD_VEIN_SPREAD");
-	case kFLD_SHADING_SWOOD_LAYER_SIZE:
-		return TEXT("kFLD_SHADING_SWOOD_LAYER_SIZE");
-	case kFLD_SHADING_SWOOD_RANDOMNESS:
-		return TEXT("kFLD_SHADING_SWOOD_RANDOMNESS");
-	case kFLD_SHADING_SWOOD_AGE:
-		return TEXT("kFLD_SHADING_SWOOD_AGE");
-	case kFLD_SHADING_SWOOD_GRAIN_COLOR_R:
-		return TEXT("kFLD_SHADING_SWOOD_GRAIN_COLOR_R");
-	case kFLD_SHADING_SWOOD_GRAIN_COLOR_G:
-		return TEXT("kFLD_SHADING_SWOOD_GRAIN_COLOR_G");
-	case kFLD_SHADING_SWOOD_GRAIN_COLOR_B:
-		return TEXT("kFLD_SHADING_SWOOD_GRAIN_COLOR_B");
-	case kFLD_SHADING_SWOOD_GRAIN_CONTRAST:
-		return TEXT("kFLD_SHADING_SWOOD_GRAIN_CONTRAST");
-	case kFLD_SHADING_SWOOD_GRAIN_SPACING:
-		return TEXT("kFLD_SHADING_SWOOD_GRAIN_SPACING");
-	case kFLD_SHADING_SWOOD_CENTER_U:
-		return TEXT("kFLD_SHADING_SWOOD_CENTER_U");
-	case kFLD_SHADING_SWOOD_CENTER_V:
-		return TEXT("kFLD_SHADING_SWOOD_CENTER_V");
-	case kFLD_SHADING_SWOOD_XAMPLITUDE:
-		return TEXT("kFLD_SHADING_SWOOD_XAMPLITUDE");
-	case kFLD_SHADING_SWOOD_YAMPLITUDE:
-		return TEXT("kFLD_SHADING_SWOOD_YAMPLITUDE");
-	case kFLD_SHADING_SWOOD_RATIO:
-		return TEXT("kFLD_SHADING_SWOOD_RATIO");
-	case kFLD_SHADING_SWOOD_LEVEL_MIN:
-		return TEXT("kFLD_SHADING_SWOOD_LEVEL_MIN");
-	case kFLD_SHADING_SWOOD_LEVEL_MAX:
-		return TEXT("kFLD_SHADING_SWOOD_LEVEL_MAX");
-	case kFLD_SHADING_SWOOD_XRIPPLES:
-		return TEXT("kFLD_SHADING_SWOOD_XRIPPLES");
-	case kFLD_SHADING_SWOOD_YRIPPLES:
-		return TEXT("kFLD_SHADING_SWOOD_YRIPPLES");
-	case kFLD_SHADING_SWOOD_ZRIPPLES:
-		return TEXT("kFLD_SHADING_SWOOD_ZRIPPLES");
-	case kFLD_SHADING_VOLUME_FROM:
-		return TEXT("kFLD_SHADING_VOLUME_FROM");
-	case kFLD_SHADING_VOLUME_TO:
-		return TEXT("kFLD_SHADING_VOLUME_TO");
-	case kFLD_SHADING_WATER_NUMWAVES:
-		return TEXT("kFLD_SHADING_WATER_NUMWAVES");
-	case kFLD_SHADING_WATER_WAVE_TIME:
-		return TEXT("kFLD_SHADING_WATER_WAVE_TIME");
-	case kFLD_SHADING_WATER_WAVE_VELOCITY:
-		return TEXT("kFLD_SHADING_WATER_WAVE_VELOCITY");
-	case kFLD_SHADING_WATER_WAVE_AMPLITUDE:
-		return TEXT("kFLD_SHADING_WATER_WAVE_AMPLITUDE");
-	case kFLD_SHADING_WATER_FREQUENCY:
-		return TEXT("kFLD_SHADING_WATER_FREQUENCY");
-	case kFLD_SHADING_WATER_SUB_FREQUENCY:
-		return TEXT("kFLD_SHADING_WATER_SUB_FREQUENCY");
-	case kFLD_SHADING_WATER_SMOOTHNESS:
-		return TEXT("kFLD_SHADING_WATER_SMOOTHNESS");
-	case kFLD_SHADING_WATER_WIND_U:
-		return TEXT("kFLD_SHADING_WATER_WIND_U");
-	case kFLD_SHADING_WATER_WIND_V:
-		return TEXT("kFLD_SHADING_WATER_WIND_V");
-	case kFLD_SHADING_WATER_RIPPLE_TIME:
-		return TEXT("kFLD_SHADING_WATER_RIPPLE_TIME");
-	case kFLD_SHADING_WATER_RIPPLE_FREQUENCY:
-		return TEXT("kFLD_SHADING_WATER_RIPPLE_FREQUENCY");
-	case kFLD_SHADING_WATER_RIPPLE_AMPLITUDE:
-		return TEXT("kFLD_SHADING_WATER_RIPPLE_AMPLITUDE");
-	case kFLD_SHADING_WATER_RIPPLE_DROP_SIZE:
-		return TEXT("kFLD_SHADING_WATER_RIPPLE_DROP_SIZE");
-	case kFLD_SHADING_WATER_RIPPLE_U_ORIGIN:
-		return TEXT("kFLD_SHADING_WATER_RIPPLE_U_ORIGIN");
-	case kFLD_SHADING_WATER_RIPPLE_V_ORIGIN:
-		return TEXT("kFLD_SHADING_WATER_RIPPLE_V_ORIGIN");
-	case kFLD_SHADING_WATER_RIPPLE_GROUP_VELOCITY:
-		return TEXT("kFLD_SHADING_WATER_RIPPLE_GROUP_VELOCITY");
-	case kFLD_SHADING_WATER_RIPPLE_PHASE_VELOCITY:
-		return TEXT("kFLD_SHADING_WATER_RIPPLE_PHASE_VELOCITY");
-	case kFLD_SHADING_WATER_RIPPLE_SPREAD_START:
-		return TEXT("kFLD_SHADING_WATER_RIPPLE_SPREAD_START");
-	case kFLD_SHADING_WATER_RIPPLE_SPREAD_RATE:
-		return TEXT("kFLD_SHADING_WATER_RIPPLE_SPREAD_RATE");
-	case kFLD_SHADING_WATER_REFLECTION_BOUNDARY:
-		return TEXT("kFLD_SHADING_WATER_REFLECTION_BOUNDARY");
-	case kFLD_SHADING_WATER_UMIN:
-		return TEXT("kFLD_SHADING_WATER_UMIN");
-	case kFLD_SHADING_WATER_UMAX:
-		return TEXT("kFLD_SHADING_WATER_UMAX");
-	case kFLD_SHADING_WATER_VMIN:
-		return TEXT("kFLD_SHADING_WATER_VMIN");
-	case kFLD_SHADING_WATER_VMAX:
-		return TEXT("kFLD_SHADING_WATER_VMAX");
-	case kFLD_SHADING_COMMON_GLOW_INTENSITY:
-		return TEXT("kFLD_SHADING_COMMON_GLOW_INTENSITY");
-	case kFLD_SHADING_COMMON_HIDE_GLOW_SOURCE:
-		return TEXT("kFLD_SHADING_COMMON_HIDE_GLOW_SOURCE");
-	case kFLD_SHADING_COMMON_SURFACE_WIDTH:
-		return TEXT("kFLD_SHADING_COMMON_SURFACE_WIDTH");
-	case kFLD_SHADING_COMMON_TRANSPARENCY_DEPTH:
-		return TEXT("kFLD_SHADING_COMMON_TRANSPARENCY_DEPTH");
-	case kFLD_SHADING_COMMON_CHROMATIC_ABBERATION:
-		return TEXT("kFLD_SHADING_COMMON_CHROMATIC_ABBERATION");
-	case kFLD_SHADING_COMMON_REFRACTION_JITTER:
-		return TEXT("kFLD_SHADING_COMMON_REFRACTION_JITTER");
-	case kFLD_SHADING_COMMON_REFRACTION_SAMPLES:
-		return TEXT("kFLD_SHADING_COMMON_REFRACTION_SAMPLES");
-	case kFLD_SHADING_COMMON_USE_FILL_COLOR:
-		return TEXT("kFLD_SHADING_COMMON_USE_FILL_COLOR");
-	case kFLD_SHADING_COMMON_FILL_COLOR_R:
-		return TEXT("kFLD_SHADING_COMMON_FILL_COLOR_R");
-	case kFLD_SHADING_COMMON_FILL_COLOR_G:
-		return TEXT("kFLD_SHADING_COMMON_FILL_COLOR_G");
-	case kFLD_SHADING_COMMON_FILL_COLOR_B:
-		return TEXT("kFLD_SHADING_COMMON_FILL_COLOR_B");
-	case kFLD_SHADING_COMMON_LINE_THICKNESS:
-		return TEXT("kFLD_SHADING_COMMON_LINE_THICKNESS");
-	case kFLD_SHADING_COMMON_U_PATCH_LINES:
-		return TEXT("kFLD_SHADING_COMMON_U_PATCH_LINES");
-	case kFLD_SHADING_COMMON_V_PATCH_LINES:
-		return TEXT("kFLD_SHADING_COMMON_V_PATCH_LINES");
-	case kFLD_SHADING_FOG_ACTIVE:
-		return TEXT("kFLD_SHADING_FOG_ACTIVE");
-	case kFLD_SHADING_COMMON_TEXTURE_STAGGER:
-		return TEXT("kFLD_SHADING_COMMON_TEXTURE_STAGGER");
-	case kFLD_SHADING_COMMON_TEXTURE_NOT_USED:
-		return TEXT("kFLD_SHADING_COMMON_TEXTURE_NOT_USED");
-	case kFLD_SHADING_COMMON_TEXTURE_INVERT:
-		return TEXT("kFLD_SHADING_COMMON_TEXTURE_INVERT");
-	case kFLD_SHADING_COMMON_TEXTURE_CHORD_LENGTH:
-		return TEXT("kFLD_SHADING_COMMON_TEXTURE_CHORD_LENGTH");
-	case kFLD_SHADING_COMMON_TEXTURE_WRAP:
-		return TEXT("kFLD_SHADING_COMMON_TEXTURE_WRAP");
-	case kFLD_SHADING_COMMON_TEXTURE_BLEND:
-		return TEXT("kFLD_SHADING_COMMON_TEXTURE_BLEND");
-	case kFLD_SHADING_RAMP_TYPE:
-		return TEXT("kFLD_SHADING_RAMP_Type");
-	case kFLD_SHADING_RAMP_INTERPOLATION:
-		return TEXT("kFLD_SHADING_RAMP_INTERPOLATION");
-	case kFLD_SHADING_HIGHLIGHT_REPEATS:
-		return TEXT("kFLD_SHADING_HIGHLIGHT_REPEATS");
-	case kFLD_SHADING_FILE_USE_EXTENSION:
-		return TEXT("kFLD_SHADING_FILE_USE_EXTENSION");
-	case kFLD_SHADING_FILE_FRAME_EXTENSION:
-		return TEXT("kFLD_SHADING_FILE_FRAME_EXTENSION");
-	case kFLD_SHADING_SFILE_PROJECTION:
-		return TEXT("kFLD_SHADING_SFILE_PROJECTION");
-	case kFLD_SHADING_SFILE_STAGGER:
-		return TEXT("kFLD_SHADING_SFILE_STAGGER");
-	case kFLD_SHADING_CURVATURE_REPEATS:
-		return TEXT("kFLD_SHADING_CURVATURE_REPEATS");
-	case kFLD_SHADING_SHADERGLOW_GLOW_TYPE:
-		return TEXT("kFLD_SHADING_SHADERGLOW_GLOW_Type");
-	case kFLD_SHADING_SHADERGLOW_HALO_TYPE:
-		return TEXT("kFLD_SHADING_SHADERGLOW_HALO_Type");
-	case kFLD_SHADING_SHADERGLOW_QUALITY:
-		return TEXT("kFLD_SHADING_SHADERGLOW_QUALITY");
-	case kFLD_SHADING_SHADERGLOW_THRESHOLD:
-		return TEXT("kFLD_SHADING_SHADERGLOW_THRESHOLD");
-	case kFLD_SHADING_SHADERGLOW_GLOW_COLOR_R:
-		return TEXT("kFLD_SHADING_SHADERGLOW_GLOW_COLOR_R");
-	case kFLD_SHADING_SHADERGLOW_GLOW_COLOR_G:
-		return TEXT("kFLD_SHADING_SHADERGLOW_GLOW_COLOR_G");
-	case kFLD_SHADING_SHADERGLOW_GLOW_COLOR_B:
-		return TEXT("kFLD_SHADING_SHADERGLOW_GLOW_COLOR_B");
-	case kFLD_SHADING_SHADERGLOW_GLOW_INTENSITY:
-		return TEXT("kFLD_SHADING_SHADERGLOW_GLOW_INTENSITY");
-	case kFLD_SHADING_SHADERGLOW_GLOW_SPREAD:
-		return TEXT("kFLD_SHADING_SHADERGLOW_GLOW_SPREAD");
-	case kFLD_SHADING_SHADERGLOW_GLOW_ECCENTRICITY:
-		return TEXT("kFLD_SHADING_SHADERGLOW_GLOW_ECCENTRICITY");
-	case kFLD_SHADING_SHADERGLOW_GLOW_RADIAL_NOISE:
-		return TEXT("kFLD_SHADING_SHADERGLOW_GLOW_RADIAL_NOISE");
-	case kFLD_SHADING_SHADERGLOW_GLOW_STAR_LEVEL:
-		return TEXT("kFLD_SHADING_SHADERGLOW_GLOW_STAR_LEVEL");
-	case kFLD_SHADING_SHADERGLOW_GLOW_OPACITY:
-		return TEXT("kFLD_SHADING_SHADERGLOW_GLOW_OPACITY");
-	case kFLD_SHADING_SHADERGLOW_HALO_COLOR_R:
-		return TEXT("kFLD_SHADING_SHADERGLOW_HALO_COLOR_R");
-	case kFLD_SHADING_SHADERGLOW_HALO_COLOR_G:
-		return TEXT("kFLD_SHADING_SHADERGLOW_HALO_COLOR_G");
-	case kFLD_SHADING_SHADERGLOW_HALO_COLOR_B:
-		return TEXT("kFLD_SHADING_SHADERGLOW_HALO_COLOR_B");
-	case kFLD_SHADING_SHADERGLOW_HALO_SPREAD:
-		return TEXT("kFLD_SHADING_SHADERGLOW_HALO_SPREAD");
-	case kFLD_SHADING_SHADERGLOW_HALO_ECCENTRICITY:
-		return TEXT("kFLD_SHADING_SHADERGLOW_HALO_ECCENTRICITY");
-	case kFLD_SHADING_SHADERGLOW_HALO_RADIAL_NOISE:
-		return TEXT("kFLD_SHADING_SHADERGLOW_HALO_RADIAL_NOISE");
-	case kFLD_SHADING_SHADERGLOW_HALO_STAR_LEVEL:
-		return TEXT("kFLD_SHADING_SHADERGLOW_HALO_STAR_LEVEL");
-	case kFLD_SHADING_SHADERGLOW_HALO_LENS_FLARE:
-		return TEXT("kFLD_SHADING_SHADERGLOW_HALO_LENS_FLARE");
-	case kFLD_SHADING_SHADERGLOW_ROTATION:
-		return TEXT("kFLD_SHADING_SHADERGLOW_ROTATION");
-	case kFLD_SHADING_SHADERGLOW_RADIAL_NOISE_FREQUENCY:
-		return TEXT("kFLD_SHADING_SHADERGLOW_RADIAL_NOISE_FREQUENCY");
-	case kFLD_SHADING_SHADERGLOW_STAR_POINTS:
-		return TEXT("kFLD_SHADING_SHADERGLOW_STAR_POINTS");
-	case kFLD_SHADING_LENSEFFECT_FILM_GRAIN:
-		return TEXT("kFLD_SHADING_LENSEFFECT_FILM_GRAIN");
-	case kFLD_SHADING_LENSEFFECT_FILTER_R:
-		return TEXT("kFLD_SHADING_LENSEFFECT_FILTER_R");
-	case kFLD_SHADING_LENSEFFECT_FILTER_G:
-		return TEXT("kFLD_SHADING_LENSEFFECT_FILTER_G");
-	case kFLD_SHADING_LENSEFFECT_FILTER_B:
-		return TEXT("kFLD_SHADING_LENSEFFECT_FILTER_B");
-	case kFLD_SHADING_MASTERLIGHT_INTENSITY:
-		return TEXT("kFLD_SHADING_MASTERLIGHT_INTENSITY");
-	case kFLD_SHADING_MASTERLIGHT_COLOR_R:
-		return TEXT("kFLD_SHADING_MASTERLIGHT_COLOR_R");
-	case kFLD_SHADING_MASTERLIGHT_COLOR_G:
-		return TEXT("kFLD_SHADING_MASTERLIGHT_COLOR_G");
-	case kFLD_SHADING_MASTERLIGHT_COLOR_B:
-		return TEXT("kFLD_SHADING_MASTERLIGHT_COLOR_B");
-	case kFLD_SHADING_COMMON_TRANSLUCENCE:
-		return TEXT("kFLD_SHADING_COMMON_TRANSLUCENCE");
-	case kFLD_SHADING_COMMON_TRANSLUCENCE_DEPTH:
-		return TEXT("kFLD_SHADING_COMMON_TRANSLUCENCE_DEPTH");
-	case kFLD_SHADING_COMMON_OPACITY_DEPTH:
-		return TEXT("kFLD_SHADING_COMMON_OPACITY_DEPTH");
-	case kFLD_SHADING_COMMON_PSYS_EMISSION:
-		return TEXT("kFLD_SHADING_COMMON_PSYS_EMISSION");
-	case kFLD_SHADING_COMMON_PSYS_RENDER_TYPE:
-		return TEXT("kFLD_SHADING_COMMON_PSYS_RENDER_Type");
-	case kFLD_SHADING_COMMON_PSYS_MOTION_TYPE:
-		return TEXT("kFLD_SHADING_COMMON_PSYS_MOTION_Type");
-	case kFLD_SHADING_COMMON_PSYS_GLOW_INTENSITY:
-		return TEXT("kFLD_SHADING_COMMON_PSYS_GLOW_INTENSITY");
-	case kFLD_SHADING_COMMON_PSYS_PARENT_SHADING:
-		return TEXT("kFLD_SHADING_COMMON_PSYS_PARENT_SHADING");
-	case kFLD_SHADING_COMMON_PSYS_COLOR_R:
-		return TEXT("kFLD_SHADING_COMMON_PSYS_COLOR_R");
-	case kFLD_SHADING_COMMON_PSYS_COLOR_G:
-		return TEXT("kFLD_SHADING_COMMON_PSYS_COLOR_G");
-	case kFLD_SHADING_COMMON_PSYS_COLOR_B:
-		return TEXT("kFLD_SHADING_COMMON_PSYS_COLOR_B");
-	case kFLD_SHADING_COMMON_PSYS_INCANDESCENCE_R:
-		return TEXT("kFLD_SHADING_COMMON_PSYS_INCANDESCENCE_R");
-	case kFLD_SHADING_COMMON_PSYS_INCANDESCENCE_G:
-		return TEXT("kFLD_SHADING_COMMON_PSYS_INCANDESCENCE_G");
-	case kFLD_SHADING_COMMON_PSYS_INCANDESCENCE_B:
-		return TEXT("kFLD_SHADING_COMMON_PSYS_INCANDESCENCE_B");
-	case kFLD_SHADING_COMMON_PSYS_TRANSPARENCY_R:
-		return TEXT("kFLD_SHADING_COMMON_PSYS_TRANSPARENCY_R");
-	case kFLD_SHADING_COMMON_PSYS_TRANSPARENCY_G:
-		return TEXT("kFLD_SHADING_COMMON_PSYS_TRANSPARENCY_G");
-	case kFLD_SHADING_COMMON_PSYS_TRANSPARENCY_B:
-		return TEXT("kFLD_SHADING_COMMON_PSYS_TRANSPARENCY_B");
-	case kFLD_SHADING_COMMON_PSYS_TRANSPARENCY_METHOD:
-		return TEXT("kFLD_SHADING_COMMON_PSYS_TRANSPARENCY_METHOD");
-	case kFLD_SHADING_COMMON_PSYS_TRANSLUCENCE:
-		return TEXT("kFLD_SHADING_COMMON_PSYS_TRANSLUCENCE");
-	case kFLD_SHADING_COMMON_PSYS_BLOB_THRESHOLD:
-		return TEXT("kFLD_SHADING_COMMON_PSYS_BLOB_THRESHOLD");
-	case kFLD_SHADING_COMMON_PSYS_BLOB_NOISE:
-		return TEXT("kFLD_SHADING_COMMON_PSYS_BLOB_NOISE");
-	case kFLD_SHADING_COMMON_PSYS_BLOB_NOISE_FREQUENCY:
-		return TEXT("kFLD_SHADING_COMMON_PSYS_BLOB_NOISE_FREQUENCY");
-	case kFLD_SHADING_COMMON_PSYS_DENSITY:
-		return TEXT("kFLD_SHADING_COMMON_PSYS_DENSITY");
-	case kFLD_SHADING_COMMON_PSYS_SIZE:
-		return TEXT("kFLD_SHADING_COMMON_PSYS_SIZE");
-	case kFLD_SHADING_COMMON_PSYS_SIZE_RANGE:
-		return TEXT("kFLD_SHADING_COMMON_PSYS_SIZE_RANGE");
-	case kFLD_SHADING_COMMON_PSYS_BUOYANCY:
-		return TEXT("kFLD_SHADING_COMMON_PSYS_BUOYANCY");
-	case kFLD_SHADING_COMMON_PSYS_MASS:
-		return TEXT("kFLD_SHADING_COMMON_PSYS_MASS");
-	case kFLD_SHADING_COMMON_PSYS_ELASTICITY:
-		return TEXT("kFLD_SHADING_COMMON_PSYS_ELASTICITY");
-	case kFLD_SHADING_COMMON_PSYS_ATOMS_PER_SEC:
-		return TEXT("kFLD_SHADING_COMMON_PSYS_ATOMS_PER_SEC");
-	case kFLD_SHADING_COMMON_PSYS_SPEED:
-		return TEXT("kFLD_SHADING_COMMON_PSYS_SPEED");
-	case kFLD_SHADING_COMMON_PSYS_SPEED_RANGE:
-		return TEXT("kFLD_SHADING_COMMON_PSYS_SPEED_RANGE");
-	case kFLD_SHADING_COMMON_PSYS_LIFESPAN_MIN:
-		return TEXT("kFLD_SHADING_COMMON_PSYS_LIFESPAN_MIN");
-	case kFLD_SHADING_COMMON_PSYS_LIFESPAN_MAX:
-		return TEXT("kFLD_SHADING_COMMON_PSYS_LIFESPAN_MAX");
-	case kFLD_SHADING_COMMON_PSYS_RANDOMIZATION:
-		return TEXT("kFLD_SHADING_COMMON_PSYS_RANDOMIZATION");
-	case kFLD_SHADING_COMMON_PSYS_TIME_RANDOM:
-		return TEXT("kFLD_SHADING_COMMON_PSYS_TIME_RANDOM");
-	case kFLD_SHADING_COMMON_PSYS_COLLISIONS:
-		return TEXT("kFLD_SHADING_COMMON_PSYS_COLLISIONS");
-	case kFLD_SHADING_COMMON_PSYS_COLLISION_METHOD:
-		return TEXT("kFLD_SHADING_COMMON_PSYS_COLLISION_METHOD");
-	case kFLD_SHADING_COMMON_PSYS_GROW_METHOD:
-		return TEXT("kFLD_SHADING_COMMON_PSYS_GROW_METHOD");
-	case kFLD_SHADING_COMMON_PSYS_CYCLES:
-		return TEXT("kFLD_SHADING_COMMON_PSYS_CYCLES");
-	case kFLD_SHADING_COMMON_PSYS_ATTRACTION:
-		return TEXT("kFLD_SHADING_COMMON_PSYS_ATTRACTION");
-	case kFLD_SHADING_COMMON_PSYS_NOISE_ASPECT:
-		return TEXT("kFLD_SHADING_COMMON_PSYS_NOISE_ASPECT");
-	case kFLD_SHADING_COMMON_PSYS_SPLIT_TIME:
-		return TEXT("kFLD_SHADING_COMMON_PSYS_SPLIT_TIME");
-	case kFLD_SHADING_COMMON_PSYS_NUM_CHILDREN:
-		return TEXT("kFLD_SHADING_COMMON_PSYS_NUM_CHILDREN");
-	case kFLD_SHADING_COMMON_PSYS_BRANCH_ANGLE:
-		return TEXT("kFLD_SHADING_COMMON_PSYS_BRANCH_ANGLE");
-	case kFLD_SHADING_COMMON_PSYS_FRICTION:
-		return TEXT("kFLD_SHADING_COMMON_PSYS_FRICTION");
-	case kFLD_SHADING_COMMON_PSYS_UNUSED:
-		return TEXT("kFLD_SHADING_COMMON_PSYS_UNUSED");
-	case kFLD_SHADING_COMMON_PSYS_SPEED_DECAY:
-		return TEXT("kFLD_SHADING_COMMON_PSYS_SPEED_DECAY");
-	case kFLD_SHADING_COMMON_PSYS_BEND_U:
-		return TEXT("kFLD_SHADING_COMMON_PSYS_BEND_U");
-	case kFLD_SHADING_COMMON_PSYS_BEND_V:
-		return TEXT("kFLD_SHADING_COMMON_PSYS_BEND_V");
-	case kFLD_SHADING_COMMON_PSYS_CURL:
-		return TEXT("kFLD_SHADING_COMMON_PSYS_CURL");
-	case kFLD_SHADING_COMMON_PSYS_BLOB_LIGHTING:
-		return TEXT("kFLD_SHADING_COMMON_PSYS_BLOB_LIGHTING");
-	case kFLD_SHADING_COMMON_PSYS_BLUR_LENGTH:
-		return TEXT("kFLD_SHADING_COMMON_PSYS_BLUR_LENGTH");
-	case kFLD_SHADING_COMMON_PSYS_BLUR_QUALITY:
-		return TEXT("kFLD_SHADING_COMMON_PSYS_BLUR_QUALITY");
-	case kFLD_SHADING_COMMON_PSYS_START_FRAME:
-		return TEXT("kFLD_SHADING_COMMON_PSYS_START_FRAME");
-	case kFLD_SHADING_COMMON_PSYS_END_FRAME:
-		return TEXT("kFLD_SHADING_COMMON_PSYS_END_FRAME");
-	case kFLD_SHADING_COMMON_PSYS_SURFACE_SHADING:
-		return TEXT("kFLD_SHADING_COMMON_PSYS_SURFACE_SHADING");
-	case kFLD_SHADING_COMMON_PSYS_DIFFUSION_METHOD:
-		return TEXT("kFLD_SHADING_COMMON_PSYS_DIFFUSION_METHOD");
-	case kFLD_SHADING_COMMON_PSYS_USE_PARTICLE_FILE:
-		return TEXT("kFLD_SHADING_COMMON_PSYS_USE_PARTICLE_FILE");
-	case kFLD_SHADING_COMMON_PSYS_HAIR_LENGTH_MAX:
-		return TEXT("kFLD_SHADING_COMMON_PSYS_HAIR_LENGTH_MAX");
-	case kFLD_SHADING_COMMON_PSYS_HAIR_LENGTH_MIN:
-		return TEXT("kFLD_SHADING_COMMON_PSYS_HAIR_LENGTH_MIN");
-	case kFLD_SHADING_COMMON_PSYS_HAIR_STIFFNESS:
-		return TEXT("kFLD_SHADING_COMMON_PSYS_HAIR_STIFFNESS");
-	case kFLD_SHADING_COMMON_PSYS_HAIR_SEGMENTS:
-		return TEXT("kFLD_SHADING_COMMON_PSYS_HAIR_SEGMENTS");
-	case kFLD_SHADING_COMMON_PROPS_PROP_MASS:
-		return TEXT("kFLD_SHADING_COMMON_PROPS_PROP_MASS");
-	case kFLD_SHADING_COMMON_PROPS_DENSITY:
-		return TEXT("kFLD_SHADING_COMMON_PROPS_DENSITY");
-	case kFLD_SHADING_COMMON_PROPS_ELASTICITY:
-		return TEXT("kFLD_SHADING_COMMON_PROPS_ELASTICITY");
-	case kFLD_SHADING_COMMON_PROPS_FRICTION:
-		return TEXT("kFLD_SHADING_COMMON_PROPS_FRICTION");
-	case kFLD_SHADING_COMMON_PROPS_DRAG_COEFF:
-		return TEXT("kFLD_SHADING_COMMON_PROPS_DRAG_COEFF");
-	case kFLD_SHADING_DYN_GRAVITY:
-		return TEXT("kFLD_SHADING_DYN_GRAVITY");
-	case kFLD_SHADING_DYN_AIR_DENSITY:
-		return TEXT("kFLD_SHADING_DYN_AIR_DENSITY");
-	case kFLD_SHADING_DYN_FORCE_SCALE:
-		return TEXT("kFLD_SHADING_DYN_FORCE_SCALE");
-	case kFLD_SHADING_DYN_TIME_SCALE:
-		return TEXT("kFLD_SHADING_DYN_TIME_SCALE");
-	case kFLD_SHADING_DYN_FLOOR:
-		return TEXT("kFLD_SHADING_DYN_FLOOR");
-	case kFLD_SHADING_DYN_FLOOR_OFFSET:
-		return TEXT("kFLD_SHADING_DYN_FLOOR_OFFSET");
-	case kFLD_SHADING_DYN_CIELING:
-		return TEXT("kFLD_SHADING_DYN_CIELING");
-	case kFLD_SHADING_DYN_CIELING_OFFSET:
-		return TEXT("kFLD_SHADING_DYN_CIELING_OFFSET");
-	case kFLD_SHADING_DYN_LEFT:
-		return TEXT("kFLD_SHADING_DYN_LEFT");
-	case kFLD_SHADING_DYN_LEFT_OFFSET:
-		return TEXT("kFLD_SHADING_DYN_LEFT_OFFSET");
-	case kFLD_SHADING_DYN_RIGHT:
-		return TEXT("kFLD_SHADING_DYN_RIGHT");
-	case kFLD_SHADING_DYN_RIGHT_OFFSET:
-		return TEXT("kFLD_SHADING_DYN_RIGHT_OFFSET");
-	case kFLD_SHADING_DYN_FRONT:
-		return TEXT("kFLD_SHADING_DYN_FRONT");
-	case kFLD_SHADING_DYN_FRONT_OFFSET:
-		return TEXT("kFLD_SHADING_DYN_FRONT_OFFSET");
-	case kFLD_SHADING_DYN_BACK:
-		return TEXT("kFLD_SHADING_DYN_BACK");
-	case kFLD_SHADING_DYN_BACK_OFFSET:
-		return TEXT("kFLD_SHADING_DYN_BACK_OFFSET");
-	case kFLD_SHADING_DYN_TURB_TYPE:
-		return TEXT("kFLD_SHADING_DYN_TURB_Type");
-	case kFLD_SHADING_DYN_TURB_INTENSITY:
-		return TEXT("kFLD_SHADING_DYN_TURB_INTENSITY");
-	case kFLD_SHADING_DYN_TURB_SPREAD:
-		return TEXT("kFLD_SHADING_DYN_TURB_SPREAD");
-	case kFLD_SHADING_DYN_TURB_PERSISTANCE:
-		return TEXT("kFLD_SHADING_DYN_TURB_PERSISTANCE");
-	case kFLD_SHADING_DYN_TURB_SPACE_RESOLUTION:
-		return TEXT("kFLD_SHADING_DYN_TURB_SPACE_RESOLUTION");
-	case kFLD_SHADING_DYN_TURB_TIME_RESOLUTION:
-		return TEXT("kFLD_SHADING_DYN_TURB_TIME_RESOLUTION");
-	case kFLD_SHADING_DYN_TURB_ROUGHNESS:
-		return TEXT("kFLD_SHADING_DYN_TURB_ROUGHNESS");
-	case kFLD_SHADING_DYN_TURB_VARIABILITY:
-		return TEXT("kFLD_SHADING_DYN_TURB_VARIABILITY");
-	case kFLD_SHADING_DYN_TURB_GRANULARITY:
-		return TEXT("kFLD_SHADING_DYN_TURB_GRANULARITY");
-	case kFLD_SHADING_DYN_TURB_ANIMATED:
-		return TEXT("kFLD_SHADING_DYN_TURB_ANIMATED");
-	case kFLD_SHADING_DYN_WALL_FRICTION:
-		return TEXT("kFLD_SHADING_DYN_WALL_FRICTION");
-	case kFLD_SHADING_DYN_WALL_ELASTICITY:
-		return TEXT("kFLD_SHADING_DYN_WALL_ELASTICITY");
-	case kFLD_SHADING_COMMON_MATTE_OBJECT:
-		return TEXT("kFLD_SHADING_COMMON_MATTE_OBJECT");
-	case kFLD_SHADING_FOG_MIN_DISTANCE:
-		return TEXT("kFLD_SHADING_FOG_MIN_DISTANCE");
-	case kFLD_SHADING_FOG_MAX_DISTANCE:
-		return TEXT("kFLD_SHADING_FOG_MAX_DISTANCE");
-	case kFLD_SHADING_FILTER_TYPE:
-		return TEXT("kFLD_SHADING_FILTER_Type");
-	case kFLD_SHADING_CACHE:
-		return TEXT("kFLD_SHADING_CACHE");
-	case kFLD_SHADING_COMMON_REFLECT_BACKGROUND:
-		return TEXT("kFLD_SHADING_COMMON_REFLECT_BACKGROUND");
-	case kFLD_SHADING_SHADERGLOW_HALO_INTENSITY:
-		return TEXT("kFLD_SHADING_SHADERGLOW_HALO_INTENSITY");
-	case kFLD_SHADING_CLEARCOAT_ACTIVE:
-		return TEXT("kFLD_SHADING_CLEARCOAT_ACTIVE");
-	case kFLD_SHADING_CLEARCOAT_INDEX:
-		return TEXT("kFLD_SHADING_CLEARCOAT_INDEX");
-	case kFLD_SHADING_CLEARCOAT_SCALE:
-		return TEXT("kFLD_SHADING_CLEARCOAT_SCALE");
-	case kFLD_SHADING_CLEARCOAT_BIAS:
-		return TEXT("kFLD_SHADING_CLEARCOAT_BIAS");
-	case kFLD_SHADING_CLEARCOAT_MATERIAL_CODE:
-		return TEXT("kFLD_SHADING_CLEARCOAT_MATERIAL_CODE");
-	case kFLD_SHADING_TRANSPARENCY_MATERIAL_CODE:
-		return TEXT("kFLD_SHADING_TRANSPARENCY_MATERIAL_CODE");
-	case kFLD_SHADING_LIGHTTUNNEL_ENABLED:
-		return TEXT("kFLD_SHADING_LIGHTTUNNEL_ENABLED");
-	case kFLD_SHADING_LIGHTTUNNEL_NUMBEROFBANDS:
-		return TEXT("kFLD_SHADING_LIGHTTUNNEL_NUMBEROFBANDS");
-	case kFLD_SHADING_LIGHTTUNNEL_INTENSITY:
-		return TEXT("kFLD_SHADING_LIGHTTUNNEL_INTENSITY");
-	case kFLD_SHADING_LIGHTTUNNEL_COLOR_R:
-		return TEXT("kFLD_SHADING_LIGHTTUNNEL_COLOR_R");
-	case kFLD_SHADING_LIGHTTUNNEL_BANDWIDTH:
-		return TEXT("kFLD_SHADING_LIGHTTUNNEL_BANDWIDTH");
-	case kFLD_SHADING_LIGHTTUNNEL_BANDFRINGE:
-		return TEXT("kFLD_SHADING_LIGHTTUNNEL_BANDFRINGE");
-	case kFLD_SHADING_LIGHTTUNNEL_VISIBLE:
-		return TEXT("kFLD_SHADING_LIGHTTUNNEL_VISIBLE");
-	case kFLD_SHADING_GEOMETRIC_TYPE:
-		return TEXT("kFLD_SHADING_GEOMETRIC_Type");
-	case kFLD_SHADING_GEOMETRY_SOURCE_TEXTURE:
-		return TEXT("kFLD_SHADING_GEOMETRY_SOURCE_TEXTURE");
-	case kFLD_SHADING_GEOMETRIC_SCALE:
-		return TEXT("kFLD_SHADING_GEOMETRIC_SCALE");
-	case kFLD_SHADING_GEOMETRIC_ROTATE:
-		return TEXT("kFLD_SHADING_GEOMETRIC_ROTATE");
-	case kFLD_SHADING_GEOMETRIC_HEIGHT:
-		return TEXT("kFLD_SHADING_GEOMETRIC_HEIGHT");
-	case kFLD_SHADING_MASTERLIGHT_AO_CONTRAST:
-		return TEXT("kFLD_SHADING_MASTERLIGHT_AO_CONTRAST");
-	case MAX_SHADING_ANIM_ID:
-		return TEXT("MAX_SHADING_ANIM_ID");
-	case kFLD_SHADING_COMMON_BUMP_:
-		return TEXT("kFLD_SHADING_COMMON_BUMP_");
-	case kFLD_SHADING_COMMON_DISPLACEMENT_:
-		return TEXT("kFLD_SHADING_COMMON_DISPLACEMENT_");
-	case kFLD_SHADING_COMMON_SHADINGMAP_:
-		return TEXT("kFLD_SHADING_COMMON_SHADINGMAP_");
-	case kFLD_SHADING_COMMON_EMIT_PARTICLES_:
-		return TEXT("kFLD_SHADING_COMMON_EMIT_PARTICLES_");
-	case kFLD_SHADING_COMMON_CONCENTRATIO_:
-		return TEXT("kFLD_SHADING_COMMON_CONCENTRATIO_");
-	case kFLD_SHADING_COMMON_CONTRAST_:
-		return TEXT("kFLD_SHADING_COMMON_CONTRAST_");
-	case kFLD_SHADING_COMMON_REFRACTIVEMATERIAL_:
-		return TEXT("kFLD_SHADING_COMMON_REFRACTIVEMATERIAL_");
-	case kFLD_SHADING_COMMON_PSYS_PRESET_:
-		return TEXT("kFLD_SHADING_COMMON_PSYS_PRESET_");
-	case kFLD_SHADING_COMMON_PSYS_USE_FILE_:
-		return TEXT("kFLD_SHADING_COMMON_PSYS_USE_FILE_");
-	case kFLD_SHADING_COMMON_PSYS_PARTICLE_FILE_:
-		return TEXT("kFLD_SHADING_COMMON_PSYS_PARTICLE_FILE_");
-	case kFLD_SHADING_COMMON_PSYS_MOTION_TYPE_:
-		return TEXT("kFLD_SHADING_COMMON_PSYS_MOTION_TYPE_");
-	case kFLD_SHADING_COMMON_PSYS_START_FRAME_:
-		return TEXT("kFLD_SHADING_COMMON_PSYS_START_FRAME_");
-	case kFLD_SHADING_COMMON_PSYS_END_FRAME_:
-		return TEXT("kFLD_SHADING_COMMON_PSYS_END_FRAME_");
-	case kFLD_SHADING_COMMON_PSYS_BLOB_MAP_:
-		return TEXT("kFLD_SHADING_COMMON_PSYS_BLOB_MAP_");
-	case kFLD_SHADING_COMMON_PSYS_HAIR_LENGTH_MAP_:
-		return TEXT("kFLD_SHADING_COMMON_PSYS_HAIR_LENGTH_MAP_");
-	case kFLD_SHADING_COMMON_PSYS_DIFFUSION_:
-		return TEXT("kFLD_SHADING_COMMON_PSYS_DIFFUSION_");
-	case kFLD_SHADING_COMMON_PSYS_DIFFUSION_RATE_:
-		return TEXT("kFLD_SHADING_COMMON_PSYS_DIFFUSION_RATE_");
-	case kFLD_SHADING_COMMON_PSYS_COLLISIONS_:
-		return TEXT("kFLD_SHADING_COMMON_PSYS_COLLISIONS_");
-	case kFLD_SHADING_COMMON_PSYS_HIT_METHOD_:
-		return TEXT("kFLD_SHADING_COMMON_PSYS_HIT_METHOD_");
-	case kFLD_SHADING_PHONG_GLOSS_:
-		return TEXT("kFLD_SHADING_PHONG_GLOSS_");
-	case kFLD_SHADING_PHONG_SPECULARITY_:
-		return TEXT("kFLD_SHADING_PHONG_SPECULARITY_");
-	case kFLD_SHADING_PHONG_REFLECTION_:
-		return TEXT("kFLD_SHADING_PHONG_REFLECTION_");
-	case kFLD_SHADING_BLINN_GLOSS_:
-		return TEXT("kFLD_SHADING_BLINN_GLOSS_");
-	case kFLD_SHADING_BLINN_SPECULARITY_:
-		return TEXT("kFLD_SHADING_BLINN_SPECULARITY_");
-	case kFLD_SHADING_BLINN_REFLECTION_:
-		return TEXT("kFLD_SHADING_BLINN_REFLECTION_");
-	case kFLD_SHADING_CLEARCOAT_MATERIAL_:
-		return TEXT("kFLD_SHADING_CLEARCOAT_MATERIAL_");
-	case kFLD_SHADING_BACKGROUND_SEQUENCE_:
-		return TEXT("kFLD_SHADING_BACKGROUND_SEQUENCE_");
-	case kFLD_SHADING_FOG_DENSITY_MAP_:
-		return TEXT("kFLD_SHADING_FOG_DENSITY_MAP_");
-	case kFLD_SHADING_SHADERGLOW_AUTO_EXPOSURE_:
-		return TEXT("kFLD_SHADING_SHADERGLOW_AUTO_EXPOSURE_");
-	case kFLD_SHADING_TONE_MAPPING_GAMMA_:
-		return TEXT("kFLD_SHADING_TONE_MAPPING_GAMMA_");
-	case kFLD_SHADING_TONE_MAPPING_EXPOSURE_:
-		return TEXT("kFLD_SHADING_TONE_MAPPING_EXPOSURE_");
-	case kFLD_SHADING_TONE_MAPPING_SATURATION_:
-		return TEXT("kFLD_SHADING_TONE_MAPPING_SATURATION_");
-	case kFLD_SHADING_GROUND_PLANE_PREVIEW_:
-		return TEXT("kFLD_SHADING_GROUND_PLANE_PREVIEW_");
-	case kFLD_SHADING_GROUND_PLANE_HEIGHT_:
-		return TEXT("kFLD_SHADING_GROUND_PLANE_HEIGHT_");
-	case kFLD_SHADING_GROUND_PLANE_SHADOWS_:
-		return TEXT("kFLD_SHADING_GROUND_PLANE_SHADOWS_");
-	case kFLD_SHADING_GROUND_PLANE_SHADOW_BLUR_:
-		return TEXT("kFLD_SHADING_GROUND_PLANE_SHADOW_BLUR_");
-	case kFLD_SHADING_GROUND_PLANE_SHADOW_POSITION_:
-		return TEXT("kFLD_SHADING_GROUND_PLANE_SHADOW_POSITION_");
-	case kFLD_SHADING_GROUND_PLANE_SHADOW_TRANSPARENCY_:
-		return TEXT("kFLD_SHADING_GROUND_PLANE_SHADOW_TRANSPARENCY_");
-	case kFLD_SHADING_GROUND_PLANE_REFLECTION_:
-		return TEXT("kFLD_SHADING_GROUND_PLANE_REFLECTION_");
-	case kFLD_SHADING_GROUND_PLANE_REFLECTIVITY_:
-		return TEXT("kFLD_SHADING_GROUND_PLANE_REFLECTIVITY_");
-	case kFLD_SHADING_GROUND_PLANE_REFLECTION_DEPTH_:
-		return TEXT("kFLD_SHADING_GROUND_PLANE_REFLECTION_DEPTH_");
-	case kFLD_SHADING_GROUND_PLANE_REFLECTION_BLUR_:
-		return TEXT("kFLD_SHADING_GROUND_PLANE_REFLECTION_BLUR_");
-	case kFLD_SHADING_GROUND_PLANE_REFLECTION_BLUR_DEPTH_:
-		return TEXT("kFLD_SHADING_GROUND_PLANE_REFLECTION_BLUR_DEPTH_");
-	case kFLD_SHADING_COMMON_TEXTURE_OVERLAY_:
-		return TEXT("kFLD_SHADING_COMMON_TEXTURE_OVERLAY_");
-	case kFLD_SHADING_COMMON_TEXTURE_COLOR_REMAP_:
-		return TEXT("kFLD_SHADING_COMMON_TEXTURE_COLOR_REMAP_");
-	case kFLD_SHADING_COMMON_TEXTURE_SMEAR_MAP_:
-		return TEXT("kFLD_SHADING_COMMON_TEXTURE_SMEAR_MAP_");
-	case kFLD_SHADING_RAMP_POSITION_:
-		return TEXT("kFLD_SHADING_RAMP_POSITION_");
-	case kFLD_SHADING_RAMP_RAMP_COLOR_:
-		return TEXT("kFLD_SHADING_RAMP_RAMP_COLOR_");
-	case kFLD_SHADING_RAMP_INTENSITY_:
-		return TEXT("kFLD_SHADING_RAMP_INTENSITY_");
-	case kFLD_SHADING_STENCIL_IMAGE_:
-		return TEXT("kFLD_SHADING_STENCIL_IMAGE_");
-	case kFLD_SHADING_STENCIL_MASK_:
-		return TEXT("kFLD_SHADING_STENCIL_MASK_");
-	case kFLD_SHADING_CUBE_FINITE_SIZE_:
-		return TEXT("kFLD_SHADING_CUBE_FINITE_SIZE_");
-	case kFLD_SHADING_CUBE_CROSS_ORIENTATION_:
-		return TEXT("kFLD_SHADING_CUBE_CROSS_ORIENTATION_");
-	case kFLD_SHADING_BALL_IMAGE_:
-		return TEXT("kFLD_SHADING_BALL_IMAGE_");
-	case kFLD_SHADING_SPHERE_SOURCE_TEXTURE_:
-		return TEXT("kFLD_SHADING_SPHERE_SOURCE_TEXTURE_");
-	case kFLD_SHADING_SKY_CLOUD_TEXTURE_:
-		return TEXT("kFLD_SHADING_SKY_CLOUD_TEXTURE_");
-	case kFLD_SHADING_SFILE_SOURCE_TEXTURE_:
-		return TEXT("kFLD_SHADING_SFILE_SOURCE_TEXTURE_");
-	case kFLD_SHADING_CURVATURE_TYPE_:
-		return TEXT("kFLD_SHADING_CURVATURE_TYPE_");
-	case kFLD_SHADING_CURVATURE_DISPLAY_:
-		return TEXT("kFLD_SHADING_CURVATURE_DISPLAY_");
-	case kFLD_SHADING_CURVATURE_SCALE_:
-		return TEXT("kFLD_SHADING_CURVATURE_SCALE_");
-	case kFLD_SHADING_BACKGROUND_BACKDROP_:
-		return TEXT("kFLD_SHADING_BACKGROUND_BACKDROP_");
-	case kFLD_SHADING_IBL_REFLECTION_MAP_:
-		return TEXT("kFLD_SHADING_IBL_REFLECTION_MAP_");
-	case kFLD_SHADING_IBL_REFLECTION_MAP_ROTATION_:
-		return TEXT("kFLD_SHADING_IBL_REFLECTION_MAP_ROTATION_");
-	case kFLD_SHADING_VOLUME_PIX_SEQUENCE_:
-		return TEXT("kFLD_SHADING_VOLUME_PIX_SEQUENCE_");
-	case kFLD_SHADING_CUBE_CROSS_IMAGE_:
-		return TEXT("kFLD_SHADING_CUBE_CROSS_IMAGE_");
-	case kFLD_SHADING_CUBE_RIGHT_IMAGE_:
-		return TEXT("kFLD_SHADING_CUBE_RIGHT_IMAGE_");
-	case kFLD_SHADING_CUBE_LEFT_IMAGE_:
-		return TEXT("kFLD_SHADING_CUBE_LEFT_IMAGE_");
-	case kFLD_SHADING_CUBE_TOP_IMAGE_:
-		return TEXT("kFLD_SHADING_CUBE_TOP_IMAGE_");
-	case kFLD_SHADING_CUBE_BOTTOM_IMAGE_:
-		return TEXT("kFLD_SHADING_CUBE_BOTTOM_IMAGE_");
-	case kFLD_SHADING_CUBE_FRONT_IMAGE_:
-		return TEXT("kFLD_SHADING_CUBE_FRONT_IMAGE_");
-	case kFLD_SHADING_CUBE_BACK_IMAGE_:
-		return TEXT("kFLD_SHADING_CUBE_BACK_IMAGE_");
-	case kFLD_SHADING_FILE_IMAGE_:
-		return TEXT("kFLD_SHADING_FILE_IMAGE_");
-	case kFLD_SHADING_COMMON_VRED_MATERIAL_TAG_:
-		return TEXT("kFLD_SHADING_COMMON_VRED_MATERIAL_TAG_");
-	}
-	return TEXT("");
-}
-
 void OpenModelUtils::SetActorTransform(TSharedPtr<IDatasmithActorElement>& OutActorElement, const AlDagNode& InDagNode)
 {
 	if (OutActorElement)
diff --git a/Engine/Plugins/Enterprise/DatasmithCADImporter/Source/DatasmithWireTranslator/Private/OpenModelUtils.h b/Engine/Plugins/Enterprise/DatasmithCADImporter/Source/DatasmithWireTranslator/Private/OpenModelUtils.h
--- a/Engine/Plugins/Enterprise/DatasmithCADImporter/Source/DatasmithWireTranslator/Private/OpenModelUtils.h
+++ b/Engine/Plugins/Enterprise/DatasmithCADImporter/Source/DatasmithWireTranslator/Private/OpenModelUtils.h
@@ -46,9 +46,6 @@ enum class EAlShaderModelType : uint8
 
 namespace OpenModelUtils
 {
-	const TCHAR* AlObjectTypeToString(AlObjectType type);
-	const TCHAR* AlShadingFieldToString(AlShadingFields field);
-
 	void SetActorTransform(TSharedPtr<IDatasmithActorElement>& OutActorElement, const AlDagNode& InDagNode);
 
 	bool IsValidActor(const TSharedPtr<IDatasmithActorElement>& ActorElement);
