Formating

diff --git a/Engine/Plugins/Enterprise/DatasmithCADImporter/Source/CADLibrary/Private/MeshDescriptionHelper.cpp b/Engine/Plugins/Enterprise/DatasmithCADImporter/Source/CADLibrary/Private/MeshDescriptionHelper.cpp
--- a/Engine/Plugins/Enterprise/DatasmithCADImporter/Source/CADLibrary/Private/MeshDescriptionHelper.cpp
+++ b/Engine/Plugins/Enterprise/DatasmithCADImporter/Source/CADLibrary/Private/MeshDescriptionHelper.cpp
@@ -18,131 +18,156 @@ typedef uint32 TriangleIndex[3];
 
 namespace CADLibrary
 {
-	struct FVertexData
-	{
-		float Z;
-		int32 Index;
-		FVector Coordinates;
-		bool bIsMerged;
-		FVertexID VertexID;
-		FVertexID SymVertexID;
 
-		/** Default constructor. */
-		FVertexData() {}
+struct FVertexData
+{
+	float Z;
+	int32 Index;
+	FVector Coordinates;
+	bool bIsMerged;
+	FVertexID VertexID;
+	FVertexID SymVertexID;
+
+	/** Default constructor. */
+	FVertexData() {}
+
+	/** Initialization constructor. */
+	FVertexData(int32 InIndex, const FVector& V)
+	{
+		//Z = V.X + V.Y + V.Z;
+		Z = 0.30f * V.X + 0.33f * V.Y + 0.37f * V.Z;
+		Index = InIndex;
+		Coordinates = V;
+		bIsMerged = false;
+		VertexID = INDEX_NONE;
+	}
+};
 
-		/** Initialization constructor. */
-		FVertexData(int32 InIndex, const FVector& V)
-		{
-			//Z = V.X + V.Y + V.Z;
-			Z = 0.30f * V.X + 0.33f * V.Y + 0.37f * V.Z;
-			Index = InIndex;
-			Coordinates = V;
-			bIsMerged = false;
-			VertexID = INDEX_NONE;
-		}
-	};
+struct FCompareVertexZ
+{
+	FORCEINLINE bool operator()(FVertexData const& A, FVertexData const& B) const { return A.Z < B.Z; }
+};
 
-	struct FCompareVertexZ
+// Verify the 3 input indices are not defining a degenerated triangle and fill up the corresponding FVertexIDs
+bool IsTriangleDegenerated(const int32_t* Indices, const TArray<FVertexID>& RemapVertexPosition, FVertexID VertexIDs[3])
+{
+	if (Indices[0] == Indices[1] || Indices[0] == Indices[2] || Indices[1] == Indices[2])
 	{
-		FORCEINLINE bool operator()(FVertexData const& A, FVertexData const& B) const { return A.Z < B.Z; }
-	};
+		return true;
+	}
 
-	// Verify the 3 input indices are not defining a degenerated triangle and fill up the corresponding FVertexIDs
-	bool IsTriangleDegenerated(const int32_t* Indices, const TArray<FVertexID>& RemapVertexPosition, FVertexID VertexIDs[3])
+	for (int32 Corner = 0; Corner < 3; ++Corner)
 	{
-		if (Indices[0] == Indices[1] || Indices[0] == Indices[2] || Indices[1] == Indices[2])
-		{
-			return true;
-		}
+		VertexIDs[Corner] = RemapVertexPosition[Indices[Corner]];
+	}
+
+	return (VertexIDs[0] == VertexIDs[1] || VertexIDs[0] == VertexIDs[2] || VertexIDs[1] == VertexIDs[2]);
+}
+
+void FillKioVertexPosition(const FImportParameters& ImportParams, const FMeshParameters& MeshParameters, FBodyMesh& Body, FMeshDescription& MeshDescription)
+{
+	int32 TriangleCount = Body.TriangleCount;
+	TArray<FTessellationData>& FaceTessellationSet = Body.Faces;
 
-		for (int32 Corner = 0; Corner < 3; ++Corner)
+	TVertexAttributesRef<FVector3f> VertexPositions = MeshDescription.GetVertexPositions();
+
+	// Create a list of vertex Z/index pairs
+	TArray<FVertexData> VertexDataSet;
+	VertexDataSet.Reserve(TriangleCount * 3);
+
+	FVector Position;
+	int32 GlobalVertexCount = 0;
+	for (FTessellationData& CTTessellation : FaceTessellationSet)
+	{
+		for (const FVector& Vertex : CTTessellation.PositionArray)
 		{
-			VertexIDs[Corner] = RemapVertexPosition[Indices[Corner]];
+			VertexDataSet.Emplace(GlobalVertexCount, Vertex);
+			++GlobalVertexCount;
 		}
-
-		return (VertexIDs[0] == VertexIDs[1] || VertexIDs[0] == VertexIDs[2] || VertexIDs[1] == VertexIDs[2]);
 	}
+	VertexDataSet.SetNum(GlobalVertexCount);
 
-	void FillKioVertexPosition(const FImportParameters& ImportParams, const FMeshParameters& MeshParameters, FBodyMesh& Body, FMeshDescription& MeshDescription)
-	{
-		int32 TriangleCount = Body.TriangleCount;
-		TArray<FTessellationData>& FaceTessellationSet = Body.Faces;
+	// Sort the vertices by z value
+	VertexDataSet.Sort(FCompareVertexZ());
 
-		TVertexAttributesRef<FVector3f> VertexPositions = MeshDescription.GetVertexPositions();
+	TArray<int32> NewIndexOf;
+	NewIndexOf.SetNumZeroed(GlobalVertexCount);
 
-		// Create a list of vertex Z/index pairs
-		TArray<FVertexData> VertexDataSet;
-		VertexDataSet.Reserve(TriangleCount * 3);
+	TArray<int32> IndexOfCoincidentNode;
+	IndexOfCoincidentNode.SetNumZeroed(GlobalVertexCount);
 
-		FVector Position;
-		int32 GlobalVertexCount = 0;
-		for (FTessellationData& CTTessellation : FaceTessellationSet)
+	int32 VertexCount = 0;
+	// Search for duplicates, quickly!
+	for (int32 i = 0; i < VertexDataSet.Num(); i++)
+	{
+		NewIndexOf[VertexDataSet[i].Index] = i;
+		if (VertexDataSet[i].bIsMerged)
 		{
-			for (const FVector& Vertex : CTTessellation.PositionArray)
-			{
-				VertexDataSet.Emplace(GlobalVertexCount, Vertex);
-				++GlobalVertexCount;
-			}
+			continue;
 		}
-		VertexDataSet.SetNum(GlobalVertexCount);
 
-		// Sort the vertices by z value
-		VertexDataSet.Sort(FCompareVertexZ());
+		VertexDataSet[i].bIsMerged = true;
+		int32 Index_i = VertexDataSet[i].Index;
+		IndexOfCoincidentNode[Index_i] = Index_i;
 
-		TArray<int32> NewIndexOf;
-		NewIndexOf.SetNumZeroed(GlobalVertexCount);
+		const FVector& PositionA = VertexDataSet[i].Coordinates;
 
-		TArray<int32> IndexOfCoincidentNode;
-		IndexOfCoincidentNode.SetNumZeroed(GlobalVertexCount);
-
-		int32 VertexCount = 0;
-		// Search for duplicates, quickly!
-		for (int32 i = 0; i < VertexDataSet.Num(); i++)
+		// only need to search forward, since we add pairs both ways
+		for (int32 j = i + 1; j < VertexDataSet.Num(); j++)
 		{
-			NewIndexOf[VertexDataSet[i].Index] = i;
-			if (VertexDataSet[i].bIsMerged)
+			if (FMath::Abs(VertexDataSet[j].Z - VertexDataSet[i].Z) > KINDA_SMALL_NUMBER)
 			{
-				continue;
+				break; // can't be any more duplicated
 			}
 
-			VertexDataSet[i].bIsMerged = true;
-			int32 Index_i = VertexDataSet[i].Index;
-			IndexOfCoincidentNode[Index_i] = Index_i;
-
-			const FVector& PositionA = VertexDataSet[i].Coordinates;
-
-			// only need to search forward, since we add pairs both ways
-			for (int32 j = i + 1; j < VertexDataSet.Num(); j++)
+			const FVector& PositionB = VertexDataSet[j].Coordinates;
+			if (PositionA.Equals(PositionB, KINDA_SMALL_NUMBER))
 			{
-				if (FMath::Abs(VertexDataSet[j].Z - VertexDataSet[i].Z) > KINDA_SMALL_NUMBER)
-				{
-					break; // can't be any more duplicated
-				}
-
-				const FVector& PositionB = VertexDataSet[j].Coordinates;
-				if (PositionA.Equals(PositionB, KINDA_SMALL_NUMBER))
-				{
-					VertexDataSet[j].bIsMerged = true;
-					IndexOfCoincidentNode[VertexDataSet[j].Index] = Index_i;
-				}
+				VertexDataSet[j].bIsMerged = true;
+				IndexOfCoincidentNode[VertexDataSet[j].Index] = Index_i;
 			}
-			VertexCount++;
 		}
+		VertexCount++;
+	}
+
+	// if Symmetric mesh, the symmetric side of the mesh have to be generated
+	FMatrix SymmetricMatrix;
+	bool bIsSymmetricMesh = MeshParameters.bIsSymmetric;
+	if (bIsSymmetricMesh)
+	{
+		SymmetricMatrix = FDatasmithUtils::GetSymmetricMatrix(MeshParameters.SymmetricOrigin, MeshParameters.SymmetricNormal);
+	}
+
+	// Make MeshDescription.VertexPositions and VertexID
+	MeshDescription.ReserveNewVertices(bIsSymmetricMesh ? VertexCount * 2 : VertexCount);
+	Body.VertexIds.Reserve(VertexCount);
+
+	int32 GlobalVertexIndex = 0;
+
+	for (int32 VertexIndex = 0; VertexIndex < GlobalVertexCount; ++VertexIndex)
+	{
+		int32 RealIndex = VertexDataSet[VertexIndex].Index;
 
-		// if Symmetric mesh, the symmetric side of the mesh have to be generated
-		FMatrix SymmetricMatrix;
-		bool bIsSymmetricMesh = MeshParameters.bIsSymmetric;
-		if (bIsSymmetricMesh)
+		// Vertex is outside bbox
+		if (RealIndex < 0)
 		{
-			SymmetricMatrix = FDatasmithUtils::GetSymmetricMatrix(MeshParameters.SymmetricOrigin, MeshParameters.SymmetricNormal);
+			continue;
 		}
 
-		// Make MeshDescription.VertexPositions and VertexID
-		MeshDescription.ReserveNewVertices(bIsSymmetricMesh ? VertexCount * 2 : VertexCount);
-		Body.VertexIds.Reserve(VertexCount);
+		if (IndexOfCoincidentNode[RealIndex] != RealIndex)
+		{
+			continue;
+		}
 
-		int32 GlobalVertexIndex = 0;
+		FVertexID VertexID = MeshDescription.CreateVertex();
+		Body.VertexIds.Add(VertexID);
+		VertexPositions[VertexID] = (FVector3f)FDatasmithUtils::ConvertVector((FDatasmithUtils::EModelCoordSystem)ImportParams.GetModelCoordSys(), VertexDataSet[VertexIndex].Coordinates);
+		VertexDataSet[VertexIndex].VertexID = VertexID;
+	}
 
+	if (bIsSymmetricMesh)
+	{
+		Body.SymmetricVertexIds.Reserve(Body.VertexIds.Num());
 		for (int32 VertexIndex = 0; VertexIndex < GlobalVertexCount; ++VertexIndex)
 		{
 			int32 RealIndex = VertexDataSet[VertexIndex].Index;
@@ -159,557 +184,534 @@ namespace CADLibrary
 			}
 
 			FVertexID VertexID = MeshDescription.CreateVertex();
-			Body.VertexIds.Add(VertexID);
-			VertexPositions[VertexID] = (FVector3f)FDatasmithUtils::ConvertVector((FDatasmithUtils::EModelCoordSystem) ImportParams.GetModelCoordSys(), VertexDataSet[VertexIndex].Coordinates);
-			VertexDataSet[VertexIndex].VertexID = VertexID;
-		}
-
-		if (bIsSymmetricMesh)
-		{
-			Body.SymmetricVertexIds.Reserve(Body.VertexIds.Num());
-			for (int32 VertexIndex = 0; VertexIndex < GlobalVertexCount; ++VertexIndex)
-			{
-				int32 RealIndex = VertexDataSet[VertexIndex].Index;
-
-				// Vertex is outside bbox
-				if (RealIndex < 0)
-				{
-					continue;
-				}
-
-				if (IndexOfCoincidentNode[RealIndex] != RealIndex)
-				{
-					continue;
-				}
-
-				FVertexID VertexID = MeshDescription.CreateVertex();
-				Body.SymmetricVertexIds.Add(VertexID);
-				VertexPositions[VertexID] = (FVector3f)FDatasmithUtils::ConvertVector((FDatasmithUtils::EModelCoordSystem) ImportParams.GetModelCoordSys(), VertexDataSet[VertexIndex].Coordinates);
-				VertexPositions[VertexID] = FVector4f(SymmetricMatrix.TransformPosition((FVector)VertexPositions[VertexID]));
+			Body.SymmetricVertexIds.Add(VertexID);
+			VertexPositions[VertexID] = (FVector3f)FDatasmithUtils::ConvertVector((FDatasmithUtils::EModelCoordSystem)ImportParams.GetModelCoordSys(), VertexDataSet[VertexIndex].Coordinates);
+			VertexPositions[VertexID] = FVector4f(SymmetricMatrix.TransformPosition((FVector)VertexPositions[VertexID]));
 
-				VertexDataSet[VertexIndex].SymVertexID = VertexID;
-			}
+			VertexDataSet[VertexIndex].SymVertexID = VertexID;
 		}
+	}
 
-		// For each face, for each vertex set VertexId
-		GlobalVertexIndex = 0;
-		for (FTessellationData& CTTessellation : FaceTessellationSet)
+	// For each face, for each vertex set VertexId
+	GlobalVertexIndex = 0;
+	for (FTessellationData& CTTessellation : FaceTessellationSet)
+	{
+		CTTessellation.PositionIndices.SetNum(CTTessellation.PositionArray.Num());
+		for (int32 VertexIndex = 0; VertexIndex < CTTessellation.PositionArray.Num(); ++VertexIndex, ++GlobalVertexIndex)
 		{
-			CTTessellation.PositionIndices.SetNum(CTTessellation.PositionArray.Num());
-			for (int32 VertexIndex = 0; VertexIndex < CTTessellation.PositionArray.Num(); ++VertexIndex, ++GlobalVertexIndex)
-			{
-				int32 NewIndex = NewIndexOf[IndexOfCoincidentNode[GlobalVertexIndex]];
-				FVertexData& Vertex = VertexDataSet[NewIndex];
-				int32 Index = Vertex.VertexID.GetValue();
-				CTTessellation.PositionIndices[VertexIndex] = VertexDataSet[NewIndex].VertexID.GetValue();
-			}
+			int32 NewIndex = NewIndexOf[IndexOfCoincidentNode[GlobalVertexIndex]];
+			FVertexData& Vertex = VertexDataSet[NewIndex];
+			int32 Index = Vertex.VertexID.GetValue();
+			CTTessellation.PositionIndices[VertexIndex] = VertexDataSet[NewIndex].VertexID.GetValue();
 		}
 	}
+}
 
-	void MergeCoincidentVertices(TArray<FVector>& VertexArray, TArray<int32>& VertexIdSet)
+void MergeCoincidentVertices(TArray<FVector>& VertexArray, TArray<int32>& VertexIdSet)
+{
+	const double CoincidenceTolerance = 0.001;
+
+	// Create a list of vertex Z/index pairs
+	TArray<FVertexData> VertexDataSet;
+	VertexDataSet.Reserve(VertexArray.Num());
+
+	const FVector* Position = VertexArray.GetData();
+	const int32* VertexId = VertexIdSet.GetData();
+	for (int32 Index = 0; Index < VertexArray.Num(); ++Index, ++Position, ++VertexId)
 	{
-		const double CoincidenceTolerance = 0.001;
+		VertexDataSet.Emplace(*VertexId, *Position);
+	}
 
-		// Create a list of vertex Z/index pairs
-		TArray<FVertexData> VertexDataSet;
-		VertexDataSet.Reserve(VertexArray.Num());
+	// Sort the vertices by z value
+	VertexDataSet.Sort(FCompareVertexZ());
 
-		const FVector* Position = VertexArray.GetData();
-		const int32* VertexId = VertexIdSet.GetData();
-		for (int32 Index = 0; Index < VertexArray.Num(); ++Index, ++Position, ++VertexId)
+	// Search for duplicates
+	for (int32 Index = 0; Index < VertexDataSet.Num(); Index++)
+	{
+		if (VertexDataSet[Index].bIsMerged)
 		{
-			VertexDataSet.Emplace(*VertexId, *Position);
+			continue;
 		}
 
-		// Sort the vertices by z value
-		VertexDataSet.Sort(FCompareVertexZ());
+		VertexDataSet[Index].bIsMerged = true;
+		int32 NewIndex = VertexDataSet[Index].Index;
+
+		const FVector& PositionA = VertexDataSet[Index].Coordinates;
 
-		// Search for duplicates
-		for (int32 Index = 0; Index < VertexDataSet.Num(); Index++)
+		// only need to search forward, since we add pairs both ways
+		for (int32 Andex = Index + 1; Andex < VertexDataSet.Num(); Andex++)
 		{
-			if (VertexDataSet[Index].bIsMerged)
+			if (FMath::Abs(VertexDataSet[Andex].Z - VertexDataSet[Index].Z) > 3. * CoincidenceTolerance)
 			{
-				continue;
+				break; // can't be any more duplicated
 			}
 
-			VertexDataSet[Index].bIsMerged = true;
-			int32 NewIndex = VertexDataSet[Index].Index;
-
-			const FVector& PositionA = VertexDataSet[Index].Coordinates;
-
-			// only need to search forward, since we add pairs both ways
-			for (int32 Andex = Index + 1; Andex < VertexDataSet.Num(); Andex++)
+			const FVector& PositionB = VertexDataSet[Andex].Coordinates;
+			if (PositionA.Equals(PositionB, CoincidenceTolerance))
 			{
-				if (FMath::Abs(VertexDataSet[Andex].Z - VertexDataSet[Index].Z) > 3.*CoincidenceTolerance)
-				{
-					break; // can't be any more duplicated
-				}
-
-				const FVector& PositionB = VertexDataSet[Andex].Coordinates;
-				if (PositionA.Equals(PositionB, CoincidenceTolerance))
-				{
-					VertexDataSet[Andex].bIsMerged = true;
-					VertexIdSet[VertexDataSet[Andex].Index] = NewIndex;
-				}
+				VertexDataSet[Andex].bIsMerged = true;
+				VertexIdSet[VertexDataSet[Andex].Index] = NewIndex;
 			}
 		}
 	}
+}
+
+void FillVertexPosition(const FImportParameters& ImportParams, const FMeshParameters& MeshParameters, FBodyMesh& Body, FMeshDescription& MeshDescription)
+{
+	int32 TriangleCount = Body.TriangleCount;
+	TArray<FTessellationData>& FaceTessellationSet = Body.Faces;
+
+	TVertexAttributesRef<FVector3f> VertexPositions = MeshDescription.GetVertexPositions();
+
+	TArray<FVector>& VertexArray = Body.VertexArray;
+	TArray<int32>& VertexIdSet = Body.VertexIds;
+	int32 VertexCount = VertexArray.Num();
+	VertexIdSet.SetNumZeroed(VertexCount);
+
+	// Make MeshDescription.VertexPositions and VertexID
+	MeshDescription.ReserveNewVertices(MeshParameters.bIsSymmetric ? VertexCount * 2 : VertexCount);
 
-	void FillVertexPosition(const FImportParameters& ImportParams, const FMeshParameters& MeshParameters, FBodyMesh& Body, FMeshDescription& MeshDescription)
+	int32 VertexIndex = -1;
+	for (const FVector& Vertex : VertexArray)
 	{
-		int32 TriangleCount = Body.TriangleCount;
-		TArray<FTessellationData>& FaceTessellationSet = Body.Faces;
+		VertexIndex++;
+
+		// Vertex is outside bbox
+		if (VertexIdSet[VertexIndex] == INDEX_NONE)
+		{
+			continue;
+		}
+
+		FVertexID VertexID = MeshDescription.CreateVertex();
+		VertexPositions[VertexID] = (FVector3f)FDatasmithUtils::ConvertVector((FDatasmithUtils::EModelCoordSystem)ImportParams.GetModelCoordSys(), Vertex);
+		VertexIdSet[VertexIndex] = VertexID;
+	}
 
-		TVertexAttributesRef<FVector3f> VertexPositions = MeshDescription.GetVertexPositions();
+	MergeCoincidentVertices(VertexArray, VertexIdSet);
 
-		TArray<FVector>& VertexArray = Body.VertexArray;
-		TArray<int32>& VertexIdSet = Body.VertexIds;
-		int32 VertexCount = VertexArray.Num();
-		VertexIdSet.SetNumZeroed(VertexCount);
+	// if Symmetric mesh, the symmetric side of the mesh have to be generated
+	if (MeshParameters.bIsSymmetric)
+	{
+		FMatrix SymmetricMatrix = FDatasmithUtils::GetSymmetricMatrix(MeshParameters.SymmetricOrigin, MeshParameters.SymmetricNormal);
 
-		// Make MeshDescription.VertexPositions and VertexID
-		MeshDescription.ReserveNewVertices(MeshParameters.bIsSymmetric ? VertexCount * 2 : VertexCount);
+		TArray<int32>& SymmetricVertexIds = Body.SymmetricVertexIds;
+		SymmetricVertexIds.SetNum(VertexArray.Num());
 
-		int32 VertexIndex = -1;
+		VertexIndex = 0;
 		for (const FVector& Vertex : VertexArray)
 		{
-			VertexIndex++;
-
-			// Vertex is outside bbox
 			if (VertexIdSet[VertexIndex] == INDEX_NONE)
 			{
+				SymmetricVertexIds[VertexIndex++] = INDEX_NONE;
 				continue;
 			}
 
 			FVertexID VertexID = MeshDescription.CreateVertex();
-			VertexPositions[VertexID] = (FVector3f)FDatasmithUtils::ConvertVector((FDatasmithUtils::EModelCoordSystem) ImportParams.GetModelCoordSys(), Vertex);
-			VertexIdSet[VertexIndex] = VertexID;
+			const FVector VertexPosition = FDatasmithUtils::ConvertVector((FDatasmithUtils::EModelCoordSystem)ImportParams.GetModelCoordSys(), Vertex);
+			const FVector4f SymmetricPosition = FVector4f(SymmetricMatrix.TransformPosition(VertexPosition));
+			VertexPositions[VertexID] = FVector3f(SymmetricPosition);
+			SymmetricVertexIds[VertexIndex++] = VertexID;
 		}
+	}
+}
 
-		MergeCoincidentVertices(VertexArray, VertexIdSet);
-
-		// if Symmetric mesh, the symmetric side of the mesh have to be generated
-		if (MeshParameters.bIsSymmetric)
-		{
-			FMatrix SymmetricMatrix = FDatasmithUtils::GetSymmetricMatrix(MeshParameters.SymmetricOrigin, MeshParameters.SymmetricNormal);
-
-			TArray<int32>& SymmetricVertexIds = Body.SymmetricVertexIds;
-			SymmetricVertexIds.SetNum(VertexArray.Num());
-
-			VertexIndex = 0;
-			for (const FVector& Vertex : VertexArray)
-			{
-				if (VertexIdSet[VertexIndex] == INDEX_NONE)
-				{
-					SymmetricVertexIds[VertexIndex++] = INDEX_NONE;
-					continue;
-				}
+// PolygonAttributes name used into modeling tools (ExtendedMeshAttribute::PolyTriGroups)
+const FName PolyTriGroups("PolyTriGroups");
 
-				FVertexID VertexID = MeshDescription.CreateVertex();
-				const FVector VertexPosition = FDatasmithUtils::ConvertVector((FDatasmithUtils::EModelCoordSystem) ImportParams.GetModelCoordSys(), Vertex);
-				const FVector4f SymmetricPosition = FVector4f(SymmetricMatrix.TransformPosition(VertexPosition));
-				VertexPositions[VertexID] = FVector3f(SymmetricPosition);
-				SymmetricVertexIds[VertexIndex++] = VertexID;
-			}
-		}
+// Copy of FMeshDescriptionBuilder::EnablePolyGroups()
+TPolygonAttributesRef<int32> EnableCADPatchGroups(FMeshDescription& OutMeshDescription)
+{
+	TPolygonAttributesRef<int32> PatchGroups = OutMeshDescription.PolygonAttributes().GetAttributesRef<int32>(PolyTriGroups);
+	if (PatchGroups.IsValid() == false)
+	{
+		OutMeshDescription.PolygonAttributes().RegisterAttribute<int32>(PolyTriGroups, 1, 0, EMeshAttributeFlags::AutoGenerated);
+		PatchGroups = OutMeshDescription.PolygonAttributes().GetAttributesRef<int32>(PolyTriGroups);
+		check(PatchGroups.IsValid());
 	}
+	return PatchGroups;
+}
 
-	// PolygonAttributes name used into modeling tools (ExtendedMeshAttribute::PolyTriGroups)
-	const FName PolyTriGroups("PolyTriGroups");
-
-	// Copy of FMeshDescriptionBuilder::EnablePolyGroups()
-	TPolygonAttributesRef<int32> EnableCADPatchGroups(FMeshDescription& OutMeshDescription)
+/**
+ * Polygon group is an attribute of polygons.As long as the mesh description is empty(no polygon), polygon group cannot be defined.
+ * The work around used is to create PolygonGroups and to set them a PolygonGroupAttributes.
+ * To get the existing polygon groups, for each created polygonGroups the PolygonGroupId is got(see GetExistingPatches)
+ * Warning: CopyPatchGroups is call in FCoreTechRetessellate_Impl::ApplyOnOneAsset(CoreTechRetessellateAction.cpp) only if the option
+ * RetessellateOptions.RetessellationRule equal EDatasmithCADRetessellationRule::SkipDeletedSurfaces
+ */
+void CopyPatchGroups(FMeshDescription& MeshSource, FMeshDescription& MeshDestination)
+{
+	TPolygonGroupAttributesRef<int32> PatchGroups = MeshDestination.PolygonGroupAttributes().GetAttributesRef<int32>(PolyTriGroups);
+	if (PatchGroups.IsValid() == false)
 	{
-		TPolygonAttributesRef<int32> PatchGroups = OutMeshDescription.PolygonAttributes().GetAttributesRef<int32>(PolyTriGroups);
-		if (PatchGroups.IsValid() == false)
-		{
-			OutMeshDescription.PolygonAttributes().RegisterAttribute<int32>(PolyTriGroups, 1, 0, EMeshAttributeFlags::AutoGenerated);
-			PatchGroups = OutMeshDescription.PolygonAttributes().GetAttributesRef<int32>(PolyTriGroups);
-			check(PatchGroups.IsValid());
-		}
-		return PatchGroups;
+		MeshDestination.PolygonGroupAttributes().RegisterAttribute<int32>(PolyTriGroups, 1, 0, EMeshAttributeFlags::AutoGenerated);
+		PatchGroups = MeshDestination.PolygonGroupAttributes().GetAttributesRef<int32>(PolyTriGroups);
+		check(PatchGroups.IsValid());
 	}
 
-	/**
-	 * Polygon group is an attribute of polygons.As long as the mesh description is empty(no polygon), polygon group cannot be defined.
-	 * The work around used is to create PolygonGroups and to set them a PolygonGroupAttributes.
-	 * To get the existing polygon groups, for each created polygonGroups the PolygonGroupId is got(see GetExistingPatches)
-	 * Warning: CopyPatchGroups is call in FCoreTechRetessellate_Impl::ApplyOnOneAsset(CoreTechRetessellateAction.cpp) only if the option 
-	 * RetessellateOptions.RetessellationRule equal EDatasmithCADRetessellationRule::SkipDeletedSurfaces
-	 */
-	void CopyPatchGroups(FMeshDescription& MeshSource, FMeshDescription& MeshDestination)
+	TSet<int32> PatchIdSet;
+	TPolygonAttributesRef<int32> ElementToGroups = MeshSource.PolygonAttributes().GetAttributesRef<int32>(PolyTriGroups);
+	for (const FPolygonID TriangleID : MeshSource.Polygons().GetElementIDs())
 	{
-		TPolygonGroupAttributesRef<int32> PatchGroups = MeshDestination.PolygonGroupAttributes().GetAttributesRef<int32>(PolyTriGroups);
-		if (PatchGroups.IsValid() == false)
-		{
-			MeshDestination.PolygonGroupAttributes().RegisterAttribute<int32>(PolyTriGroups, 1, 0, EMeshAttributeFlags::AutoGenerated);
-			PatchGroups = MeshDestination.PolygonGroupAttributes().GetAttributesRef<int32>(PolyTriGroups);
-			check(PatchGroups.IsValid());
-		}
+		int32 PatchId = ElementToGroups[TriangleID];
 
-		TSet<int32> PatchIdSet;
-		TPolygonAttributesRef<int32> ElementToGroups = MeshSource.PolygonAttributes().GetAttributesRef<int32>(PolyTriGroups);
-		for (const FPolygonID TriangleID : MeshSource.Polygons().GetElementIDs())
+		bool bIsAlreadyInSet = false;
+		PatchIdSet.Add(PatchId, &bIsAlreadyInSet);
+		if (!bIsAlreadyInSet)
 		{
-			int32 PatchId = ElementToGroups[TriangleID];
-
-			bool bIsAlreadyInSet = false;
-			PatchIdSet.Add(PatchId, &bIsAlreadyInSet);
-			if (!bIsAlreadyInSet)
-			{
-				const FPolygonGroupID PolygonID = MeshDestination.CreatePolygonGroup();
-				PatchGroups[PolygonID] = PatchId;
-			}
+			const FPolygonGroupID PolygonID = MeshDestination.CreatePolygonGroup();
+			PatchGroups[PolygonID] = PatchId;
 		}
 	}
+}
 
-	/**
-	 * See CopyPatchGroups
-	 */
-	void GetExistingPatches(FMeshDescription& MeshDestination, TSet<int32>& OutPatchIdSet)
+/**
+ * See CopyPatchGroups
+ */
+void GetExistingPatches(FMeshDescription& MeshDestination, TSet<int32>& OutPatchIdSet)
+{
+	TPolygonGroupAttributesRef<int32> PatchGroups = MeshDestination.PolygonGroupAttributes().GetAttributesRef<int32>(PolyTriGroups);
+	if (PatchGroups.IsValid() == false)
 	{
-		TPolygonGroupAttributesRef<int32> PatchGroups = MeshDestination.PolygonGroupAttributes().GetAttributesRef<int32>(PolyTriGroups);
-		if (PatchGroups.IsValid() == false)
-		{
-			return;
-		}
+		return;
+	}
 
-		for (const FPolygonGroupID PolygonGroupID : MeshDestination.PolygonGroups().GetElementIDs())
+	for (const FPolygonGroupID PolygonGroupID : MeshDestination.PolygonGroups().GetElementIDs())
+	{
+		int32 PatchId = PatchGroups[PolygonGroupID];
+		if (PatchId > 0)
 		{
-			int32 PatchId = PatchGroups[PolygonGroupID];
-			if (PatchId > 0) 
-			{
-				OutPatchIdSet.Add(PatchId);
-			}
+			OutPatchIdSet.Add(PatchId);
 		}
 	}
+}
 
-	bool FillMesh(const FMeshParameters& MeshParameters, const FImportParameters& ImportParams, FBodyMesh& BodyTessellation, FMeshDescription& MeshDescription)
+bool FillMesh(const FMeshParameters& MeshParameters, const FImportParameters& ImportParams, FBodyMesh& BodyTessellation, FMeshDescription& MeshDescription)
+{
+	const int32 UVChannel = 0;
+	const int32 VertexCountPerFace = 3;
+	const TriangleIndex Clockwise = { 0, 1, 2 };
+	const TriangleIndex CounterClockwise = { 0, 2, 1 };
+
+	TArray<FVertexInstanceID> TriangleVertexInstanceIDs;
+	TriangleVertexInstanceIDs.SetNum(VertexCountPerFace);
+
+	TArray<FVertexInstanceID> MeshVertexInstanceIDs;
+	TArray<uint32> NewFaceIndex;  // new CT face index to remove degenerated face
+
+	// Gather all array data
+	FStaticMeshAttributes Attributes(MeshDescription);
+	TVertexInstanceAttributesRef<FVector3f> VertexInstanceNormals = Attributes.GetVertexInstanceNormals();
+	TVertexInstanceAttributesRef<FVector3f> VertexInstanceTangents = Attributes.GetVertexInstanceTangents();
+	TVertexInstanceAttributesRef<float> VertexInstanceBinormalSigns = Attributes.GetVertexInstanceBinormalSigns();
+	TVertexInstanceAttributesRef<FVector4f> VertexInstanceColors = Attributes.GetVertexInstanceColors();
+	TVertexInstanceAttributesRef<FVector2f> VertexInstanceUVs = Attributes.GetVertexInstanceUVs();
+	TPolygonGroupAttributesRef<FName> PolygonGroupImportedMaterialSlotNames = Attributes.GetPolygonGroupMaterialSlotNames();
+
+	if (!VertexInstanceNormals.IsValid() || !VertexInstanceTangents.IsValid() || !VertexInstanceBinormalSigns.IsValid() || !VertexInstanceColors.IsValid() || !VertexInstanceUVs.IsValid() || !PolygonGroupImportedMaterialSlotNames.IsValid())
 	{
-		const int32 UVChannel = 0;
-		const int32 VertexCountPerFace = 3;
-		const TriangleIndex Clockwise = { 0, 1, 2 };
-		const TriangleIndex CounterClockwise = { 0, 2, 1 };
+		return false;
+	}
 
-		TArray<FVertexInstanceID> TriangleVertexInstanceIDs;
-		TriangleVertexInstanceIDs.SetNum(VertexCountPerFace);
+	// Find all the materials used
+	TMap<uint32, FPolygonGroupID> MaterialToPolygonGroupMapping;
+	for (const FTessellationData& FaceTessellation : BodyTessellation.Faces)
+	{
+		// material is preferred over color
+		MaterialToPolygonGroupMapping.Add(FaceTessellation.MaterialUId ? FaceTessellation.MaterialUId : FaceTessellation.ColorUId, INDEX_NONE);
+	}
 
-		TArray<FVertexInstanceID> MeshVertexInstanceIDs;
-		TArray<uint32> NewFaceIndex;  // new CT face index to remove degenerated face
+	// Add to the mesh, a polygon groups per material
+	for (auto& Material : MaterialToPolygonGroupMapping)
+	{
+		uint32 MaterialHash = Material.Key;
+		FName ImportedSlotName = *LexToString<uint32>(MaterialHash);
 
-		// Gather all array data
-		FStaticMeshAttributes Attributes(MeshDescription);
-		TVertexInstanceAttributesRef<FVector3f> VertexInstanceNormals = Attributes.GetVertexInstanceNormals();
-		TVertexInstanceAttributesRef<FVector3f> VertexInstanceTangents = Attributes.GetVertexInstanceTangents();
-		TVertexInstanceAttributesRef<float> VertexInstanceBinormalSigns = Attributes.GetVertexInstanceBinormalSigns();
-		TVertexInstanceAttributesRef<FVector4f> VertexInstanceColors = Attributes.GetVertexInstanceColors();
-		TVertexInstanceAttributesRef<FVector2f> VertexInstanceUVs = Attributes.GetVertexInstanceUVs();
-		TPolygonGroupAttributesRef<FName> PolygonGroupImportedMaterialSlotNames = Attributes.GetPolygonGroupMaterialSlotNames();
+		FPolygonGroupID PolyGroupID = MeshDescription.CreatePolygonGroup();
+		PolygonGroupImportedMaterialSlotNames[PolyGroupID] = ImportedSlotName;
+		Material.Value = PolyGroupID;
+	}
 
-		if (!VertexInstanceNormals.IsValid() || !VertexInstanceTangents.IsValid() || !VertexInstanceBinormalSigns.IsValid() || !VertexInstanceColors.IsValid() || !VertexInstanceUVs.IsValid() || !PolygonGroupImportedMaterialSlotNames.IsValid())
-		{
-			return false;
-		}
+	if (Algo::AnyOf(BodyTessellation.Faces, [](const FTessellationData& FaceTessellation) { return !FaceTessellation.TexCoordArray.IsEmpty(); }))
+	{
+		VertexInstanceUVs.SetNumChannels(1);
+	}
 
-		// Find all the materials used
-		TMap<uint32, FPolygonGroupID> MaterialToPolygonGroupMapping;
-		for (const FTessellationData& FaceTessellation : BodyTessellation.Faces)
-		{
-			// material is preferred over color
-			MaterialToPolygonGroupMapping.Add(FaceTessellation.MaterialUId ? FaceTessellation.MaterialUId : FaceTessellation.ColorUId, INDEX_NONE);
-		}
+	int32 NbStep = 1;
+	if (MeshParameters.bIsSymmetric)
+	{
+		NbStep = 2;
+	}
 
-		// Add to the mesh, a polygon groups per material
-		for (auto& Material : MaterialToPolygonGroupMapping)
-		{
-			uint32 MaterialHash = Material.Key;
-			FName ImportedSlotName = *LexToString<uint32>(MaterialHash);
+	TSet<int32> PatchIdSet;
+	GetExistingPatches(MeshDescription, PatchIdSet);
+	bool bImportOnlyAlreadyPresent = (bool)PatchIdSet.Num();
 
-			FPolygonGroupID PolyGroupID = MeshDescription.CreatePolygonGroup();
-			PolygonGroupImportedMaterialSlotNames[PolyGroupID] = ImportedSlotName;
-			Material.Value = PolyGroupID;
-		}
+	TPolygonAttributesRef<int32> PatchGroups = EnableCADPatchGroups(MeshDescription);
+	int32 PatchIndex = 0;
+	for (int32 Step = 0; Step < NbStep; ++Step)
+	{
+		// Swap mesh if needed
+		const TriangleIndex& Orientation = (!MeshParameters.bNeedSwapOrientation == (bool)Step) ? CounterClockwise : Clockwise;
+		TArray<int32>& VertexIdSet = (Step == 0) ? BodyTessellation.VertexIds : BodyTessellation.SymmetricVertexIds;
 
-		if (Algo::AnyOf(BodyTessellation.Faces, [](const FTessellationData& FaceTessellation) { return !FaceTessellation.TexCoordArray.IsEmpty(); }))
+		for (FTessellationData& Tessellation : BodyTessellation.Faces)
 		{
-			VertexInstanceUVs.SetNumChannels(1);
-		}
+			if (bImportOnlyAlreadyPresent && !PatchIdSet.Contains(Tessellation.PatchId))
+			{
+				continue;
+			}
 
-		int32 NbStep = 1;
-		if (MeshParameters.bIsSymmetric)
-		{
-			NbStep = 2;
-		}
+			// Get the polygonGroup (material is preferred over color)
+			FMaterialUId GraphicUId = Tessellation.MaterialUId ? Tessellation.MaterialUId : Tessellation.ColorUId;
+			const FPolygonGroupID* PolygonGroupID = MaterialToPolygonGroupMapping.Find(GraphicUId);
+			if (PolygonGroupID == nullptr)
+			{
+				continue;
+			}
 
-		TSet<int32> PatchIdSet;
-		GetExistingPatches(MeshDescription, PatchIdSet);
-		bool bImportOnlyAlreadyPresent = (bool) PatchIdSet.Num();
+			int32 FaceVertexIDs[3];
+			int32 FaceVertexIndices[3];
+			FVector Temp3D = { 0, 0, 0 };
+			FVector2D TexCoord2D = { 0, 0 };
 
-		TPolygonAttributesRef<int32> PatchGroups = EnableCADPatchGroups(MeshDescription);
-		int32 PatchIndex = 0;
-		for (int32 Step = 0; Step < NbStep; ++Step)
-		{
-			// Swap mesh if needed
-			const TriangleIndex& Orientation = (!MeshParameters.bNeedSwapOrientation == (bool)Step) ? CounterClockwise : Clockwise;
-			TArray<int32>& VertexIdSet = (Step == 0) ? BodyTessellation.VertexIds : BodyTessellation.SymmetricVertexIds;
+			MeshVertexInstanceIDs.Empty(Tessellation.VertexIndices.Num());
+			NewFaceIndex.Empty(Tessellation.VertexIndices.Num());
+
+			PatchIndex++;
 
-			for (FTessellationData& Tessellation : BodyTessellation.Faces)
+			// build each valid face i.e. 3 different indexes
+			for (int32 FaceIndex = 0; FaceIndex < Tessellation.VertexIndices.Num(); FaceIndex += VertexCountPerFace)
 			{
-				if (bImportOnlyAlreadyPresent && !PatchIdSet.Contains(Tessellation.PatchId))
+				FaceVertexIndices[0] = Tessellation.PositionIndices[Tessellation.VertexIndices[FaceIndex + Orientation[0]]];
+				FaceVertexIndices[1] = Tessellation.PositionIndices[Tessellation.VertexIndices[FaceIndex + Orientation[1]]];
+				FaceVertexIndices[2] = Tessellation.PositionIndices[Tessellation.VertexIndices[FaceIndex + Orientation[2]]];
+
+				if (FaceVertexIndices[0] == INDEX_NONE || FaceVertexIndices[1] == INDEX_NONE || FaceVertexIndices[2] == INDEX_NONE)
 				{
 					continue;
 				}
 
-				// Get the polygonGroup (material is preferred over color)
-				FMaterialUId GraphicUId = Tessellation.MaterialUId ? Tessellation.MaterialUId : Tessellation.ColorUId;
-				const FPolygonGroupID* PolygonGroupID = MaterialToPolygonGroupMapping.Find(GraphicUId);
-				if (PolygonGroupID == nullptr)
+				// Verify the 3 input indices are not defining a degenerated triangle
+				if (FaceVertexIndices[0] == FaceVertexIndices[1] || FaceVertexIndices[0] == FaceVertexIndices[2] || FaceVertexIndices[1] == FaceVertexIndices[2])
 				{
 					continue;
 				}
 
-				int32 FaceVertexIDs[3];
-				int32 FaceVertexIndices[3];
-				FVector Temp3D = { 0, 0, 0 };
-				FVector2D TexCoord2D = { 0, 0 };
-
-				MeshVertexInstanceIDs.Empty(Tessellation.VertexIndices.Num());
-				NewFaceIndex.Empty(Tessellation.VertexIndices.Num());
+				FaceVertexIDs[0] = VertexIdSet[FaceVertexIndices[0]];
+				FaceVertexIDs[1] = VertexIdSet[FaceVertexIndices[1]];
+				FaceVertexIDs[2] = VertexIdSet[FaceVertexIndices[2]];
 
-				PatchIndex++;
-
-				// build each valid face i.e. 3 different indexes
-				for (int32 FaceIndex = 0; FaceIndex < Tessellation.VertexIndices.Num(); FaceIndex += VertexCountPerFace)
-				{
-					FaceVertexIndices[0] = Tessellation.PositionIndices[Tessellation.VertexIndices[FaceIndex + Orientation[0]]];
-					FaceVertexIndices[1] = Tessellation.PositionIndices[Tessellation.VertexIndices[FaceIndex + Orientation[1]]];
-					FaceVertexIndices[2] = Tessellation.PositionIndices[Tessellation.VertexIndices[FaceIndex + Orientation[2]]];
+				NewFaceIndex.Add(Tessellation.VertexIndices[FaceIndex + Orientation[0]]);
+				NewFaceIndex.Add(Tessellation.VertexIndices[FaceIndex + Orientation[1]]);
+				NewFaceIndex.Add(Tessellation.VertexIndices[FaceIndex + Orientation[2]]);
 
-					if (FaceVertexIndices[0] == INDEX_NONE || FaceVertexIndices[1] == INDEX_NONE || FaceVertexIndices[2] == INDEX_NONE)
-					{
-						continue;
-					}
+				MeshVertexInstanceIDs.Add(TriangleVertexInstanceIDs[0] = MeshDescription.CreateVertexInstance((FVertexID)FaceVertexIDs[0]));
+				MeshVertexInstanceIDs.Add(TriangleVertexInstanceIDs[1] = MeshDescription.CreateVertexInstance((FVertexID)FaceVertexIDs[1]));
+				MeshVertexInstanceIDs.Add(TriangleVertexInstanceIDs[2] = MeshDescription.CreateVertexInstance((FVertexID)FaceVertexIDs[2]));
 
-					// Verify the 3 input indices are not defining a degenerated triangle
-					if (FaceVertexIndices[0] == FaceVertexIndices[1] || FaceVertexIndices[0] == FaceVertexIndices[2] || FaceVertexIndices[1] == FaceVertexIndices[2])
-					{
-						continue;
-					}
+				// Add the triangle as a polygon to the mesh description
+				const FPolygonID PolygonID = MeshDescription.CreatePolygon(*PolygonGroupID, TriangleVertexInstanceIDs);
 
-					FaceVertexIDs[0] = VertexIdSet[FaceVertexIndices[0]];
-					FaceVertexIDs[1] = VertexIdSet[FaceVertexIndices[1]];
-					FaceVertexIDs[2] = VertexIdSet[FaceVertexIndices[2]];
-
-					NewFaceIndex.Add(Tessellation.VertexIndices[FaceIndex + Orientation[0]]);
-					NewFaceIndex.Add(Tessellation.VertexIndices[FaceIndex + Orientation[1]]);
-					NewFaceIndex.Add(Tessellation.VertexIndices[FaceIndex + Orientation[2]]);
-
-					MeshVertexInstanceIDs.Add(TriangleVertexInstanceIDs[0] = MeshDescription.CreateVertexInstance((FVertexID)FaceVertexIDs[0]));
-					MeshVertexInstanceIDs.Add(TriangleVertexInstanceIDs[1] = MeshDescription.CreateVertexInstance((FVertexID)FaceVertexIDs[1]));
-					MeshVertexInstanceIDs.Add(TriangleVertexInstanceIDs[2] = MeshDescription.CreateVertexInstance((FVertexID)FaceVertexIDs[2]));
-
-					// Add the triangle as a polygon to the mesh description
-					const FPolygonID PolygonID = MeshDescription.CreatePolygon(*PolygonGroupID, TriangleVertexInstanceIDs);
-
-					// Set patch id attribute
-					PatchGroups[PolygonID] = Tessellation.PatchId;
-				}
+				// Set patch id attribute
+				PatchGroups[PolygonID] = Tessellation.PatchId;
+			}
 
-				if (!Tessellation.TexCoordArray.IsEmpty())
+			if (!Tessellation.TexCoordArray.IsEmpty())
+			{
+				for (int32 FaceIndex = 0; FaceIndex < MeshVertexInstanceIDs.Num(); FaceIndex += VertexCountPerFace)
 				{
-					for (int32 FaceIndex = 0; FaceIndex < MeshVertexInstanceIDs.Num(); FaceIndex += VertexCountPerFace)
+					for (int32 VertexIndex = 0; VertexIndex < VertexCountPerFace; VertexIndex++)
 					{
-						for (int32 VertexIndex = 0; VertexIndex < VertexCountPerFace; VertexIndex++)
-						{
-							const FVertexInstanceID& VertexInstanceID = MeshVertexInstanceIDs[FaceIndex + Orientation[VertexIndex]];
-							VertexInstanceUVs.Set(VertexInstanceID, UVChannel, FVector2f(Tessellation.TexCoordArray[NewFaceIndex[FaceIndex + Orientation[VertexIndex]]]));
-						}
+						const FVertexInstanceID& VertexInstanceID = MeshVertexInstanceIDs[FaceIndex + Orientation[VertexIndex]];
+						VertexInstanceUVs.Set(VertexInstanceID, UVChannel, FVector2f(Tessellation.TexCoordArray[NewFaceIndex[FaceIndex + Orientation[VertexIndex]]]));
 					}
 				}
+			}
 
-				for (const FVertexInstanceID& VertexInstanceID : MeshVertexInstanceIDs)
-				{
-					VertexInstanceColors[VertexInstanceID] = FLinearColor::White;
-					VertexInstanceTangents[VertexInstanceID] = FVector3f(ForceInitToZero);
-					VertexInstanceBinormalSigns[VertexInstanceID] = 0.0f;
-				}
+			for (const FVertexInstanceID& VertexInstanceID : MeshVertexInstanceIDs)
+			{
+				VertexInstanceColors[VertexInstanceID] = FLinearColor::White;
+				VertexInstanceTangents[VertexInstanceID] = FVector3f(ForceInitToZero);
+				VertexInstanceBinormalSigns[VertexInstanceID] = 0.0f;
+			}
 
-				if (!Step)
+			if (!Step)
+			{
+				FDatasmithUtils::ConvertVectorArray(ImportParams.GetModelCoordSys(), Tessellation.NormalArray);
+				for (FVector& Normal : Tessellation.NormalArray)
 				{
-					FDatasmithUtils::ConvertVectorArray(ImportParams.GetModelCoordSys(), Tessellation.NormalArray);
-					for (FVector& Normal : Tessellation.NormalArray)
-					{
-						Normal = Normal.GetSafeNormal();
-					}
+					Normal = Normal.GetSafeNormal();
 				}
+			}
 
-				if (Tessellation.NormalArray.Num() == 1)
+			if (Tessellation.NormalArray.Num() == 1)
+			{
+				const FVector& Normal = Tessellation.NormalArray[0];
+				for (int32 Index = 0; Index < NewFaceIndex.Num(); Index++)
 				{
-					const FVector& Normal = Tessellation.NormalArray[0];
-					for (int32 Index = 0; Index < NewFaceIndex.Num(); Index++)
-					{
-						const FVertexInstanceID VertexInstanceID = MeshVertexInstanceIDs[Index];
-						VertexInstanceNormals[VertexInstanceID] = (FVector3f)Normal;
-					}
+					const FVertexInstanceID VertexInstanceID = MeshVertexInstanceIDs[Index];
+					VertexInstanceNormals[VertexInstanceID] = (FVector3f)Normal;
 				}
-				else
+			}
+			else
+			{
+				for (int32 IndexFace = 0; IndexFace < NewFaceIndex.Num(); IndexFace += 3)
 				{
-					for (int32 IndexFace = 0; IndexFace < NewFaceIndex.Num(); IndexFace += 3)
+					for (int32 Index = 0; Index < 3; Index++)
 					{
-						for (int32 Index = 0; Index < 3; Index++)
-						{
-							const FVertexInstanceID VertexInstanceID = MeshVertexInstanceIDs[IndexFace + Orientation[Index]];
-							VertexInstanceNormals[VertexInstanceID] = (FVector3f)Tessellation.NormalArray[NewFaceIndex[IndexFace + Orientation[Index]]];
-						}
+						const FVertexInstanceID VertexInstanceID = MeshVertexInstanceIDs[IndexFace + Orientation[Index]];
+						VertexInstanceNormals[VertexInstanceID] = (FVector3f)Tessellation.NormalArray[NewFaceIndex[IndexFace + Orientation[Index]]];
 					}
 				}
+			}
 
-				// compute normals
-				if (Step)
+			// compute normals
+			if (Step)
+			{
+				// compute normals of Symmetric vertex
+				FMatrix SymmetricMatrix;
+				SymmetricMatrix = FDatasmithUtils::GetSymmetricMatrix(MeshParameters.SymmetricOrigin, MeshParameters.SymmetricNormal);
+				for (const FVertexInstanceID& VertexInstanceID : MeshVertexInstanceIDs)
 				{
-					// compute normals of Symmetric vertex
-					FMatrix SymmetricMatrix;
-					SymmetricMatrix = FDatasmithUtils::GetSymmetricMatrix(MeshParameters.SymmetricOrigin, MeshParameters.SymmetricNormal);
-					for (const FVertexInstanceID& VertexInstanceID : MeshVertexInstanceIDs)
-					{
-						VertexInstanceNormals[VertexInstanceID] = FVector4f(SymmetricMatrix.TransformVector((FVector)VertexInstanceNormals[VertexInstanceID]));
-					}
+					VertexInstanceNormals[VertexInstanceID] = FVector4f(SymmetricMatrix.TransformVector((FVector)VertexInstanceNormals[VertexInstanceID]));
 				}
+			}
 
-				if (MeshParameters.bNeedSwapOrientation)
+			if (MeshParameters.bNeedSwapOrientation)
+			{
+				for (int32 Index = 0; Index < MeshVertexInstanceIDs.Num(); Index++)
 				{
-					for (int32 Index = 0; Index < MeshVertexInstanceIDs.Num(); Index++)
-					{
-						VertexInstanceNormals[MeshVertexInstanceIDs[Index]] = VertexInstanceNormals[MeshVertexInstanceIDs[Index]] * -1.f;
-					}
+					VertexInstanceNormals[MeshVertexInstanceIDs[Index]] = VertexInstanceNormals[MeshVertexInstanceIDs[Index]] * -1.f;
 				}
 			}
 		}
-		return true;
 	}
+	return true;
+}
+
+bool ConvertBodyMeshToMeshDescription(const FImportParameters& ImportParams, const FMeshParameters& MeshParameters, FBodyMesh& Body, FMeshDescription& MeshDescription)
+{
+	// in a closed big mesh VertexCount ~ TriangleCount / 2, EdgeCount ~ 1.5* TriangleCount
+	MeshDescription.ReserveNewVertexInstances(Body.VertexArray.Num());
+	MeshDescription.ReserveNewPolygons(Body.TriangleCount);
+	MeshDescription.ReserveNewEdges(Body.TriangleCount * 3);
 
-	bool ConvertBodyMeshToMeshDescription(const FImportParameters& ImportParams, const FMeshParameters& MeshParameters, FBodyMesh& Body, FMeshDescription& MeshDescription)
+	// CoreTech is generating position duplicates. make sure to remove them before filling the mesh description
+	if (Body.VertexArray.Num() == 0)
+	{
+		FillKioVertexPosition(ImportParams, MeshParameters, Body, MeshDescription);
+	}
+	else
 	{
-		// in a closed big mesh VertexCount ~ TriangleCount / 2, EdgeCount ~ 1.5* TriangleCount
-		MeshDescription.ReserveNewVertexInstances(Body.VertexArray.Num());
-		MeshDescription.ReserveNewPolygons(Body.TriangleCount);
-		MeshDescription.ReserveNewEdges(Body.TriangleCount * 3);
+		FillVertexPosition(ImportParams, MeshParameters, Body, MeshDescription);
+	}
 
-		// CoreTech is generating position duplicates. make sure to remove them before filling the mesh description
-		if (Body.VertexArray.Num() == 0)
-		{
-			FillKioVertexPosition(ImportParams, MeshParameters, Body, MeshDescription);
-		}
-		else
-		{
-			FillVertexPosition(ImportParams, MeshParameters, Body, MeshDescription);
-		}
-		
-		if (!FillMesh(MeshParameters, ImportParams, Body, MeshDescription))
-		{
-			return false;
-		}
+	if (!FillMesh(MeshParameters, ImportParams, Body, MeshDescription))
+	{
+		return false;
+	}
 
-		// Orient mesh
-		MeshOperator::OrientMesh(MeshDescription);
+	// Orient mesh
+	MeshOperator::OrientMesh(MeshDescription);
 
-		// Build edge meta data
-		FStaticMeshOperations::DetermineEdgeHardnessesFromVertexInstanceNormals(MeshDescription);
+	// Build edge meta data
+	FStaticMeshOperations::DetermineEdgeHardnessesFromVertexInstanceNormals(MeshDescription);
 
-		return MeshDescription.Polygons().Num() > 0;
-	}
+	return MeshDescription.Polygons().Num() > 0;
+}
 
-	TSharedPtr<IDatasmithUEPbrMaterialElement> CreateDefaultUEPbrMaterial()
-	{
-		// Take the Material diffuse color and connect it to the BaseColor of a UEPbrMaterial
-		TSharedRef<IDatasmithUEPbrMaterialElement> MaterialElement = FDatasmithSceneFactory::CreateUEPbrMaterial(TEXT("0"));
-		MaterialElement->SetLabel(TEXT("DefaultCADImportMaterial"));
+TSharedPtr<IDatasmithUEPbrMaterialElement> CreateDefaultUEPbrMaterial()
+{
+	// Take the Material diffuse color and connect it to the BaseColor of a UEPbrMaterial
+	TSharedRef<IDatasmithUEPbrMaterialElement> MaterialElement = FDatasmithSceneFactory::CreateUEPbrMaterial(TEXT("0"));
+	MaterialElement->SetLabel(TEXT("DefaultCADImportMaterial"));
+
+	FLinearColor LinearColor = FLinearColor::FromPow22Color(FColor(200, 200, 200, 255));
+	IDatasmithMaterialExpressionColor* ColorExpression = MaterialElement->AddMaterialExpression<IDatasmithMaterialExpressionColor>();
+	ColorExpression->SetName(TEXT("Base Color"));
+	ColorExpression->GetColor() = LinearColor;
+	MaterialElement->GetBaseColor().SetExpression(ColorExpression);
+	MaterialElement->SetParentLabel(TEXT("M_DatasmithCAD"));
+
+	return MaterialElement;
+}
 
-		FLinearColor LinearColor = FLinearColor::FromPow22Color(FColor(200, 200, 200, 255));
-		IDatasmithMaterialExpressionColor* ColorExpression = MaterialElement->AddMaterialExpression<IDatasmithMaterialExpressionColor>();
-		ColorExpression->SetName(TEXT("Base Color"));
-		ColorExpression->GetColor() = LinearColor;
-		MaterialElement->GetBaseColor().SetExpression(ColorExpression);
-		MaterialElement->SetParentLabel(TEXT("M_DatasmithCAD"));
+TSharedPtr<IDatasmithUEPbrMaterialElement> CreateUEPbrMaterialFromColor(const FColor& InColor)
+{
+	FString Name = FString::FromInt(BuildColorUId(InColor));
+	FString Label = FString::Printf(TEXT("color_%02x%02x%02x%02x"), InColor.R, InColor.G, InColor.B, InColor.A);
 
-		return MaterialElement;
-	}
+	// Take the Material diffuse color and connect it to the BaseColor of a UEPbrMaterial
+	TSharedRef<IDatasmithUEPbrMaterialElement> MaterialElement = FDatasmithSceneFactory::CreateUEPbrMaterial(*Name);
+	MaterialElement->SetLabel(*Label);
 
-	TSharedPtr<IDatasmithUEPbrMaterialElement> CreateUEPbrMaterialFromColor(const FColor& InColor)
-	{
-		FString Name = FString::FromInt(BuildColorUId(InColor));
-		FString Label = FString::Printf(TEXT("color_%02x%02x%02x%02x"), InColor.R, InColor.G, InColor.B, InColor.A);
+	FLinearColor LinearColor = FLinearColor::FromSRGBColor(InColor);
 
-		// Take the Material diffuse color and connect it to the BaseColor of a UEPbrMaterial
-		TSharedRef<IDatasmithUEPbrMaterialElement> MaterialElement = FDatasmithSceneFactory::CreateUEPbrMaterial(*Name);
-		MaterialElement->SetLabel(*Label);
+	IDatasmithMaterialExpressionColor* ColorExpression = MaterialElement->AddMaterialExpression<IDatasmithMaterialExpressionColor>();
+	ColorExpression->SetName(TEXT("Base Color"));
+	ColorExpression->GetColor() = LinearColor;
 
-		FLinearColor LinearColor = FLinearColor::FromSRGBColor(InColor);
+	MaterialElement->GetBaseColor().SetExpression(ColorExpression);
 
-		IDatasmithMaterialExpressionColor* ColorExpression = MaterialElement->AddMaterialExpression<IDatasmithMaterialExpressionColor>();
-		ColorExpression->SetName(TEXT("Base Color"));
-		ColorExpression->GetColor() = LinearColor;
+	if (LinearColor.A < 1.0f)
+	{
+		MaterialElement->SetBlendMode(/*EBlendMode::BLEND_Translucent*/2);
 
-		MaterialElement->GetBaseColor().SetExpression(ColorExpression);
+		IDatasmithMaterialExpressionScalar* Scalar = MaterialElement->AddMaterialExpression<IDatasmithMaterialExpressionScalar>();
+		Scalar->GetScalar() = LinearColor.A;
+		Scalar->SetName(TEXT("Opacity Level"));
 
-		if (LinearColor.A < 1.0f)
-		{
-			MaterialElement->SetBlendMode(/*EBlendMode::BLEND_Translucent*/2);
+		MaterialElement->GetOpacity().SetExpression(Scalar);
+		MaterialElement->SetParentLabel(TEXT("M_DatasmithCADTransparent"));
+	}
+	else
+	{
+		MaterialElement->SetParentLabel(TEXT("M_DatasmithCAD"));
+	}
 
-			IDatasmithMaterialExpressionScalar* Scalar = MaterialElement->AddMaterialExpression<IDatasmithMaterialExpressionScalar>();
-			Scalar->GetScalar() = LinearColor.A;
-			Scalar->SetName(TEXT("Opacity Level"));
+	return MaterialElement;
+}
 
-			MaterialElement->GetOpacity().SetExpression(Scalar);
-			MaterialElement->SetParentLabel(TEXT("M_DatasmithCADTransparent"));
-		}
-		else
-		{
-			MaterialElement->SetParentLabel(TEXT("M_DatasmithCAD"));
-		}
+TSharedPtr<IDatasmithUEPbrMaterialElement> CreateUEPbrMaterialFromMaterial(FCADMaterial& InMaterial, TSharedRef<IDatasmithScene> Scene)
+{
+	FString Name = FString::FromInt(BuildMaterialUId(InMaterial));
 
-		return MaterialElement;
+	// Take the Material diffuse color and connect it to the BaseColor of a UEPbrMaterial
+	TSharedRef<IDatasmithUEPbrMaterialElement> MaterialElement = FDatasmithSceneFactory::CreateUEPbrMaterial(*Name);
+	FString MaterialLabel(InMaterial.MaterialName);
+	if (MaterialLabel.IsEmpty())
+	{
+		MaterialLabel = TEXT("Material");
 	}
+	MaterialElement->SetLabel(*MaterialLabel);
 
-	TSharedPtr<IDatasmithUEPbrMaterialElement> CreateUEPbrMaterialFromMaterial(FCADMaterial& InMaterial, TSharedRef<IDatasmithScene> Scene)
+	// Set a diffuse color if there's nothing in the BaseColor
+	if (MaterialElement->GetBaseColor().GetExpression() == nullptr)
 	{
-		FString Name = FString::FromInt(BuildMaterialUId(InMaterial));
-
-		// Take the Material diffuse color and connect it to the BaseColor of a UEPbrMaterial
-		TSharedRef<IDatasmithUEPbrMaterialElement> MaterialElement = FDatasmithSceneFactory::CreateUEPbrMaterial(*Name);
-		FString MaterialLabel(InMaterial.MaterialName);
-		if (MaterialLabel.IsEmpty())
-		{
-			MaterialLabel = TEXT("Material");
-		}
-		MaterialElement->SetLabel(*MaterialLabel);
+		FLinearColor LinearColor = FLinearColor::FromSRGBColor(InMaterial.Diffuse);
 
-		// Set a diffuse color if there's nothing in the BaseColor
-		if (MaterialElement->GetBaseColor().GetExpression() == nullptr)
-		{
-			FLinearColor LinearColor = FLinearColor::FromSRGBColor(InMaterial.Diffuse);
+		IDatasmithMaterialExpressionColor* ColorExpression = MaterialElement->AddMaterialExpression<IDatasmithMaterialExpressionColor>();
+		ColorExpression->SetName(TEXT("Base Color"));
+		ColorExpression->GetColor() = LinearColor;
 
-			IDatasmithMaterialExpressionColor* ColorExpression = MaterialElement->AddMaterialExpression<IDatasmithMaterialExpressionColor>();
-			ColorExpression->SetName(TEXT("Base Color"));
-			ColorExpression->GetColor() = LinearColor;
+		MaterialElement->GetBaseColor().SetExpression(ColorExpression);
+	}
 
-			MaterialElement->GetBaseColor().SetExpression(ColorExpression);
-		}
+	if (InMaterial.Transparency > 0.0f)
+	{
+		MaterialElement->SetBlendMode(/*EBlendMode::BLEND_Translucent*/2);
+		IDatasmithMaterialExpressionScalar* Scalar = MaterialElement->AddMaterialExpression<IDatasmithMaterialExpressionScalar>();
+		Scalar->GetScalar() = InMaterial.Transparency;
+		Scalar->SetName(TEXT("Opacity Level"));
+		MaterialElement->GetOpacity().SetExpression(Scalar);
+		MaterialElement->SetParentLabel(TEXT("M_DatasmithCADTransparent"));
+	}
+	else
+	{
+		MaterialElement->SetParentLabel(TEXT("M_DatasmithCAD"));
+	}
 
-		if (InMaterial.Transparency > 0.0f)
-		{
-			MaterialElement->SetBlendMode(/*EBlendMode::BLEND_Translucent*/2);
-			IDatasmithMaterialExpressionScalar* Scalar = MaterialElement->AddMaterialExpression<IDatasmithMaterialExpressionScalar>();
-			Scalar->GetScalar() = InMaterial.Transparency;
-			Scalar->SetName(TEXT("Opacity Level"));
-			MaterialElement->GetOpacity().SetExpression(Scalar);
-			MaterialElement->SetParentLabel(TEXT("M_DatasmithCADTransparent"));
-		}
-		else
-		{
-			MaterialElement->SetParentLabel(TEXT("M_DatasmithCAD"));
-		}
+	return MaterialElement;
+}
 
-		return MaterialElement;
-	}
 }
