Improve CAD File parsing according to TechSoft recommendation
Process Refence once
Fix material assignation
#swarm https://p4-swarm.epicgames.net/reviews/20942477

diff --git a/Engine/Plugins/Enterprise/DatasmithCADImporter/Source/CADInterfaces/Private/CADSceneGraph.cpp b/Engine/Plugins/Enterprise/DatasmithCADImporter/Source/CADInterfaces/Private/CADSceneGraph.cpp
--- a/Engine/Plugins/Enterprise/DatasmithCADImporter/Source/CADInterfaces/Private/CADSceneGraph.cpp
+++ b/Engine/Plugins/Enterprise/DatasmithCADImporter/Source/CADInterfaces/Private/CADSceneGraph.cpp
@@ -132,8 +132,10 @@ FArchiveInstance& FArchiveSceneGraph::AddInstance(const FArchiveCADObject& Paren
 
 	int32 Index = Instances.Emplace(InstanceId, Parent);
 	CADIdToIndex.Add(Index);
+	FArchiveInstance& Instance = Instances[Index];
+	Instance.Unit = Parent.Unit;
 
-	return Instances[Index];
+	return Instance;
 }
 
 FArchiveInstance& FArchiveSceneGraph::GetInstance(FCadId CadId)
diff --git a/Engine/Plugins/Enterprise/DatasmithCADImporter/Source/CADInterfaces/Private/TechSoftFileParser.cpp b/Engine/Plugins/Enterprise/DatasmithCADImporter/Source/CADInterfaces/Private/TechSoftFileParser.cpp
--- a/Engine/Plugins/Enterprise/DatasmithCADImporter/Source/CADInterfaces/Private/TechSoftFileParser.cpp
+++ b/Engine/Plugins/Enterprise/DatasmithCADImporter/Source/CADInterfaces/Private/TechSoftFileParser.cpp
@@ -18,6 +18,47 @@ namespace CADLibrary
 
 #ifdef USE_TECHSOFT_SDK
 
+namespace ProductOccurrence
+{
+
+bool HasNoPartNoChildButHasReference(const A3DAsmProductOccurrenceData& OccurrenceData)
+{
+	return ((OccurrenceData.m_pPrototype != nullptr || OccurrenceData.m_pExternalData != nullptr) && OccurrenceData.m_pPart == nullptr && OccurrenceData.m_uiPOccurrencesSize == 0);
+}
+
+bool HasNoPartButHasReference(const A3DAsmProductOccurrenceData& OccurrenceData)
+{
+	return ((OccurrenceData.m_pPrototype != nullptr || OccurrenceData.m_pExternalData != nullptr) && OccurrenceData.m_pPart == nullptr);
+}
+
+bool HasNoChildButHasReference(const A3DAsmProductOccurrenceData& OccurrenceData)
+{
+	return ((OccurrenceData.m_pPrototype != nullptr || OccurrenceData.m_pExternalData != nullptr) && OccurrenceData.m_uiPOccurrencesSize == 0);
+}
+
+bool HasNoPartNoChild(const A3DAsmProductOccurrenceData& OccurrenceData)
+{
+	return (OccurrenceData.m_pPart == nullptr && OccurrenceData.m_uiPOccurrencesSize == 0);
+}
+
+bool HasPartOrChild(const A3DAsmProductOccurrenceData& OccurrenceData)
+{
+	return (OccurrenceData.m_pPart != nullptr || OccurrenceData.m_uiPOccurrencesSize != 0);
+}
+
+bool PrototypeIsValid(const A3DAsmProductOccurrenceData& OccurrenceData)
+{
+	return (OccurrenceData.m_pPrototype != nullptr || OccurrenceData.m_pExternalData != nullptr);
+}
+
+A3DAsmProductOccurrence* GetReference(const A3DAsmProductOccurrenceData& OccurrenceData)
+{
+	return OccurrenceData.m_pPrototype ? OccurrenceData.m_pPrototype : OccurrenceData.m_pExternalData;
+}
+
+
+}
+
 namespace TechSoftFileParserImpl
 {
 
@@ -521,6 +562,8 @@ void FTechSoftFileParser::CountUnderModel()
 			CountUnderOccurrence(ModelFileData->m_ppPOccurrences[Index]);
 		}
 	}
+
+	ReferenceCache.Empty();
 }
 
 ECADParsingResult FTechSoftFileParser::TraverseModel()
@@ -533,21 +576,27 @@ ECADParsingResult FTechSoftFileParser::TraverseModel()
 
 	FArchiveInstance EmptyInstance;
 	FArchiveReference& Reference = SceneGraph.AddReference(EmptyInstance);
-	ExtractMetaData(ModelFile.Get(), Reference);
 	ExtractSpecificMetaData(ModelFile.Get(), Reference);
-	BuildReferenceName(Reference);
 	Reference.Unit = FileUnit;
 
-	for (uint32 Index = 0; Index < ModelFileData->m_uiPOccurrencesSize; ++Index)
+	if(ModelFileData->m_uiPOccurrencesSize == 0)
 	{
-		if (IsConfigurationSet(ModelFileData->m_ppPOccurrences[Index]))
-		{
-			TraverseConfigurationSet(ModelFileData->m_ppPOccurrences[Index], Reference);
-		}
-		else
-		{
-			TraverseReference(ModelFileData->m_ppPOccurrences[Index], Reference);
-		}
+		CADFileData.AddWarningMessages(FString::Printf(TEXT("File %s is empty."), *CADFileData.GetCADFileDescription().GetFileName()));
+		return ECADParsingResult::ProcessFailed;
+	}
+
+	if (ModelFileData->m_uiPOccurrencesSize > 1)
+	{
+		CADFileData.AddWarningMessages(FString::Printf(TEXT("File %s has many root components, only the first is loaded."), *CADFileData.GetCADFileDescription().GetFileName()));
+	}
+
+	if (IsConfigurationSet(ModelFileData->m_ppPOccurrences[0]))
+	{
+		TraverseConfigurationSet(ModelFileData->m_ppPOccurrences[0], Reference);
+	}
+	else
+	{
+		TraverseReference(ModelFileData->m_ppPOccurrences[0], Reference);
 	}
 
 	return ECADParsingResult::ProcessOk;
@@ -567,11 +616,8 @@ void FTechSoftFileParser::TraverseConfigurationSet(const A3DAsmProductOccurrence
 
 	const FString& ConfigurationToLoad = CADFileData.GetCADFileDescription().GetConfiguration();
 
-	A3DMiscTransformation* Location = ConfigurationSetData->m_pLocation;
-	if (Location)
-	{
-		ExtractTransformation(Location, Reference);
-	}
+	A3DMiscTransformation* Transform = ConfigurationSetData->m_pLocation;
+	ExtractTransformation(Transform, Reference);
 
 	TUniqueTSObj<A3DAsmProductOccurrenceData> ConfigurationData;
 	for (unsigned int Index = 0; Index < ConfigurationSetData->m_uiPOccurrencesSize; ++Index)
@@ -609,7 +655,7 @@ void FTechSoftFileParser::TraverseConfigurationSet(const A3DAsmProductOccurrence
 
 	if (ConfigurationToLoad.IsEmpty())
 	{
-		// no default configuration, traverse the first configuration
+		// no default configuration, traverse the first occurence
 		for (unsigned int Index = 0; Index < ConfigurationSetData->m_uiPOccurrencesSize; ++Index)
 		{
 			ConfigurationData.FillFrom(ConfigurationSetData->m_ppPOccurrences[Index]);
@@ -698,6 +744,8 @@ void FTechSoftFileParser::TraverseReference(const A3DAsmProductOccurrence* A3DRe
 		return;
 	}
 
+	ExtractMetaData(A3DReferencePtr, Reference);
+
 	if (Reference.bIsRemoved || !Reference.bShow)
 	{
 		SceneGraph.RemoveLastReference();
@@ -706,17 +754,12 @@ void FTechSoftFileParser::TraverseReference(const A3DAsmProductOccurrence* A3DRe
 
 	ExtractSpecificMetaData(A3DReferencePtr, Reference);
 
-	if (!Reference.IsNameDefined())
-	{
-		BuildReferenceName(Reference);
-	}
+	BuildReferenceName(Reference);
 
 	FMatrix ReferenceMatrix = FMatrix::Identity;
-	A3DMiscTransformation* Location = ReferenceData->m_pLocation;
-	if (Location)
-	{
-		ExtractTransformation(Location, Reference);
-	}
+	A3DMiscTransformation* Transform = ReferenceData->m_pLocation;
+	ExtractTransformation(Transform, Reference);
+	
 	Reference.TransformMatrix = Reference.TransformMatrix * ReferenceMatrix;
 
 	FArchiveInstance EmptyInstance;
@@ -781,23 +824,40 @@ void FTechSoftFileParser::TraverseOccurrence(const A3DAsmProductOccurrence* Occu
 	ExtractSpecificMetaData(OccurrencePtr, Instance);
 	BuildInstanceName(Instance, ParentReference);
 
-	A3DMiscTransformation* Location = OccurrenceData->m_pLocation;
+	A3DMiscTransformation* Transform = OccurrenceData->m_pLocation;
+	
+	A3DAsmProductOccurrence* ReferencePtr = ProductOccurrence::GetReference(*OccurrenceData);
+
+	// Is the Reference already processed ?
+	FCadId* ReferenceId = ReferenceCache.Find(ReferencePtr);
+	if (ReferenceId != nullptr)
+	{
+		Instance.ReferenceNodeId = *ReferenceId;
+
+		if (SceneGraph.IsAUnloadedReference(Instance.ReferenceNodeId))
+		{
+			Instance.bIsExternalReference = true;
+		}
+
+		ExtractTransformation(Transform, Instance);
+		return;
+	}
 
 	FArchiveUnloadedReference& UnloadedReference = SceneGraph.AddUnloadedReference(Instance);
-	if (OccurrenceData->m_pPrototype)
+
+	bool bIsUnloaded = OccurrenceData->m_uiProductFlags & A3D_PRODUCT_FLAG_EXTERNAL_REFERENCE;
+
+	// Extract metadata and define if it's an unloaded reference or not
+	if (ReferencePtr)
 	{
-		ProcessPrototype(OccurrenceData->m_pPrototype, UnloadedReference, &Location);
+		ProcessPrototype(ReferencePtr, UnloadedReference, &Transform);
 	}
 	else
 	{
 		UnloadedReference.bIsUnloaded = false;
 	}
 
-	Instance.Unit = ParentReference.Unit;
-	if (Location)
-	{
-		ExtractTransformation(Location, Instance);
-	}
+	ExtractTransformation(Transform, Instance);
 
 	if (UnloadedReference.bIsUnloaded)
 	{
@@ -809,6 +869,8 @@ void FTechSoftFileParser::TraverseOccurrence(const A3DAsmProductOccurrence* Occu
 		Instance.bIsExternalReference = false;
 		ProcessReference(OccurrencePtr, Instance, NewReference);
 	}
+
+	ReferenceCache.Add(ReferencePtr, Instance.ReferenceNodeId);
 }
 
 void FTechSoftFileParser::ProcessReference(const A3DAsmProductOccurrence* OccurrencePtr, FArchiveInstance& Instance, FArchiveReference& Reference)
@@ -822,21 +884,21 @@ void FTechSoftFileParser::ProcessReference(const A3DAsmProductOccurrence* Occurr
 		Reference.Label = Instance.Label;
 	}
 
-	while (OccurrenceData->m_pPrototype != nullptr && OccurrenceData->m_pPart == nullptr && OccurrenceData->m_uiPOccurrencesSize == 0 && OccurrenceData->m_pExternalData == nullptr)
+	while (ProductOccurrence::HasNoPartNoChildButHasReference(*OccurrenceData))
 	{
-		CachedOccurrencePtr = OccurrenceData->m_pPrototype;
-		OccurrenceData.FillFrom(OccurrenceData->m_pPrototype);
+		CachedOccurrencePtr = ProductOccurrence::GetReference(*OccurrenceData);
+		OccurrenceData.FillFrom(CachedOccurrencePtr);
 	}
 
-	if(OccurrenceData->m_pPart == nullptr && OccurrenceData->m_uiPOccurrencesSize == 0 && OccurrenceData->m_pExternalData == nullptr)
+	if(ProductOccurrence::HasNoPartNoChild(*OccurrenceData))
 	{
 		return;
 	}
 
 	// Add part
-	while (OccurrenceData->m_pPrototype != nullptr && OccurrenceData->m_pPart == nullptr)
+	while (ProductOccurrence::HasNoPartButHasReference(*OccurrenceData))
 	{
-		OccurrenceData.FillFrom(OccurrenceData->m_pPrototype);
+		OccurrenceData.FillFrom(ProductOccurrence::GetReference(*OccurrenceData));
 	}
 	if (OccurrenceData->m_pPart != nullptr)
 	{
@@ -846,9 +908,9 @@ void FTechSoftFileParser::ProcessReference(const A3DAsmProductOccurrence* Occurr
 
 	// Add Occurrence's Children
 	OccurrenceData.FillFrom(CachedOccurrencePtr);
-	while (OccurrenceData->m_pPrototype != nullptr && OccurrenceData->m_uiPOccurrencesSize == 0)
+	while (ProductOccurrence::HasNoChildButHasReference(*OccurrenceData))
 	{
-		OccurrenceData.FillFrom(OccurrenceData->m_pPrototype);
+		OccurrenceData.FillFrom(ProductOccurrence::GetReference(*OccurrenceData));
 	}
 
 	uint32 ChildrenCount = OccurrenceData->m_uiPOccurrencesSize;
@@ -857,29 +919,6 @@ void FTechSoftFileParser::ProcessReference(const A3DAsmProductOccurrence* Occurr
 	{
 		TraverseOccurrence(Children[Index], Reference);
 	}
-
-	// Add External data
-	OccurrenceData.FillFrom(CachedOccurrencePtr);
-	while (OccurrenceData->m_pPrototype != nullptr && OccurrenceData->m_pExternalData == nullptr)
-	{
-		OccurrenceData.FillFrom(OccurrenceData->m_pPrototype);
-	}
-	if (OccurrenceData->m_pExternalData != nullptr)
-	{
-		TUniqueTSObj<A3DAsmProductOccurrenceData> ExternalData(OccurrenceData->m_pExternalData);
-		if (ExternalData->m_pPart != nullptr)
-		{
-			A3DAsmPartDefinition* PartDefinition = ExternalData->m_pPart;
-			TraversePartDefinition(PartDefinition, Reference);
-		}
-
-		uint32 ExternalChildrenCount = ExternalData->m_uiPOccurrencesSize;
-		A3DAsmProductOccurrence** ExternalChildren = ExternalData->m_ppPOccurrences;
-		for (uint32 Index = 0; Index < ExternalChildrenCount; ++Index)
-		{
-			TraverseOccurrence(ExternalChildren[Index], Reference);
-		}
-	}
 }
 
 void FTechSoftFileParser::CountUnderOccurrence(const A3DAsmProductOccurrence* Occurrence)
@@ -888,24 +927,34 @@ void FTechSoftFileParser::CountUnderOccurrence(const A3DAsmProductOccurrence* Oc
 	if (Occurrence && OccurrenceData.IsValid())
 	{
 		ComponentCount[EComponentType::Instance]++;
+
+		A3DAsmProductOccurrence* ReferencePtr = ProductOccurrence::GetReference(*OccurrenceData);
+
+		// Is the Reference already processed ?
+		FCadId* ReferenceId = ReferenceCache.Find(ReferencePtr);
+		if (ReferenceId != nullptr)
+		{
+			return;
+		}
+		ReferenceCache.Add(ReferencePtr, 1);
 		ComponentCount[EComponentType::Reference]++;
 
 		const A3DAsmProductOccurrence* CachedOccurrencePtr = Occurrence;
-		while (OccurrenceData->m_pPrototype != nullptr && OccurrenceData->m_pPart == nullptr && OccurrenceData->m_uiPOccurrencesSize == 0 && OccurrenceData->m_pExternalData == nullptr)
+		while (ProductOccurrence::HasNoPartNoChildButHasReference(*OccurrenceData))
 		{
-			CachedOccurrencePtr = OccurrenceData->m_pPrototype;
-			OccurrenceData.FillFrom(OccurrenceData->m_pPrototype);
+			CachedOccurrencePtr = ProductOccurrence::GetReference(*OccurrenceData);
+			OccurrenceData.FillFrom(CachedOccurrencePtr);
 		}
 
-		if (OccurrenceData->m_pPart == nullptr && OccurrenceData->m_uiPOccurrencesSize == 0 && OccurrenceData->m_pExternalData == nullptr)
+		if (ProductOccurrence::HasNoPartNoChild(*OccurrenceData))
 		{
 			return;
 		}
 
 		// count under part
-		while (OccurrenceData->m_pPrototype != nullptr && OccurrenceData->m_pPart == nullptr)
+		while (ProductOccurrence::HasNoPartButHasReference(*OccurrenceData))
 		{
-			OccurrenceData.FillFrom(OccurrenceData->m_pPrototype);
+			OccurrenceData.FillFrom(ProductOccurrence::GetReference(*OccurrenceData));
 		}
 		if (OccurrenceData->m_pPart != nullptr)
 		{
@@ -914,9 +963,9 @@ void FTechSoftFileParser::CountUnderOccurrence(const A3DAsmProductOccurrence* Oc
 
 		// count under Occurrence
 		OccurrenceData.FillFrom(CachedOccurrencePtr);
-		while (OccurrenceData->m_pPrototype != nullptr && OccurrenceData->m_uiPOccurrencesSize == 0)
+		while (ProductOccurrence::HasNoChildButHasReference(*OccurrenceData))
 		{
-			OccurrenceData.FillFrom(OccurrenceData->m_pPrototype);
+			OccurrenceData.FillFrom(ProductOccurrence::GetReference(*OccurrenceData));
 		}
 
 		uint32 ChildrenCount = OccurrenceData->m_uiPOccurrencesSize;
@@ -925,59 +974,14 @@ void FTechSoftFileParser::CountUnderOccurrence(const A3DAsmProductOccurrence* Oc
 		{
 			CountUnderOccurrence(Children[Index]);
 		}
-
-		// count under External data
-		OccurrenceData.FillFrom(CachedOccurrencePtr);
-		while (OccurrenceData->m_pPrototype != nullptr && OccurrenceData->m_pExternalData == nullptr)
-		{
-			OccurrenceData.FillFrom(OccurrenceData->m_pPrototype);
-		}
-		if (OccurrenceData->m_pExternalData != nullptr)
-		{
-			TUniqueTSObj<A3DAsmProductOccurrenceData> ExternalData(OccurrenceData->m_pExternalData);
-			if (ExternalData->m_pPart != nullptr)
-			{
-				CountUnderPartDefinition(ExternalData->m_pPart);
-			}
-
-			uint32 ExternalChildrenCount = ExternalData->m_uiPOccurrencesSize;
-			A3DAsmProductOccurrence** ExternalChildren = ExternalData->m_ppPOccurrences;
-			for (uint32 Index = 0; Index < ExternalChildrenCount; ++Index)
-			{
-				CountUnderOccurrence(ExternalChildren[Index]);
-			}
-		}
 	}
 }
 
-void FTechSoftFileParser::CountUnderSubPrototype(const A3DAsmProductOccurrence* InPrototypePtr)
-{
-	TUniqueTSObj<A3DAsmProductOccurrenceData> SubPrototypeData(InPrototypePtr);
-	if (!SubPrototypeData.IsValid())
-	{
-		return;
-	}
-
-	for (uint32 Index = 0; Index < SubPrototypeData->m_uiPOccurrencesSize; ++Index)
-	{
-		CountUnderOccurrence(SubPrototypeData->m_ppPOccurrences[Index]);
-	}
-
-	if (SubPrototypeData->m_pPart)
-	{
-		CountUnderPartDefinition(SubPrototypeData->m_pPart);
-	}
-
-	if (SubPrototypeData->m_pPrototype && !SubPrototypeData->m_uiPOccurrencesSize && !SubPrototypeData->m_pPart)
-	{
-		CountUnderSubPrototype(SubPrototypeData->m_pPrototype);
-	}
-}
-
-void FTechSoftFileParser::ProcessPrototype(const A3DAsmProductOccurrence* InPrototypePtr, FArchiveUnloadedReference& OutReference, A3DMiscTransformation** OutLocation)
+void FTechSoftFileParser::ProcessPrototype(const A3DAsmProductOccurrence* InPrototypePtr, FArchiveUnloadedReference& OutReference, A3DMiscTransformation** OutTransform)
 {
 	const A3DAsmProductOccurrence* PrototypePtr = InPrototypePtr;
 	TUniqueTSObj<A3DAsmProductOccurrenceData> PrototypeData;
+
 	while (PrototypePtr)
 	{
 		PrototypeData.FillFrom(PrototypePtr);
@@ -986,11 +990,11 @@ void FTechSoftFileParser::ProcessPrototype(const A3DAsmProductOccurrence* InProt
 			return;
 		}
 
-		if (PrototypeData->m_pPart != nullptr || PrototypeData->m_uiPOccurrencesSize != 0 || PrototypeData->m_pExternalData != nullptr || PrototypeData->m_pPrototype == nullptr)
-		{
-			ExtractMetaData(PrototypePtr, OutReference);
-			ExtractSpecificMetaData(PrototypePtr, OutReference);
+		ExtractMetaData(PrototypePtr, OutReference);
+		ExtractSpecificMetaData(PrototypePtr, OutReference);
 
+		if (OutReference.ExternalFile.IsEmpty())
+		{
 			TUniqueTSObj<A3DUTF8Char*> FilePathUTF8Ptr;
 			FilePathUTF8Ptr.FillWith(&TechSoftInterface::GetFilePathName, PrototypePtr);
 			if (!FilePathUTF8Ptr.IsValid() || *FilePathUTF8Ptr == nullptr)
@@ -1009,39 +1013,40 @@ void FTechSoftFileParser::ProcessPrototype(const A3DAsmProductOccurrence* InProt
 			}
 		}
 		
-		// Case of DWG file, external reference is in fact internal data as a part or a component
-		// Case of CATProduct referencing CGR files: unloaded references are saved in ExternalData. In this case non empty ExternalFile means unloaded reference
-		if (PrototypeData->m_pPart != nullptr|| PrototypeData->m_uiPOccurrencesSize != 0 || (PrototypeData->m_pExternalData != nullptr && OutReference.ExternalFile.IsEmpty()))
+		if (ProductOccurrence::HasPartOrChild(*PrototypeData) )
 		{
 			OutReference.bIsUnloaded = false;
 			PrototypePtr = nullptr;
 		}
 		else
 		{
-			PrototypePtr = PrototypeData->m_pPrototype;
+			PrototypePtr = ProductOccurrence::GetReference(*PrototypeData);
 		}
 
-		if (!*OutLocation)
+		if (!*OutTransform)
 		{
-			*OutLocation = PrototypeData->m_pLocation;
+			*OutTransform = PrototypeData->m_pLocation;
 		}
 	}
 
-	if (!*OutLocation)
+	if (!*OutTransform)
 	{
-		while (PrototypeData.IsValid() && PrototypeData->m_pLocation == nullptr && PrototypeData->m_pPrototype != nullptr)
+		while (PrototypeData->m_pLocation == nullptr && ProductOccurrence::PrototypeIsValid(*PrototypeData))
 		{
-			PrototypeData.FillFrom(PrototypeData->m_pPrototype);
+			PrototypeData.FillFrom(ProductOccurrence::GetReference(*PrototypeData));
 		}
 		if (PrototypeData.IsValid())
 		{
-			*OutLocation = PrototypeData->m_pLocation;
+			*OutTransform = PrototypeData->m_pLocation;
 		}
 	}
 
 	if (OutReference.bIsUnloaded)
 	{
-		OutReference.Label = OutReference.ExternalFile.GetFileName();
+		if(OutReference.Label.IsEmpty())
+		{
+			OutReference.Label = OutReference.ExternalFile.GetFileName();
+		}
 	}
 	else
 	{
@@ -1051,17 +1056,6 @@ void FTechSoftFileParser::ProcessPrototype(const A3DAsmProductOccurrence* InProt
 	BuildReferenceName(OutReference);
 }
 
-void FTechSoftFileParser::CountUnderPrototype(const A3DAsmProductOccurrence* Prototype)
-{
-	TUniqueTSObj<A3DAsmProductOccurrenceData> PrototypeData(Prototype);
-	if (!PrototypeData.IsValid())
-	{
-		return;
-	}
-
-	ComponentCount[EComponentType::Reference] ++;
-}
-
 void FTechSoftFileParser::TraversePartDefinition(const A3DAsmPartDefinition* PartDefinitionPtr, FArchiveReference& Part)
 {
 	ExtractMetaData(PartDefinitionPtr, Part);
@@ -1247,7 +1241,7 @@ void FTechSoftFileParser::ExtractMetaData(const A3DEntity* Entity, FArchiveCADOb
 	TUniqueTSObj<A3DRootBaseData> MetaData(Entity);
 	if (MetaData.IsValid())
 	{
-		if (MetaData->m_pcName && MetaData->m_pcName[0] != '\0')
+		if (OutObject.Label.IsEmpty() && MetaData->m_pcName && MetaData->m_pcName[0] != '\0')
 		{
 			FString Name = UTF8_TO_TCHAR(MetaData->m_pcName);
 			if(Name != TEXT("unnamed"))  // "unnamed" is create by Techsoft. This name is ignored 
@@ -1299,6 +1293,11 @@ void FTechSoftFileParser::BuildReferenceName(FArchiveCADObject& ReferenceData)
 		return;
 	}
 
+	if (ReferenceData.SetNameWithAttributeValue(TEXT("PartNumber")))
+	{
+		return;
+	}
+
 	switch (Format)
 	{
 	case ECADFormat::CATIA_3DXML:
@@ -1345,11 +1344,6 @@ void FTechSoftFileParser::BuildInstanceName(FArchiveInstance& InstanceData, cons
 		return;
 	}
 
-	if (InstanceData.SetNameWithAttributeValue(TEXT("PartNumber")))
-	{
-		return;
-	}
-
 	if (InstanceData.Label.IsEmpty())
 	{
 		InstanceData.Label = Parent.Label + "_" + FString::FromInt(Parent.Children.Num());
@@ -1790,23 +1784,7 @@ bool FTechSoftFileParser::IsConfigurationSet(const A3DAsmProductOccurrence* Occu
 		return false;
 	}
 
-	bool bIsConfiguration = false;
-	if (OccurrenceData->m_uiPOccurrencesSize)
-	{
-		TUniqueTSObj<A3DAsmProductOccurrenceData> ChildData;
-		for (uint32 Index = 0; Index < OccurrenceData->m_uiPOccurrencesSize; ++Index)
-		{
-			if (ChildData.FillFrom(OccurrenceData->m_ppPOccurrences[Index]) == A3D_SUCCESS)
-			{
-				if (ChildData->m_uiProductFlags & A3D_PRODUCT_FLAG_CONFIG)
-				{
-					bIsConfiguration = true;
-				}
-				break;
-			}
-		}
-	}
-	return bIsConfiguration;
+	return OccurrenceData->m_uiProductFlags & A3D_PRODUCT_FLAG_CONTAINER;
 }
 
 uint32 FTechSoftFileParser::CountColorAndMaterial()
diff --git a/Engine/Plugins/Enterprise/DatasmithCADImporter/Source/CADInterfaces/Private/TechSoftFileParser.h b/Engine/Plugins/Enterprise/DatasmithCADImporter/Source/CADInterfaces/Private/TechSoftFileParser.h
--- a/Engine/Plugins/Enterprise/DatasmithCADImporter/Source/CADInterfaces/Private/TechSoftFileParser.h
+++ b/Engine/Plugins/Enterprise/DatasmithCADImporter/Source/CADInterfaces/Private/TechSoftFileParser.h
@@ -74,8 +74,6 @@ private:
 	void CountUnderModel();
 	void CountUnderConfigurationSet(const A3DAsmProductOccurrence* Occurrence);
 	void CountUnderOccurrence(const A3DAsmProductOccurrence* Occurrence);
-	void CountUnderPrototype(const A3DAsmProductOccurrence* Prototype);
-	void CountUnderSubPrototype(const A3DAsmProductOccurrence* Prototype);
 	void CountUnderPartDefinition(const A3DAsmPartDefinition* PartDefinition);
 	void CountUnderRepresentationItem(const A3DRiRepresentationItem* RepresentationItem);
 	void CountUnderRepresentationSet(const A3DRiSet* RepresentationSet);
@@ -164,6 +162,7 @@ protected:
 	double FileUnit = 1;
 
 	TMap<A3DRiRepresentationItem*, FCadId> RepresentationItemsCache;
+	TMap<A3DAsmProductOccurrence*, FCadId> ReferenceCache;
 };
 
 } // ns CADLibrary
\ No newline at end of file
diff --git a/Engine/Plugins/Enterprise/DatasmithCADImporter/Source/CADTools/Public/CADData.h b/Engine/Plugins/Enterprise/DatasmithCADImporter/Source/CADTools/Public/CADData.h
--- a/Engine/Plugins/Enterprise/DatasmithCADImporter/Source/CADTools/Public/CADData.h
+++ b/Engine/Plugins/Enterprise/DatasmithCADImporter/Source/CADTools/Public/CADData.h
@@ -125,13 +125,10 @@ public:
 	 */
 	void DefineGraphicsPropertiesFromNoOverwrite(const FArchiveGraphicProperties& Source)
 	{
-		if (!ColorUId)
-		{
-			ColorUId = Source.ColorUId;
-		}
 
-		if (!MaterialUId)
+		if (!ColorUId && !MaterialUId)
 		{
+			ColorUId = Source.ColorUId;
 			MaterialUId = Source.MaterialUId;
 		}
 	}
diff --git a/Engine/Plugins/Enterprise/DatasmithCADImporter/Source/DatasmithCADTranslator/Private/DatasmithSceneGraphBuilder.cpp b/Engine/Plugins/Enterprise/DatasmithCADImporter/Source/DatasmithCADTranslator/Private/DatasmithSceneGraphBuilder.cpp
--- a/Engine/Plugins/Enterprise/DatasmithCADImporter/Source/DatasmithCADTranslator/Private/DatasmithSceneGraphBuilder.cpp
+++ b/Engine/Plugins/Enterprise/DatasmithCADImporter/Source/DatasmithCADTranslator/Private/DatasmithSceneGraphBuilder.cpp
@@ -361,20 +361,19 @@ TSharedPtr<IDatasmithActorElement>  FDatasmithSceneBaseGraphBuilder::BuildInstan
 	DatasmithSceneGraphBuilderImpl::GetNodeUuidAndLabel(Instance, *Reference, ParentData.Uuid, ActorUUID, ActorLabel);
 
 	TSharedPtr<IDatasmithActorElement> Actor = CreateActor(*ActorUUID, *ActorLabel);
-	if (!Actor.IsValid())
+	if (Actor.IsValid())
 	{
-		return TSharedPtr<IDatasmithActorElement>();
-	}
-	AddMetaData(Actor, Instance, *Reference);
+		AddMetaData(Actor, Instance, *Reference);
 
-	ActorData InstanceData(*ActorUUID, ParentData);
+		ActorData InstanceData(*ActorUUID, ParentData);
 
-	DatasmithSceneGraphBuilderImpl::SpreadGraphicProperties(Instance, InstanceData);
-	DatasmithSceneGraphBuilderImpl::SpreadGraphicProperties(*Reference, InstanceData);
+		DatasmithSceneGraphBuilderImpl::SpreadGraphicProperties(Instance, InstanceData);
+		DatasmithSceneGraphBuilderImpl::SpreadGraphicProperties(*Reference, InstanceData);
 
-	AddChildren(Actor, *Reference, InstanceData);
+		AddChildren(Actor, *Reference, InstanceData);
 
-	DatasmithSceneGraphBuilderImpl::AddTransformToActor(Instance, Actor, ImportParameters);
+		DatasmithSceneGraphBuilderImpl::AddTransformToActor(Instance, Actor, ImportParameters);
+	}
 
 	if (SceneGraph != InstanceSceneGraph)
 	{
