#swarm https://p4-swarm.epicgames.net/reviews/20995757

diff --git a/Engine/Plugins/Enterprise/DatasmithCADImporter/DatasmithCADImporter.uplugin b/Engine/Plugins/Enterprise/DatasmithCADImporter/DatasmithCADImporter.uplugin
--- a/Engine/Plugins/Enterprise/DatasmithCADImporter/DatasmithCADImporter.uplugin
+++ b/Engine/Plugins/Enterprise/DatasmithCADImporter/DatasmithCADImporter.uplugin
@@ -158,6 +158,10 @@
 		{
 			"Name": "DatasmithImporter",
 			"Enabled": true
+		},
+		{
+			"Name": "GeometryProcessing",
+			"Enabled": true
 		}
 	]
 }
\ No newline at end of file
diff --git a/Engine/Plugins/Enterprise/DatasmithCADImporter/Source/CADLibrary/CADLibrary.Build.cs b/Engine/Plugins/Enterprise/DatasmithCADImporter/Source/CADLibrary/CADLibrary.Build.cs
--- a/Engine/Plugins/Enterprise/DatasmithCADImporter/Source/CADLibrary/CADLibrary.Build.cs
+++ b/Engine/Plugins/Enterprise/DatasmithCADImporter/Source/CADLibrary/CADLibrary.Build.cs
@@ -17,7 +17,10 @@ public class CADLibrary : ModuleRules
 				"CADTools",
 				"CADKernel",
 				"DatasmithCore",
+				"DynamicMesh",
+				"MeshConversion",
 				"MeshDescription",
+				"GeometryCore",
 				"StaticMeshDescription"
 			}
 		);
diff --git a/Engine/Plugins/Enterprise/DatasmithCADImporter/Source/CADLibrary/Private/MeshDescriptionHelper.cpp b/Engine/Plugins/Enterprise/DatasmithCADImporter/Source/CADLibrary/Private/MeshDescriptionHelper.cpp
--- a/Engine/Plugins/Enterprise/DatasmithCADImporter/Source/CADLibrary/Private/MeshDescriptionHelper.cpp
+++ b/Engine/Plugins/Enterprise/DatasmithCADImporter/Source/CADLibrary/Private/MeshDescriptionHelper.cpp
@@ -385,6 +385,20 @@ void GetExistingPatches(FMeshDescription& MeshDestination, TSet<int32>& OutPatch
 	}
 }
 
+void CopyMaterialSlotNames(FMeshDescription& MeshSource, FMeshDescription& MeshDestination)
+{
+	FStaticMeshAttributes MeshAttributes(MeshSource);
+	TPolygonGroupAttributesRef<FName> PolygonGroupImportedMaterialSlotNames = MeshAttributes.GetPolygonGroupMaterialSlotNames();
+
+	FStaticMeshAttributes NewMeshDescriptionAttributes(MeshDestination);
+	TPolygonGroupAttributesRef<FName> NewPolygonGroupImportedMaterialSlotNames = NewMeshDescriptionAttributes.GetPolygonGroupMaterialSlotNames();
+
+	for (int32 Index = 0; Index < PolygonGroupImportedMaterialSlotNames.GetNumElements(); ++Index)
+	{
+		NewPolygonGroupImportedMaterialSlotNames[Index] = PolygonGroupImportedMaterialSlotNames[Index];
+	}
+}
+
 bool FillMesh(const FMeshParameters& MeshParameters, const FImportParameters& ImportParams, FBodyMesh& BodyTessellation, FMeshDescription& MeshDescription)
 {
 	const int32 UVChannel = 0;
@@ -614,6 +628,13 @@ bool ConvertBodyMeshToMeshDescription(const FImportParameters& ImportParams, con
 	// Orient mesh
 	MeshOperator::OrientMesh(MeshDescription);
 
+	// Sew mesh
+	if(FImportParameters::bGSewMeshIfNeeded)
+	{
+		double Tolerance = FImportParameters::GStitchingTolerance;
+		MeshOperator::ResolveTJunctions(MeshDescription, Tolerance);
+	}
+
 	// Build edge meta data
 	FStaticMeshOperations::DetermineEdgeHardnessesFromVertexInstanceNormals(MeshDescription);
 
diff --git a/Engine/Plugins/Enterprise/DatasmithCADImporter/Source/CADLibrary/Private/MeshOperator.cpp b/Engine/Plugins/Enterprise/DatasmithCADImporter/Source/CADLibrary/Private/MeshOperator.cpp
--- a/Engine/Plugins/Enterprise/DatasmithCADImporter/Source/CADLibrary/Private/MeshOperator.cpp
+++ b/Engine/Plugins/Enterprise/DatasmithCADImporter/Source/CADLibrary/Private/MeshOperator.cpp
@@ -6,15 +6,20 @@
 
 #include "Containers/Array.h"
 #include "Containers/Queue.h"
+#include "DynamicMesh/DynamicMesh3.h"
+#include "DynamicMesh/Operations/MergeCoincidentMeshEdges.h"
+#include "DynamicMeshToMeshDescription.h"
 #include "MeshAttributes.h"
+#include "MeshDescriptionHelper.h"
+#include "MeshDescriptionToDynamicMesh.h"
 #include "MeshElementArray.h"
+#include "Operations/MeshResolveTJunctions.h"
 #include "StaticMeshAttributes.h"
 
-using namespace MeshOperator;
-using namespace MeshCategory;
-
 bool MeshOperator::OrientMesh(FMeshDescription& MeshDescription)
 {
+	using namespace MeshCategory;
+
 	FMeshEditingWrapper MeshWrapper(MeshDescription);
 
 	TQueue<FTriangleID> Front;
@@ -243,4 +248,65 @@ bool MeshOperator::OrientMesh(FMeshDescription& MeshDescription)
 	}
 
 	return true;
+}
+
+void MeshOperator::ResolveTJunctions(FMeshDescription& MeshDescription, double Tolerance)
+{
+	UE::Geometry::FDynamicMesh3 DynamicMesh(UE::Geometry::EMeshComponents::FaceGroups);
+	DynamicMesh.EnableAttributes();
+
+	{
+		FMeshDescriptionToDynamicMesh ConverterToDynamicMesh;
+		ConverterToDynamicMesh.Convert(&MeshDescription, DynamicMesh);
+	}
+
+	// Check if there are boundary edges
+	int32 BoundaryEdgeCount = 0;
+	for (int32 eid : DynamicMesh.BoundaryEdgeIndicesItr())
+	{
+		BoundaryEdgeCount++;
+		break;
+	}
+
+	if (BoundaryEdgeCount == 0)
+	{
+		return;
+	}
+
+	UE::Geometry::FMeshResolveTJunctions MeshResolveTJunctions(&DynamicMesh);
+	MeshResolveTJunctions.DistanceTolerance = Tolerance;
+	bool bResolveOK = MeshResolveTJunctions.Apply();
+
+	if (bResolveOK && MeshResolveTJunctions.NumSplitEdges > 0)
+	{
+		UE::Geometry::FMergeCoincidentMeshEdges MergeCoincidentEdges(&DynamicMesh);
+		MergeCoincidentEdges.MergeVertexTolerance = Tolerance;
+		MergeCoincidentEdges.MergeSearchTolerance = 2 * Tolerance;
+		MergeCoincidentEdges.Apply();
+		BoundaryEdgeCount = MergeCoincidentEdges.FinalNumBoundaryEdges;
+	}
+
+	// TODO: Should not be needed as soon as CopyMaterialSlotNames, and CopyPatchGroups are no more called
+	FMeshDescription NewMeshDescription(MeshDescription);
+
+	{
+		FConversionToMeshDescriptionOptions ConversionOptions;
+		ConversionOptions.bSetPolyGroups = true;
+		ConversionOptions.bUpdatePositions = true;
+		ConversionOptions.bUpdateNormals = true;
+		ConversionOptions.bUpdateTangents = true;
+		ConversionOptions.bUpdateUVs = true;
+		ConversionOptions.bUpdateVtxColors = true;
+		ConversionOptions.bTransformVtxColorsSRGBToLinear = false;
+
+		FDynamicMeshToMeshDescription ConverterToMeshDescription(ConversionOptions);
+		ConverterToMeshDescription.Convert(&DynamicMesh, NewMeshDescription);
+	}
+
+
+	// TODO: Should not be needed as soon as ConverterToMeshDescription.Convert won't lose data anymore
+	CADLibrary::CopyMaterialSlotNames(MeshDescription, NewMeshDescription);
+	CADLibrary::CopyPatchGroups(MeshDescription, NewMeshDescription);
+
+	MeshDescription = MoveTemp(NewMeshDescription);
 }
\ No newline at end of file
diff --git a/Engine/Plugins/Enterprise/DatasmithCADImporter/Source/CADLibrary/Private/MeshOperator.h b/Engine/Plugins/Enterprise/DatasmithCADImporter/Source/CADLibrary/Private/MeshOperator.h
--- a/Engine/Plugins/Enterprise/DatasmithCADImporter/Source/CADLibrary/Private/MeshOperator.h
+++ b/Engine/Plugins/Enterprise/DatasmithCADImporter/Source/CADLibrary/Private/MeshOperator.h
@@ -11,4 +11,6 @@ namespace MeshOperator
 
 bool OrientMesh(FMeshDescription& MeshDescription);
 
+void ResolveTJunctions(FMeshDescription& MeshDescription, double Tolerance);
+
 }
\ No newline at end of file
diff --git a/Engine/Plugins/Enterprise/DatasmithCADImporter/Source/CADLibrary/Public/MeshDescriptionHelper.h b/Engine/Plugins/Enterprise/DatasmithCADImporter/Source/CADLibrary/Public/MeshDescriptionHelper.h
--- a/Engine/Plugins/Enterprise/DatasmithCADImporter/Source/CADLibrary/Public/MeshDescriptionHelper.h
+++ b/Engine/Plugins/Enterprise/DatasmithCADImporter/Source/CADLibrary/Public/MeshDescriptionHelper.h
@@ -27,6 +27,7 @@ CADLIBRARY_API TSharedPtr<IDatasmithUEPbrMaterialElement> CreateUEPbrMaterialFro
 
 CADLIBRARY_API bool ConvertBodyMeshToMeshDescription(const FImportParameters& ImportParams, const FMeshParameters& MeshParameters, FBodyMesh& Body, FMeshDescription& MeshDescription);
 CADLIBRARY_API void CopyPatchGroups(FMeshDescription& MeshSource, FMeshDescription& MeshDestination);
+CADLIBRARY_API void CopyMaterialSlotNames(FMeshDescription& MeshSource, FMeshDescription& MeshDestination);
 
 /** 
  * Enable per-triangle integer attribute named PolyTriGroups 
diff --git a/Engine/Plugins/Enterprise/DatasmithCADImporter/Source/CADTools/Private/CADOption.cpp b/Engine/Plugins/Enterprise/DatasmithCADImporter/Source/CADTools/Private/CADOption.cpp
--- a/Engine/Plugins/Enterprise/DatasmithCADImporter/Source/CADTools/Private/CADOption.cpp
+++ b/Engine/Plugins/Enterprise/DatasmithCADImporter/Source/CADTools/Private/CADOption.cpp
@@ -62,6 +62,15 @@ FAutoConsoleVariableRef GCADTranslatorStitchingTolerance(
 Default value of StitchingTolerance is 0.001 cm\n"),
 ECVF_Default);
 
+bool FImportParameters::bGSewMeshIfNeeded = true;
+FAutoConsoleVariableRef GCADTranslatorSewMeshIfNeeded(
+	TEXT("ds.CADTranslator.SewMeshIfNeeded"),
+	FImportParameters::bGSewMeshIfNeeded,
+	TEXT("Prefrom a welding of the mesh to try to stitch mesh cracks\n\
+This welding is performed respecting the ds.CADTranslator.StitchingTolerance\n\
+Default value is true\n"),
+ECVF_Default);
+
 static bool bGAliasSewByColor = false;
 FAutoConsoleVariableRef GAliasSewByColor(
 	TEXT("ds.CADTranslator.Alias.SewByColor"),
diff --git a/Engine/Plugins/Enterprise/DatasmithCADImporter/Source/CADTools/Public/CADOptions.h b/Engine/Plugins/Enterprise/DatasmithCADImporter/Source/CADTools/Public/CADOptions.h
--- a/Engine/Plugins/Enterprise/DatasmithCADImporter/Source/CADTools/Public/CADOptions.h
+++ b/Engine/Plugins/Enterprise/DatasmithCADImporter/Source/CADTools/Public/CADOptions.h
@@ -35,6 +35,7 @@ namespace CADLibrary
 		CADTOOLS_API static bool bGOverwriteCache;
 		CADTOOLS_API static bool bGPreferJtFileEmbeddedTessellation;
 		CADTOOLS_API static float GStitchingTolerance;
+		CADTOOLS_API static bool bGSewMeshIfNeeded;
 
 	public:
 		FImportParameters(FDatasmithUtils::EModelCoordSystem NewCoordinateSystem = FDatasmithUtils::EModelCoordSystem::ZUp_RightHanded)
