Fix management of multi-instance

diff --git a/Engine/Plugins/Enterprise/DatasmithCADImporter/Source/CADInterfaces/Private/TechSoftFileParser.cpp b/Engine/Plugins/Enterprise/DatasmithCADImporter/Source/CADInterfaces/Private/TechSoftFileParser.cpp
--- a/Engine/Plugins/Enterprise/DatasmithCADImporter/Source/CADInterfaces/Private/TechSoftFileParser.cpp
+++ b/Engine/Plugins/Enterprise/DatasmithCADImporter/Source/CADInterfaces/Private/TechSoftFileParser.cpp
@@ -829,18 +829,21 @@ void FTechSoftFileParser::TraverseOccurrence(const A3DAsmProductOccurrence* Occu
 	A3DAsmProductOccurrence* ReferencePtr = ProductOccurrence::GetReference(*OccurrenceData);
 
 	// Is the Reference already processed ?
-	FCadId* ReferenceId = ReferenceCache.Find(ReferencePtr);
-	if (ReferenceId != nullptr)
+	if(ReferencePtr)
 	{
-		Instance.ReferenceNodeId = *ReferenceId;
-
-		if (SceneGraph.IsAUnloadedReference(Instance.ReferenceNodeId))
+		FCadId* ReferenceId = ReferenceCache.Find(ReferencePtr);
+		if (ReferenceId != nullptr)
 		{
-			Instance.bIsExternalReference = true;
-		}
+			Instance.ReferenceNodeId = *ReferenceId;
 
-		ExtractTransformation(Transform, Instance);
-		return;
+			if (SceneGraph.IsAUnloadedReference(Instance.ReferenceNodeId))
+			{
+				Instance.bIsExternalReference = true;
+			}
+
+			ExtractTransformation(Transform, Instance);
+			return;
+		}
 	}
 
 	FArchiveUnloadedReference& UnloadedReference = SceneGraph.AddUnloadedReference(Instance);
@@ -870,7 +873,10 @@ void FTechSoftFileParser::TraverseOccurrence(const A3DAsmProductOccurrence* Occu
 		ProcessReference(OccurrencePtr, Instance, NewReference);
 	}
 
-	ReferenceCache.Add(ReferencePtr, Instance.ReferenceNodeId);
+	if(ReferencePtr)
+	{
+		ReferenceCache.Add(ReferencePtr, Instance.ReferenceNodeId);
+	}
 }
 
 void FTechSoftFileParser::ProcessReference(const A3DAsmProductOccurrence* OccurrencePtr, FArchiveInstance& Instance, FArchiveReference& Reference)
@@ -931,12 +937,15 @@ void FTechSoftFileParser::CountUnderOccurrence(const A3DAsmProductOccurrence* Oc
 		A3DAsmProductOccurrence* ReferencePtr = ProductOccurrence::GetReference(*OccurrenceData);
 
 		// Is the Reference already processed ?
-		FCadId* ReferenceId = ReferenceCache.Find(ReferencePtr);
-		if (ReferenceId != nullptr)
+		if(ReferencePtr)
 		{
-			return;
+			FCadId* ReferenceId = ReferenceCache.Find(ReferencePtr);
+			if (ReferenceId != nullptr)
+			{
+				return;
+			}
+			ReferenceCache.Add(ReferencePtr, 1);
 		}
-		ReferenceCache.Add(ReferencePtr, 1);
 		ComponentCount[EComponentType::Reference]++;
 
 		const A3DAsmProductOccurrence* CachedOccurrencePtr = Occurrence;
