Remove unused function

diff --git a/Engine/Plugins/Enterprise/DatasmithCADImporter/Source/CADLibrary/Private/MeshDescriptionHelper.cpp b/Engine/Plugins/Enterprise/DatasmithCADImporter/Source/CADLibrary/Private/MeshDescriptionHelper.cpp
--- a/Engine/Plugins/Enterprise/DatasmithCADImporter/Source/CADLibrary/Private/MeshDescriptionHelper.cpp
+++ b/Engine/Plugins/Enterprise/DatasmithCADImporter/Source/CADLibrary/Private/MeshDescriptionHelper.cpp
@@ -64,149 +64,6 @@ bool IsTriangleDegenerated(const int32_t* Indices, const TArray<FVertexID>& Rema
 	return (VertexIDs[0] == VertexIDs[1] || VertexIDs[0] == VertexIDs[2] || VertexIDs[1] == VertexIDs[2]);
 }
 
-void FillKioVertexPosition(const FImportParameters& ImportParams, const FMeshParameters& MeshParameters, FBodyMesh& Body, FMeshDescription& MeshDescription)
-{
-	int32 TriangleCount = Body.TriangleCount;
-	TArray<FTessellationData>& FaceTessellationSet = Body.Faces;
-
-	TVertexAttributesRef<FVector3f> VertexPositions = MeshDescription.GetVertexPositions();
-
-	// Create a list of vertex Z/index pairs
-	TArray<FVertexData> VertexDataSet;
-	VertexDataSet.Reserve(TriangleCount * 3);
-
-	FVector Position;
-	int32 GlobalVertexCount = 0;
-	for (FTessellationData& CTTessellation : FaceTessellationSet)
-	{
-		for (const FVector& Vertex : CTTessellation.PositionArray)
-		{
-			VertexDataSet.Emplace(GlobalVertexCount, Vertex);
-			++GlobalVertexCount;
-		}
-	}
-	VertexDataSet.SetNum(GlobalVertexCount);
-
-	// Sort the vertices by z value
-	VertexDataSet.Sort(FCompareVertexZ());
-
-	TArray<int32> NewIndexOf;
-	NewIndexOf.SetNumZeroed(GlobalVertexCount);
-
-	TArray<int32> IndexOfCoincidentNode;
-	IndexOfCoincidentNode.SetNumZeroed(GlobalVertexCount);
-
-	int32 VertexCount = 0;
-	// Search for duplicates, quickly!
-	for (int32 i = 0; i < VertexDataSet.Num(); i++)
-	{
-		NewIndexOf[VertexDataSet[i].Index] = i;
-		if (VertexDataSet[i].bIsMerged)
-		{
-			continue;
-		}
-
-		VertexDataSet[i].bIsMerged = true;
-		int32 Index_i = VertexDataSet[i].Index;
-		IndexOfCoincidentNode[Index_i] = Index_i;
-
-		const FVector& PositionA = VertexDataSet[i].Coordinates;
-
-		// only need to search forward, since we add pairs both ways
-		for (int32 j = i + 1; j < VertexDataSet.Num(); j++)
-		{
-			if (FMath::Abs(VertexDataSet[j].Z - VertexDataSet[i].Z) > KINDA_SMALL_NUMBER)
-			{
-				break; // can't be any more duplicated
-			}
-
-			const FVector& PositionB = VertexDataSet[j].Coordinates;
-			if (PositionA.Equals(PositionB, KINDA_SMALL_NUMBER))
-			{
-				VertexDataSet[j].bIsMerged = true;
-				IndexOfCoincidentNode[VertexDataSet[j].Index] = Index_i;
-			}
-		}
-		VertexCount++;
-	}
-
-	// if Symmetric mesh, the symmetric side of the mesh have to be generated
-	FMatrix SymmetricMatrix;
-	bool bIsSymmetricMesh = MeshParameters.bIsSymmetric;
-	if (bIsSymmetricMesh)
-	{
-		SymmetricMatrix = FDatasmithUtils::GetSymmetricMatrix(MeshParameters.SymmetricOrigin, MeshParameters.SymmetricNormal);
-	}
-
-	// Make MeshDescription.VertexPositions and VertexID
-	MeshDescription.ReserveNewVertices(bIsSymmetricMesh ? VertexCount * 2 : VertexCount);
-	Body.VertexIds.Reserve(VertexCount);
-
-	int32 GlobalVertexIndex = 0;
-
-	for (int32 VertexIndex = 0; VertexIndex < GlobalVertexCount; ++VertexIndex)
-	{
-		int32 RealIndex = VertexDataSet[VertexIndex].Index;
-
-		// Vertex is outside bbox
-		if (RealIndex < 0)
-		{
-			continue;
-		}
-
-		if (IndexOfCoincidentNode[RealIndex] != RealIndex)
-		{
-			continue;
-		}
-
-		FVertexID VertexID = MeshDescription.CreateVertex();
-		Body.VertexIds.Add(VertexID);
-		VertexPositions[VertexID] = (FVector3f)FDatasmithUtils::ConvertVector((FDatasmithUtils::EModelCoordSystem)ImportParams.GetModelCoordSys(), VertexDataSet[VertexIndex].Coordinates);
-		VertexDataSet[VertexIndex].VertexID = VertexID;
-	}
-
-	if (bIsSymmetricMesh)
-	{
-		Body.SymmetricVertexIds.Reserve(Body.VertexIds.Num());
-		for (int32 VertexIndex = 0; VertexIndex < GlobalVertexCount; ++VertexIndex)
-		{
-			int32 RealIndex = VertexDataSet[VertexIndex].Index;
-
-			// Vertex is outside bbox
-			if (RealIndex < 0)
-			{
-				continue;
-			}
-
-			if (IndexOfCoincidentNode[RealIndex] != RealIndex)
-			{
-				continue;
-			}
-
-			FVertexID VertexID = MeshDescription.CreateVertex();
-			Body.SymmetricVertexIds.Add(VertexID);
-			VertexPositions[VertexID] = (FVector3f)FDatasmithUtils::ConvertVector((FDatasmithUtils::EModelCoordSystem)ImportParams.GetModelCoordSys(), VertexDataSet[VertexIndex].Coordinates);
-			VertexPositions[VertexID] = FVector4f(SymmetricMatrix.TransformPosition((FVector)VertexPositions[VertexID]));
-
-			VertexDataSet[VertexIndex].SymVertexID = VertexID;
-		}
-	}
-
-	// For each face, for each vertex set VertexId
-	GlobalVertexIndex = 0;
-	for (FTessellationData& CTTessellation : FaceTessellationSet)
-	{
-		CTTessellation.PositionIndices.SetNum(CTTessellation.PositionArray.Num());
-		for (int32 VertexIndex = 0; VertexIndex < CTTessellation.PositionArray.Num(); ++VertexIndex, ++GlobalVertexIndex)
-		{
-			int32 NewIndex = NewIndexOf[IndexOfCoincidentNode[GlobalVertexIndex]];
-			FVertexData& Vertex = VertexDataSet[NewIndex];
-			int32 Index = Vertex.VertexID.GetValue();
-			CTTessellation.PositionIndices[VertexIndex] = VertexDataSet[NewIndex].VertexID.GetValue();
-		}
-	}
-}
-
 void MergeCoincidentVertices(TArray<FVector>& VertexArray, TArray<int32>& VertexIdSet)
 {
 	const double CoincidenceTolerance = 0.001;
@@ -610,15 +467,7 @@ bool ConvertBodyMeshToMeshDescription(const FImportParameters& ImportParams, con
 	MeshDescription.ReserveNewPolygons(Body.TriangleCount);
 	MeshDescription.ReserveNewEdges(Body.TriangleCount * 3);
 
-	// CoreTech is generating position duplicates. make sure to remove them before filling the mesh description
-	if (Body.VertexArray.Num() == 0)
-	{
-		FillKioVertexPosition(ImportParams, MeshParameters, Body, MeshDescription);
-	}
-	else
-	{
-		FillVertexPosition(ImportParams, MeshParameters, Body, MeshDescription);
-	}
+	FillVertexPosition(ImportParams, MeshParameters, Body, MeshDescription);
 
 	if (!FillMesh(MeshParameters, ImportParams, Body, MeshDescription))
 	{
