Add ds.CADTranslator.UnitScale CVar to scale staticMesh ( and translation of each transform matrices)

diff --git a/Engine/Plugins/Enterprise/DatasmithCADImporter/Source/CADInterfaces/Private/TechSoftFileParser.cpp b/Engine/Plugins/Enterprise/DatasmithCADImporter/Source/CADInterfaces/Private/TechSoftFileParser.cpp
--- a/Engine/Plugins/Enterprise/DatasmithCADImporter/Source/CADInterfaces/Private/TechSoftFileParser.cpp
+++ b/Engine/Plugins/Enterprise/DatasmithCADImporter/Source/CADInterfaces/Private/TechSoftFileParser.cpp
@@ -1679,7 +1679,7 @@ FMatrix FTechSoftFileParser::ExtractTransformation3D(const A3DMiscTransformation
 
 		FVector ZVector = XVector ^ YVector;
 
-		Origin *= InOutUnit;
+		Origin *= InOutUnit * FImportParameters::GUnitScale;
 
 		const A3DVector3dData& A3DScale = CartesianTransformationData->m_sScale;
 		FVector3d Scale(A3DScale.m_dX, A3DScale.m_dY, A3DScale.m_dZ);
@@ -1726,14 +1726,16 @@ FMatrix FTechSoftFileParser::ExtractGeneralTransformation(const A3DMiscTransform
 		FVector3d Scale = Transform.GetScale3D();
 		if (Scale.Equals(FVector3d::OneVector, KINDA_SMALL_NUMBER))
 		{
+			const double TranslationScale = InOutUnit * FImportParameters::GUnitScale;
 			for (Index = 0; Index < 3; ++Index, ++Index)
 			{
-				Matrix.M[3][Index] *= InOutUnit;
+				Matrix.M[3][Index] *= TranslationScale;
 			}
 			return Matrix;
 		}
 
 		FVector3d Translation = Transform.GetTranslation();
+		Translation *= FImportParameters::GUnitScale;
 
 		double UniformScale = TechSoftFileParserImpl::ExtractUniformScale(Scale);
 		InOutUnit *= UniformScale;
diff --git a/Engine/Plugins/Enterprise/DatasmithCADImporter/Source/CADInterfaces/Private/TechSoftInterfaceUtils.inl b/Engine/Plugins/Enterprise/DatasmithCADImporter/Source/CADInterfaces/Private/TechSoftInterfaceUtils.inl
--- a/Engine/Plugins/Enterprise/DatasmithCADImporter/Source/CADInterfaces/Private/TechSoftInterfaceUtils.inl
+++ b/Engine/Plugins/Enterprise/DatasmithCADImporter/Source/CADInterfaces/Private/TechSoftInterfaceUtils.inl
@@ -82,10 +82,11 @@ private:
 		int32 VertexCount = TessellationBaseData->m_uiCoordSize / 3;
 		VertexArray.Reserve(VertexArray.Num() + VertexCount);
 
+		const double ScaleFactor = BodyUnit * FImportParameters::GUnitScale;
 		double* Coordinates = TessellationBaseData->m_pdCoords;
 		for (unsigned int Index = 0; Index < TessellationBaseData->m_uiCoordSize; ++Index)
 		{
-			Coordinates[Index] *= BodyUnit;
+			Coordinates[Index] *= ScaleFactor;
 		}
 
 		for (unsigned int Index = 0; Index < TessellationBaseData->m_uiCoordSize; Index += 3)
diff --git a/Engine/Plugins/Enterprise/DatasmithCADImporter/Source/CADTools/Private/CADOption.cpp b/Engine/Plugins/Enterprise/DatasmithCADImporter/Source/CADTools/Private/CADOption.cpp
--- a/Engine/Plugins/Enterprise/DatasmithCADImporter/Source/CADTools/Private/CADOption.cpp
+++ b/Engine/Plugins/Enterprise/DatasmithCADImporter/Source/CADTools/Private/CADOption.cpp
@@ -62,6 +62,14 @@ FAutoConsoleVariableRef GCADTranslatorStitchingTolerance(
 Default value of StitchingTolerance is 0.001 cm\n"),
 ECVF_Default);
 
+float FImportParameters::GUnitScale = 1.f;
+FAutoConsoleVariableRef GCADTranslatorUnitScale(
+	TEXT("ds.CADTranslator.UnitScale"),
+	FImportParameters::GUnitScale,
+	TEXT("Scale factor to change the unit of the DMU (Only applies to TechSoft import.)\n\
+Default value of UnitScale is 1 i.e. unit = cm\n"),
+ECVF_Default);
+
 bool FImportParameters::bGSewMeshIfNeeded = true;
 FAutoConsoleVariableRef GCADTranslatorSewMeshIfNeeded(
 	TEXT("ds.CADTranslator.SewMeshIfNeeded"),
diff --git a/Engine/Plugins/Enterprise/DatasmithCADImporter/Source/CADTools/Public/CADOptions.h b/Engine/Plugins/Enterprise/DatasmithCADImporter/Source/CADTools/Public/CADOptions.h
--- a/Engine/Plugins/Enterprise/DatasmithCADImporter/Source/CADTools/Public/CADOptions.h
+++ b/Engine/Plugins/Enterprise/DatasmithCADImporter/Source/CADTools/Public/CADOptions.h
@@ -30,12 +30,13 @@ namespace CADLibrary
 		
 	public:
 		CADTOOLS_API static bool bGDisableCADKernelTessellation;
-		CADTOOLS_API static bool bGEnableTimeControl;
 		CADTOOLS_API static bool bGEnableCADCache;
+		CADTOOLS_API static bool bGEnableTimeControl;
 		CADTOOLS_API static bool bGOverwriteCache;
 		CADTOOLS_API static bool bGPreferJtFileEmbeddedTessellation;
-		CADTOOLS_API static float GStitchingTolerance;
 		CADTOOLS_API static bool bGSewMeshIfNeeded;
+		CADTOOLS_API static float GStitchingTolerance;
+		CADTOOLS_API static float GUnitScale;
 
 	public:
 		FImportParameters(FDatasmithUtils::EModelCoordSystem NewCoordinateSystem = FDatasmithUtils::EModelCoordSystem::ZUp_RightHanded)
