Add SewMeshIfNeeded and StitchingTolerance CVar value in CAD geometric cache file hash to allow to remesh with a new value of one of these values.

diff --git a/Engine/Plugins/Enterprise/DatasmithCADImporter/Source/CADInterfaces/Public/CADFileData.h b/Engine/Plugins/Enterprise/DatasmithCADImporter/Source/CADInterfaces/Public/CADFileData.h
--- a/Engine/Plugins/Enterprise/DatasmithCADImporter/Source/CADInterfaces/Public/CADFileData.h
+++ b/Engine/Plugins/Enterprise/DatasmithCADImporter/Source/CADInterfaces/Public/CADFileData.h
@@ -38,6 +38,8 @@ public:
 		{
 			GeomFileHash = GetSceneFileHash();
 			GeomFileHash = HashCombine(GeomFileHash, GetTypeHash(ImportParameters));
+			GeomFileHash = HashCombine(GeomFileHash, ::GetTypeHash(FImportParameters::bGSewMeshIfNeeded));
+			GeomFileHash = HashCombine(GeomFileHash, ::GetTypeHash(FImportParameters::GStitchingTolerance));
 		}
 		return GeomFileHash;
 	}
