Rename ECADFormat::N_X

diff --git a/Engine/Plugins/Enterprise/DatasmithCADImporter/Source/CADInterfaces/Private/TechSoftFileParser.cpp b/Engine/Plugins/Enterprise/DatasmithCADImporter/Source/CADInterfaces/Private/TechSoftFileParser.cpp
--- a/Engine/Plugins/Enterprise/DatasmithCADImporter/Source/CADInterfaces/Private/TechSoftFileParser.cpp
+++ b/Engine/Plugins/Enterprise/DatasmithCADImporter/Source/CADInterfaces/Private/TechSoftFileParser.cpp
@@ -780,7 +780,7 @@ void FTechSoftFileParser::ProcessUnloadedReference(const FArchiveInstance& Insta
 		}
 		break;
 
-	case ECADFormat::NX:
+	case ECADFormat::N_X:
 		// todo
 		break;
 	default:
diff --git a/Engine/Plugins/Enterprise/DatasmithCADImporter/Source/CADTools/Private/CADData.cpp b/Engine/Plugins/Enterprise/DatasmithCADImporter/Source/CADTools/Private/CADData.cpp
--- a/Engine/Plugins/Enterprise/DatasmithCADImporter/Source/CADTools/Private/CADData.cpp
+++ b/Engine/Plugins/Enterprise/DatasmithCADImporter/Source/CADTools/Private/CADData.cpp
@@ -240,7 +240,7 @@ ECADFormat FileFormat(const FString& Extension)
 	}
 	else if (Extension == TEXT("prt") || Extension == TEXT("asm"))
 	{
-		return ECADFormat::NX;
+		return ECADFormat::N_X;
 	}
 	else if (Extension == TEXT("sat") || Extension == TEXT("sab"))
 	{
diff --git a/Engine/Plugins/Enterprise/DatasmithCADImporter/Source/CADTools/Public/CADData.h b/Engine/Plugins/Enterprise/DatasmithCADImporter/Source/CADTools/Public/CADData.h
--- a/Engine/Plugins/Enterprise/DatasmithCADImporter/Source/CADTools/Public/CADData.h
+++ b/Engine/Plugins/Enterprise/DatasmithCADImporter/Source/CADTools/Public/CADData.h
@@ -32,7 +32,7 @@ enum class ECADFormat
 	IGES,
 	INVENTOR,
 	JT,
-	NX,
+	N_X,  
 	MICROSTATION,
 	PARASOLID,
 	SOLID_EDGE,
diff --git a/Engine/Plugins/Enterprise/DatasmithCADImporter/Source/DatasmithCADTranslator/Private/DatasmithCADTranslator.cpp b/Engine/Plugins/Enterprise/DatasmithCADImporter/Source/DatasmithCADTranslator/Private/DatasmithCADTranslator.cpp
--- a/Engine/Plugins/Enterprise/DatasmithCADImporter/Source/DatasmithCADTranslator/Private/DatasmithCADTranslator.cpp
+++ b/Engine/Plugins/Enterprise/DatasmithCADImporter/Source/DatasmithCADTranslator/Private/DatasmithCADTranslator.cpp
@@ -142,7 +142,7 @@ bool FDatasmithCADTranslator::LoadScene(TSharedRef<IDatasmithScene> DatasmithSce
 
 	switch (FileDescriptor.GetFileFormat())
 	{
-	case CADLibrary::ECADFormat::NX:
+	case CADLibrary::ECADFormat::N_X:
 	{
 		break;
 	}
diff --git a/Engine/Plugins/Enterprise/DatasmithCADImporter/Source/DatasmithCADTranslator/Private/DatasmithSceneGraphBuilder.cpp b/Engine/Plugins/Enterprise/DatasmithCADImporter/Source/DatasmithCADTranslator/Private/DatasmithSceneGraphBuilder.cpp
--- a/Engine/Plugins/Enterprise/DatasmithCADImporter/Source/DatasmithCADTranslator/Private/DatasmithSceneGraphBuilder.cpp
+++ b/Engine/Plugins/Enterprise/DatasmithCADImporter/Source/DatasmithCADTranslator/Private/DatasmithSceneGraphBuilder.cpp
@@ -285,8 +285,8 @@ bool FDatasmithSceneBaseGraphBuilder::Build()
 		case CADLibrary::ECADFormat::INVENTOR:
 			DatasmithScene->SetProductName(TEXT("Inventor"));
 			break;
-		case CADLibrary::ECADFormat::NX:
-			DatasmithScene->SetProductName(TEXT("NX"));
+		case CADLibrary::ECADFormat::N_X:
+			DatasmithScene->SetProductName(TEXT("N")TEXT("X"));
 			break;
 		case CADLibrary::ECADFormat::PARASOLID:
 			DatasmithScene->SetProductName(TEXT("Parasolid"));
