Fix CAD N X Format import

diff --git a/Engine/Plugins/Enterprise/DatasmithCADImporter/Source/CADInterfaces/Private/TechSoftFileParser.cpp b/Engine/Plugins/Enterprise/DatasmithCADImporter/Source/CADInterfaces/Private/TechSoftFileParser.cpp
--- a/Engine/Plugins/Enterprise/DatasmithCADImporter/Source/CADInterfaces/Private/TechSoftFileParser.cpp
+++ b/Engine/Plugins/Enterprise/DatasmithCADImporter/Source/CADInterfaces/Private/TechSoftFileParser.cpp
@@ -254,6 +254,14 @@ void UpdateIOOptionAccordingToFormat(const CADLibrary::ECADFormat Format, A3DImp
 		}
 	}
 
+	case CADLibrary::ECADFormat::N_X:
+	{
+		Importer.m_sLoadData.m_sGeneral.m_bReadActiveFilter = A3D_TRUE;
+		// jira UE-159972
+		Importer.m_sLoadData.m_sIncremental.m_bLoadNoDependencies = false;
+		break;
+	}
+
 	case CADLibrary::ECADFormat::INVENTOR:
 	case CADLibrary::ECADFormat::CATIA_3DXML:
 	{
@@ -780,9 +788,6 @@ void FTechSoftFileParser::ProcessUnloadedReference(const FArchiveInstance& Insta
 		}
 		break;
 
-	case ECADFormat::N_X:
-		// todo
-		break;
 	default:
 		break;
 	}
@@ -1201,7 +1206,6 @@ void FTechSoftFileParser::TraverseBRepModel(A3DRiBrepModel* BRepModelPtr, FArchi
 
 	ExtractSpecificMetaData(BRepModelPtr, BRep);
 
-
 	if (!BRep.bShow || BRep.bIsRemoved)
 	{
 		SceneGraph.RemoveLastBody();
@@ -1495,16 +1499,18 @@ void FTechSoftFileParser::ExtractSpecificMetaData(const A3DAsmProductOccurrence*
 
 			if(UnigraphicsSpecificData->m_uiSolidsByRefsetsSize)
 			{
-				OutMetaData.BodyIdentsOfConfigurations.SetNum(UnigraphicsSpecificData->m_uiSolidsByRefsetsSize);
 				for (uint32 Index = 0; Index < UnigraphicsSpecificData->m_uiSolidsByRefsetsSize; ++Index)
 				{
 					A3DElementsByRefsetUg Refset = UnigraphicsSpecificData->m_asSolidsByRefsets[Index];
-					TPair<FString, TArray<uint32>>& BodyIdentsOfConfiguration = OutMetaData.BodyIdentsOfConfigurations[Index];
-					BodyIdentsOfConfiguration.Key = UTF8_TO_TCHAR(Refset.m_psRefset);
-					BodyIdentsOfConfiguration.Value.Reserve(Refset.m_uiElementsSize);
-					for (uint32 Andex = 0; Andex < Refset.m_uiElementsSize; ++Andex)
+
+					FString ReferenceSetName = UTF8_TO_TCHAR(Refset.m_psRefset);
+					if (ReferenceSetName == CADFileData.GetCADFileDescription().GetConfiguration())
 					{
-						BodyIdentsOfConfiguration.Value.Add(Refset.m_auiElements[Andex]);
+						UnigraphicsReferenceSet.Reserve(Refset.m_uiElementsSize);
+						for (uint32 Andex = 0; Andex < Refset.m_uiElementsSize; ++Andex)
+						{
+							UnigraphicsReferenceSet.Add(Refset.m_auiElements[Andex]);
+						}
 					}
 				}
 			}
