The problem is due to MeshOperator::ResolveTJunctions that is currently not compilant in parallel processing (due to FMeshDescriptionToDynamicMesh among other things).
The temporary solution is to disable parallel loading of StaticMesh if MeshOperator::ResolveTJunctions is used

diff --git a/Engine/Plugins/Enterprise/DatasmithCADImporter/Source/DatasmithCADTranslator/Private/DatasmithCADTranslator.cpp b/Engine/Plugins/Enterprise/DatasmithCADImporter/Source/DatasmithCADTranslator/Private/DatasmithCADTranslator.cpp
--- a/Engine/Plugins/Enterprise/DatasmithCADImporter/Source/DatasmithCADTranslator/Private/DatasmithCADTranslator.cpp
+++ b/Engine/Plugins/Enterprise/DatasmithCADImporter/Source/DatasmithCADTranslator/Private/DatasmithCADTranslator.cpp
@@ -5,6 +5,7 @@
 #include "CADFileReader.h"
 #include "CADInterfacesModule.h"
 #include "CADKernelSurfaceExtension.h"
+#include "CADOptions.h"
 
 #include "DatasmithCADTranslatorModule.h"
 #include "DatasmithDispatcher.h"
@@ -111,6 +112,9 @@ void FDatasmithCADTranslator::Initialize(FDatasmithTranslatorCapabilities& OutCa
 
 	OutCapabilities.SupportedFileFormats.Add(FFileFormatInfo{ TEXT("hsf"), TEXT("HOOPS stream files") });
 	OutCapabilities.SupportedFileFormats.Add(FFileFormatInfo{ TEXT("prc"), TEXT("HOOPS stream files") });
+
+	// FDatasmithCADTranslator::LoadStaticMesh is not compilant with parallel loading if MeshOperator::ResolveTJunctions is used
+	OutCapabilities.bParallelLoadStaticMeshSupported = !CADLibrary::FImportParameters::bGSewMeshIfNeeded;
 }
 
 bool FDatasmithCADTranslator::IsSourceSupported(const FDatasmithSceneSource& Source)
