[NO NON-PII COMMIT MESSAGE]

diff --git a/Engine/Plugins/Enterprise/DatasmithCADImporter/Source/CADInterfaces/Private/CADSceneGraph.cpp b/Engine/Plugins/Enterprise/DatasmithCADImporter/Source/CADInterfaces/Private/CADSceneGraph.cpp
--- a/Engine/Plugins/Enterprise/DatasmithCADImporter/Source/CADInterfaces/Private/CADSceneGraph.cpp
+++ b/Engine/Plugins/Enterprise/DatasmithCADImporter/Source/CADInterfaces/Private/CADSceneGraph.cpp
@@ -257,6 +257,11 @@ void FArchiveReference::AddChild(const FCadId ChildId)
 	Children.Add(ChildId);
 }
 
+void FArchiveReference::RemoveLastChild()
+{
+	Children.Pop();
+}
+
 int32 FArchiveReference::ChildrenCount()
 {
 	return Children.Num();
diff --git a/Engine/Plugins/Enterprise/DatasmithCADImporter/Source/CADInterfaces/Private/TechSoftFileParser.cpp b/Engine/Plugins/Enterprise/DatasmithCADImporter/Source/CADInterfaces/Private/TechSoftFileParser.cpp
--- a/Engine/Plugins/Enterprise/DatasmithCADImporter/Source/CADInterfaces/Private/TechSoftFileParser.cpp
+++ b/Engine/Plugins/Enterprise/DatasmithCADImporter/Source/CADInterfaces/Private/TechSoftFileParser.cpp
@@ -1173,6 +1173,7 @@ void FTechSoftFileParser::TraverseRepresentationSet(const A3DRiSet* Representati
 
 	if (RepresentationSet.bIsRemoved || !RepresentationSet.bShow)
 	{
+		Parent.RemoveLastChild();
 		SceneGraph.RemoveLastOccurence();
 		return;
 	}
diff --git a/Engine/Plugins/Enterprise/DatasmithCADImporter/Source/CADInterfaces/Public/CADSceneGraph.h b/Engine/Plugins/Enterprise/DatasmithCADImporter/Source/CADInterfaces/Public/CADSceneGraph.h
--- a/Engine/Plugins/Enterprise/DatasmithCADImporter/Source/CADInterfaces/Public/CADSceneGraph.h
+++ b/Engine/Plugins/Enterprise/DatasmithCADImporter/Source/CADInterfaces/Public/CADSceneGraph.h
@@ -78,6 +78,7 @@ public:
 	friend FArchive& operator<<(FArchive& Ar, FArchiveReference& C);
 
 	void AddChild(const FCadId ChildId);
+	void RemoveLastChild();
 
 	int32 ChildrenCount();
 	void MoveTemp(FArchiveUnloadedReference& Reference);
diff --git a/Engine/Plugins/Enterprise/DatasmithCADImporter/Source/DatasmithCADTranslator/Private/DatasmithSceneGraphBuilder.cpp b/Engine/Plugins/Enterprise/DatasmithCADImporter/Source/DatasmithCADTranslator/Private/DatasmithSceneGraphBuilder.cpp
--- a/Engine/Plugins/Enterprise/DatasmithCADImporter/Source/DatasmithCADTranslator/Private/DatasmithSceneGraphBuilder.cpp
+++ b/Engine/Plugins/Enterprise/DatasmithCADImporter/Source/DatasmithCADTranslator/Private/DatasmithSceneGraphBuilder.cpp
@@ -652,7 +652,7 @@ void FDatasmithSceneBaseGraphBuilder::AddChildren(TSharedPtr<IDatasmithActorElem
 				Actor->AddChild(ChildActor);
 			}
 		}
-		else
+		else if (SceneGraph->IsABody(ChildId))
 		{
 			TSharedPtr<IDatasmithActorElement> ChildActor = BuildBody(ChildId, ParentData);
 			if (ChildActor.IsValid() && DoesActorHaveChildrenOrIsAStaticMesh(ChildActor))
