[NO NON-PII COMMIT MESSAGE]

diff --git a/Engine/Plugins/Enterprise/DatasmithCADImporter/Source/CADInterfaces/Private/TechSoftFileParser.cpp b/Engine/Plugins/Enterprise/DatasmithCADImporter/Source/CADInterfaces/Private/TechSoftFileParser.cpp
--- a/Engine/Plugins/Enterprise/DatasmithCADImporter/Source/CADInterfaces/Private/TechSoftFileParser.cpp
+++ b/Engine/Plugins/Enterprise/DatasmithCADImporter/Source/CADInterfaces/Private/TechSoftFileParser.cpp
@@ -503,6 +503,15 @@ void FTechSoftFileParser::GenerateBodyMesh(A3DRiRepresentationItem* Representati
 	Body.ColorFaceSet = BodyMesh.ColorSet;
 	Body.MaterialFaceSet = BodyMesh.MaterialSet;
 
+	if (Body.ColorUId == 0 && Body.ColorFaceSet.Num())
+	{
+		Body.ColorUId = *Body.ColorFaceSet.begin();
+	}
+	if (Body.MaterialUId == 0 && Body.MaterialFaceSet.Num())
+	{
+		Body.MaterialUId = *Body.MaterialFaceSet.begin();
+	}
+
 	// Write part's representation as Prc file if it is a BRep
 	A3DEEntityType Type;
 	A3DEntityGetType(Representation, &Type);
