Use FVector3f for mesh vertex coordinates, normal, and UV as staticMesh works with float.

diff --git a/Engine/Plugins/Enterprise/DatasmithCADImporter/Source/CADInterfaces/Private/TechSoftInterfaceUtils.inl b/Engine/Plugins/Enterprise/DatasmithCADImporter/Source/CADInterfaces/Private/TechSoftInterfaceUtils.inl
--- a/Engine/Plugins/Enterprise/DatasmithCADImporter/Source/CADInterfaces/Private/TechSoftInterfaceUtils.inl
+++ b/Engine/Plugins/Enterprise/DatasmithCADImporter/Source/CADInterfaces/Private/TechSoftInterfaceUtils.inl
@@ -68,7 +68,7 @@ public:
 	}
 
 private:
-	void FillVertexArray(TArray<FVector>& VertexArray)
+	void FillVertexArray(TArray<FVector3f>& VertexArray)
 	{
 		using namespace CADLibrary;
 
@@ -355,7 +355,7 @@ private:
 		return true;
 	};
 
-	void AddNormals(const A3DDouble* Normals, const int32 Indices[3], TArray<FVector>& NormalsArray)
+	void AddNormals(const A3DDouble* Normals, const int32 Indices[3], TArray<FVector3f>& NormalsArray)
 	{
 		for (int32 Index = 0; Index < 3; ++Index)
 		{
@@ -364,7 +364,7 @@ private:
 		}
 	};
 
-	void AddTextureCoordinates(const A3DDouble* TextureCoords, const int32 Indices[3], TArray<FVector2D>& TessellationTextures)
+	void AddTextureCoordinates(const A3DDouble* TextureCoords, const int32 Indices[3], TArray<FVector2f>& TessellationTextures)
 	{
 		for (int32 Index = 0; Index < 3; ++Index)
 		{
@@ -791,7 +791,7 @@ private:
 		}
 	}
 
-	void ScaleUV(const A3DTopoFace* TopoFace, TArray<FVector2D>& TexCoordArray)
+	void ScaleUV(const A3DTopoFace* TopoFace, TArray<FVector2f>& TexCoordArray)
 	{
 		CADLibrary::TUniqueTSObj<A3DTopoFaceData> TopoFaceData(TopoFace);
 		if (!TopoFaceData.IsValid())
@@ -875,10 +875,10 @@ private:
 
 		// Texture unit is meter, Coord unit from TechSoft is mm, so TextureScale = 0.001 to convert mm into m
 		const double TextureScale = 0.01;
-		const double UScale = TextureUnit * TextureScale * LengthUMed / (Domain->m_sMax.m_dX - Domain->m_sMin.m_dX);
-		const double VScale = TextureUnit * TextureScale * LengthVMed / (Domain->m_sMax.m_dY - Domain->m_sMin.m_dY);
+		const float UScale = TextureUnit * TextureScale * LengthUMed / (Domain->m_sMax.m_dX - Domain->m_sMin.m_dX);
+		const float VScale = TextureUnit * TextureScale * LengthVMed / (Domain->m_sMax.m_dY - Domain->m_sMin.m_dY);
 
-		for (FVector2D& TexCoord : TexCoordArray)
+		for (FVector2f& TexCoord : TexCoordArray)
 		{
 			TexCoord[0] *= UScale;
 			TexCoord[1] *= VScale;
diff --git a/Engine/Plugins/Enterprise/DatasmithCADImporter/Source/CADInterfaces/Private/TechSoftParserUtils.h b/Engine/Plugins/Enterprise/DatasmithCADImporter/Source/CADInterfaces/Private/TechSoftParserUtils.h
--- a/Engine/Plugins/Enterprise/DatasmithCADImporter/Source/CADInterfaces/Private/TechSoftParserUtils.h
+++ b/Engine/Plugins/Enterprise/DatasmithCADImporter/Source/CADInterfaces/Private/TechSoftParserUtils.h
@@ -51,7 +51,7 @@ namespace TechSoftParserUtils
 		return true;
 	};
 
-	inline void AddNormals(const A3DDouble* Normals, const int32 Indices[3], TArray<FVector>& TessellationNormals)
+	inline void AddNormals(const A3DDouble* Normals, const int32 Indices[3], TArray<FVector3f>& TessellationNormals)
 	{
 		for (int32 Index = 0; Index < 3; ++Index)
 		{
@@ -60,7 +60,7 @@ namespace TechSoftParserUtils
 		}
 	};
 
-	inline void AddTextureCoordinates(const A3DDouble* TextureCoords, const int32 Indices[3], TArray<FVector2D>& TessellationTextures)
+	inline void AddTextureCoordinates(const A3DDouble* TextureCoords, const int32 Indices[3], TArray<FVector2f>& TessellationTextures)
 	{
 		for (int32 Index = 0; Index < 3; ++Index)
 		{
diff --git a/Engine/Plugins/Enterprise/DatasmithCADImporter/Source/CADKernelSurface/Private/CADKernelSurfaceExtension.cpp b/Engine/Plugins/Enterprise/DatasmithCADImporter/Source/CADKernelSurface/Private/CADKernelSurfaceExtension.cpp
--- a/Engine/Plugins/Enterprise/DatasmithCADImporter/Source/CADKernelSurface/Private/CADKernelSurfaceExtension.cpp
+++ b/Engine/Plugins/Enterprise/DatasmithCADImporter/Source/CADKernelSurface/Private/CADKernelSurfaceExtension.cpp
@@ -41,8 +41,8 @@ bool UCADKernelParametricSurfaceData::Tessellate(UStaticMesh& StaticMesh, const
 	CADLibrary::FMeshParameters CadMeshParameters;
 	CadMeshParameters.bNeedSwapOrientation = MeshParameters.bNeedSwapOrientation;
 	CadMeshParameters.bIsSymmetric = MeshParameters.bIsSymmetric;
-	CadMeshParameters.SymmetricNormal = MeshParameters.SymmetricNormal;
-	CadMeshParameters.SymmetricOrigin = MeshParameters.SymmetricOrigin;
+	CadMeshParameters.SymmetricNormal = (FVector3f) MeshParameters.SymmetricNormal;
+	CadMeshParameters.SymmetricOrigin = (FVector3f) MeshParameters.SymmetricOrigin;
 
 	// Previous MeshDescription is get to be able to create a new one with the same order of PolygonGroup (the matching of color and partition is currently based on their order)
 	if (FMeshDescription* DestinationMeshDescription = StaticMesh.GetMeshDescription(0))
diff --git a/Engine/Plugins/Enterprise/DatasmithCADImporter/Source/CADLibrary/Private/CADKernelTools.cpp b/Engine/Plugins/Enterprise/DatasmithCADImporter/Source/CADLibrary/Private/CADKernelTools.cpp
--- a/Engine/Plugins/Enterprise/DatasmithCADImporter/Source/CADLibrary/Private/CADKernelTools.cpp
+++ b/Engine/Plugins/Enterprise/DatasmithCADImporter/Source/CADLibrary/Private/CADKernelTools.cpp
@@ -48,10 +48,10 @@ namespace CADLibrary
 
 	static void FillVertexPosition(FMeshConversionContext& Context, const TSharedRef<CADKernel::FModelMesh>& ModelMesh, FMeshDescription& MeshDescription)
 	{
-		TArray<FVector> VertexArray;
+		TArray<FVector3f> VertexArray;
 		ModelMesh->GetNodeCoordinates(VertexArray);
 
-		for (FVector& Vertex : VertexArray)
+		for (FVector3f& Vertex : VertexArray)
 		{
 			Vertex *= 0.1; // mm (CADKernel unit) to cm (UE unit)
 		}
@@ -65,28 +65,28 @@ namespace CADLibrary
 
 		// Make MeshDescription.VertexPositions and VertexID
 		int32 VertexIndex = -1;
-		for (const FVector& Vertex : VertexArray)
+		for (const FVector3f& Vertex : VertexArray)
 		{
 			VertexIndex++;
 
 			FVertexID VertexID = MeshDescription.CreateVertex();
-			VertexPositions[VertexID] = (FVector3f)FDatasmithUtils::ConvertVector((FDatasmithUtils::EModelCoordSystem) Context.ImportParams.GetModelCoordSys(), Vertex);
+			VertexPositions[VertexID] = FDatasmithUtils::ConvertVector((FDatasmithUtils::EModelCoordSystem) Context.ImportParams.GetModelCoordSys(), Vertex);
 			Context.VertexIds[VertexIndex] = VertexID;
 		}
 
 		// if Symmetric mesh, the symmetric side of the mesh have to be generated
 		if (Context.MeshParameters.bIsSymmetric)
 		{
-			FMatrix SymmetricMatrix = FDatasmithUtils::GetSymmetricMatrix(Context.MeshParameters.SymmetricOrigin, Context.MeshParameters.SymmetricNormal);
+			FMatrix44f SymmetricMatrix = FDatasmithUtils::GetSymmetricMatrix(Context.MeshParameters.SymmetricOrigin, Context.MeshParameters.SymmetricNormal);
 
 			Context.SymmetricVertexIds.SetNum(VertexArray.Num());
 
 			VertexIndex = 0;
-			for (const FVector& Vertex : VertexArray)
+			for (const FVector3f& Vertex : VertexArray)
 			{
 				FVertexID VertexID = MeshDescription.CreateVertex();
-				VertexPositions[VertexID] = (FVector3f)FDatasmithUtils::ConvertVector((FDatasmithUtils::EModelCoordSystem) Context.ImportParams.GetModelCoordSys(), Vertex);
-				VertexPositions[VertexID] = FVector4f(SymmetricMatrix.TransformPosition((FVector)VertexPositions[VertexID]));
+				VertexPositions[VertexID] = FDatasmithUtils::ConvertVector((FDatasmithUtils::EModelCoordSystem) Context.ImportParams.GetModelCoordSys(), Vertex);
+				VertexPositions[VertexID] = SymmetricMatrix.TransformPosition(VertexPositions[VertexID]);
 				Context.SymmetricVertexIds[VertexIndex++] = VertexID;
 			}
 		}
@@ -227,7 +227,7 @@ namespace CADLibrary
 				if (!Step)
 				{
 					FDatasmithUtils::ConvertVectorArray(Context.ImportParams.GetModelCoordSys(), FaceMesh->Normals);
-					for (FVector& Normal : FaceMesh->Normals)
+					for (FVector3f& Normal : FaceMesh->Normals)
 					{
 						Normal = Normal.GetSafeNormal();
 					}
@@ -246,11 +246,10 @@ namespace CADLibrary
 				if (Step)
 				{
 					// compute normals of Symmetric vertex
-					FMatrix SymmetricMatrix;
-					SymmetricMatrix = FDatasmithUtils::GetSymmetricMatrix(Context.MeshParameters.SymmetricOrigin, Context.MeshParameters.SymmetricNormal);
+					FMatrix44f SymmetricMatrix = FDatasmithUtils::GetSymmetricMatrix(Context.MeshParameters.SymmetricOrigin, Context.MeshParameters.SymmetricNormal);
 					for (const FVertexInstanceID& VertexInstanceID : MeshVertexInstanceIDs)
 					{
-						VertexInstanceNormals[VertexInstanceID] = FVector4f(SymmetricMatrix.TransformVector((FVector)VertexInstanceNormals[VertexInstanceID]));
+						VertexInstanceNormals[VertexInstanceID] = SymmetricMatrix.TransformVector(VertexInstanceNormals[VertexInstanceID]);
 					}
 				}
 			}
@@ -346,9 +345,9 @@ namespace CADLibrary
 	{
 		ModelMesh.GetNodeCoordinates(OutBodyMesh.VertexArray);
 
-		for (FVector& Vertex : OutBodyMesh.VertexArray)
+		for (FVector3f& Vertex : OutBodyMesh.VertexArray)
 		{
-			Vertex *= 0.1;
+			Vertex *= 0.1f;
 		}
 
 		uint32 FaceSize = Body.FaceCount();
diff --git a/Engine/Plugins/Enterprise/DatasmithCADImporter/Source/CADLibrary/Private/MeshDescriptionHelper.cpp b/Engine/Plugins/Enterprise/DatasmithCADImporter/Source/CADLibrary/Private/MeshDescriptionHelper.cpp
--- a/Engine/Plugins/Enterprise/DatasmithCADImporter/Source/CADLibrary/Private/MeshDescriptionHelper.cpp
+++ b/Engine/Plugins/Enterprise/DatasmithCADImporter/Source/CADLibrary/Private/MeshDescriptionHelper.cpp
@@ -23,7 +23,7 @@ struct FVertexData
 {
 	float Z;
 	int32 Index;
-	FVector Coordinates;
+	FVector3f Coordinates;
 	bool bIsMerged;
 	FVertexID VertexID;
 	FVertexID SymVertexID;
@@ -32,7 +32,7 @@ struct FVertexData
 	FVertexData() {}
 
 	/** Initialization constructor. */
-	FVertexData(int32 InIndex, const FVector& V)
+	FVertexData(int32 InIndex, const FVector3f& V)
 	{
 		//Z = V.X + V.Y + V.Z;
 		Z = 0.30f * V.X + 0.33f * V.Y + 0.37f * V.Z;
@@ -64,7 +64,7 @@ bool IsTriangleDegenerated(const int32_t* Indices, const TArray<FVertexID>& Rema
 	return (VertexIDs[0] == VertexIDs[1] || VertexIDs[0] == VertexIDs[2] || VertexIDs[1] == VertexIDs[2]);
 }
 
-void MergeCoincidentVertices(TArray<FVector>& VertexArray, TArray<int32>& VertexIdSet)
+void MergeCoincidentVertices(TArray<FVector3f>& VertexArray, TArray<int32>& VertexIdSet)
 {
 	const double CoincidenceTolerance = 0.001;
 
@@ -72,7 +72,7 @@ void MergeCoincidentVertices(TArray<FVector>& VertexArray, TArray<int32>& Vertex
 	TArray<FVertexData> VertexDataSet;
 	VertexDataSet.Reserve(VertexArray.Num());
 
-	const FVector* Position = VertexArray.GetData();
+	const FVector3f* Position = VertexArray.GetData();
 	const int32* VertexId = VertexIdSet.GetData();
 	for (int32 Index = 0; Index < VertexArray.Num(); ++Index, ++Position, ++VertexId)
 	{
@@ -93,7 +93,7 @@ void MergeCoincidentVertices(TArray<FVector>& VertexArray, TArray<int32>& Vertex
 		VertexDataSet[Index].bIsMerged = true;
 		int32 NewIndex = VertexDataSet[Index].Index;
 
-		const FVector& PositionA = VertexDataSet[Index].Coordinates;
+		const FVector3f& PositionA = VertexDataSet[Index].Coordinates;
 
 		// only need to search forward, since we add pairs both ways
 		for (int32 Andex = Index + 1; Andex < VertexDataSet.Num(); Andex++)
@@ -103,7 +103,7 @@ void MergeCoincidentVertices(TArray<FVector>& VertexArray, TArray<int32>& Vertex
 				break; // can't be any more duplicated
 			}
 
-			const FVector& PositionB = VertexDataSet[Andex].Coordinates;
+			const FVector3f& PositionB = VertexDataSet[Andex].Coordinates;
 			if (PositionA.Equals(PositionB, CoincidenceTolerance))
 			{
 				VertexDataSet[Andex].bIsMerged = true;
@@ -120,7 +120,7 @@ void FillVertexPosition(const FImportParameters& ImportParams, const FMeshParame
 
 	TVertexAttributesRef<FVector3f> VertexPositions = MeshDescription.GetVertexPositions();
 
-	TArray<FVector>& VertexArray = Body.VertexArray;
+	TArray<FVector3f>& VertexArray = Body.VertexArray;
 	TArray<int32>& VertexIdSet = Body.VertexIds;
 	int32 VertexCount = VertexArray.Num();
 	VertexIdSet.SetNumZeroed(VertexCount);
@@ -129,7 +129,7 @@ void FillVertexPosition(const FImportParameters& ImportParams, const FMeshParame
 	MeshDescription.ReserveNewVertices(MeshParameters.bIsSymmetric ? VertexCount * 2 : VertexCount);
 
 	int32 VertexIndex = -1;
-	for (const FVector& Vertex : VertexArray)
+	for (const FVector3f& Vertex : VertexArray)
 	{
 		VertexIndex++;
 
@@ -149,13 +149,13 @@ void FillVertexPosition(const FImportParameters& ImportParams, const FMeshParame
 	// if Symmetric mesh, the symmetric side of the mesh have to be generated
 	if (MeshParameters.bIsSymmetric)
 	{
-		FMatrix SymmetricMatrix = FDatasmithUtils::GetSymmetricMatrix(MeshParameters.SymmetricOrigin, MeshParameters.SymmetricNormal);
+		FMatrix44f SymmetricMatrix = (FMatrix44f) FDatasmithUtils::GetSymmetricMatrix(MeshParameters.SymmetricOrigin, MeshParameters.SymmetricNormal);
 
 		TArray<int32>& SymmetricVertexIds = Body.SymmetricVertexIds;
 		SymmetricVertexIds.SetNum(VertexArray.Num());
 
 		VertexIndex = 0;
-		for (const FVector& Vertex : VertexArray)
+		for (const FVector3f& Vertex : VertexArray)
 		{
 			if (VertexIdSet[VertexIndex] == INDEX_NONE)
 			{
@@ -164,7 +164,7 @@ void FillVertexPosition(const FImportParameters& ImportParams, const FMeshParame
 			}
 
 			FVertexID VertexID = MeshDescription.CreateVertex();
-			const FVector VertexPosition = FDatasmithUtils::ConvertVector((FDatasmithUtils::EModelCoordSystem)ImportParams.GetModelCoordSys(), Vertex);
+			const FVector3f VertexPosition = FDatasmithUtils::ConvertVector((FDatasmithUtils::EModelCoordSystem)ImportParams.GetModelCoordSys(), Vertex);
 			const FVector4f SymmetricPosition = FVector4f(SymmetricMatrix.TransformPosition(VertexPosition));
 			VertexPositions[VertexID] = FVector3f(SymmetricPosition);
 			SymmetricVertexIds[VertexIndex++] = VertexID;
@@ -409,7 +409,7 @@ bool FillMesh(const FMeshParameters& MeshParameters, const FImportParameters& Im
 			if (!Step)
 			{
 				FDatasmithUtils::ConvertVectorArray(ImportParams.GetModelCoordSys(), Tessellation.NormalArray);
-				for (FVector& Normal : Tessellation.NormalArray)
+				for (FVector3f& Normal : Tessellation.NormalArray)
 				{
 					Normal = Normal.GetSafeNormal();
 				}
@@ -417,7 +417,7 @@ bool FillMesh(const FMeshParameters& MeshParameters, const FImportParameters& Im
 
 			if (Tessellation.NormalArray.Num() == 1)
 			{
-				const FVector& Normal = Tessellation.NormalArray[0];
+				const FVector3f& Normal = Tessellation.NormalArray[0];
 				for (int32 Index = 0; Index < NewFaceIndex.Num(); Index++)
 				{
 					const FVertexInstanceID VertexInstanceID = MeshVertexInstanceIDs[Index];
@@ -431,7 +431,7 @@ bool FillMesh(const FMeshParameters& MeshParameters, const FImportParameters& Im
 					for (int32 Index = 0; Index < 3; Index++)
 					{
 						const FVertexInstanceID VertexInstanceID = MeshVertexInstanceIDs[IndexFace + Orientation[Index]];
-						VertexInstanceNormals[VertexInstanceID] = (FVector3f)Tessellation.NormalArray[NewFaceIndex[IndexFace + Orientation[Index]]];
+						VertexInstanceNormals[VertexInstanceID] = Tessellation.NormalArray[NewFaceIndex[IndexFace + Orientation[Index]]];
 					}
 				}
 			}
@@ -440,11 +440,10 @@ bool FillMesh(const FMeshParameters& MeshParameters, const FImportParameters& Im
 			if (Step)
 			{
 				// compute normals of Symmetric vertex
-				FMatrix SymmetricMatrix;
-				SymmetricMatrix = FDatasmithUtils::GetSymmetricMatrix(MeshParameters.SymmetricOrigin, MeshParameters.SymmetricNormal);
+				FMatrix44f SymmetricMatrix = FDatasmithUtils::GetSymmetricMatrix(MeshParameters.SymmetricOrigin, MeshParameters.SymmetricNormal);
 				for (const FVertexInstanceID& VertexInstanceID : MeshVertexInstanceIDs)
 				{
-					VertexInstanceNormals[VertexInstanceID] = FVector4f(SymmetricMatrix.TransformVector((FVector)VertexInstanceNormals[VertexInstanceID]));
+					VertexInstanceNormals[VertexInstanceID] = SymmetricMatrix.TransformVector(VertexInstanceNormals[VertexInstanceID]);
 				}
 			}
 
diff --git a/Engine/Plugins/Enterprise/DatasmithCADImporter/Source/CADTools/Public/CADData.h b/Engine/Plugins/Enterprise/DatasmithCADImporter/Source/CADTools/Public/CADData.h
--- a/Engine/Plugins/Enterprise/DatasmithCADImporter/Source/CADTools/Public/CADData.h
+++ b/Engine/Plugins/Enterprise/DatasmithCADImporter/Source/CADTools/Public/CADData.h
@@ -295,7 +295,7 @@ struct CADTOOLS_API FTessellationData : public FArchiveGraphicProperties
 	friend CADTOOLS_API FArchive& operator<<(FArchive& Ar, FTessellationData& Tessellation);
 
 	/** Empty with CADKernel as set in FBodyMesh, Set by CoreTech (this is only the vertices of the face) */
-	TArray<FVector> PositionArray;
+	TArray<FVector3f> PositionArray;
 
 	/** Index of each vertex in FBody::VertexArray. Empty with CoreTech and filled by FillKioVertexPosition */
 	TArray<int32> PositionIndices;
@@ -304,10 +304,10 @@ struct CADTOOLS_API FTessellationData : public FArchiveGraphicProperties
 	TArray<int32> VertexIndices;
 
 	/** Normal of each vertex */
-	TArray<FVector> NormalArray;
+	TArray<FVector3f> NormalArray;
 
 	/** UV coordinates of each vertex */
-	TArray<FVector2D> TexCoordArray;
+	TArray<FVector2f> TexCoordArray;
 
 	int32 PatchId;
 };
@@ -336,7 +336,7 @@ public:
 		}
 	}
 
-	TArray<FVector> VertexArray; // set by CADKernel, filled by FillKioVertexPosition that merges coincident vertices (CoreTechHelper)
+	TArray<FVector3f> VertexArray; // set by CADKernel, filled by FillKioVertexPosition that merges coincident vertices (CoreTechHelper)
 	TArray<FTessellationData> Faces;
 	FBox BBox;
 
@@ -359,90 +359,5 @@ CADTOOLS_API FMaterialUId BuildMaterialUId(const FCADMaterial& Material);
 CADTOOLS_API void SerializeBodyMeshSet(const TCHAR* Filename, TArray<FBodyMesh>& InBodySet);
 CADTOOLS_API void DeserializeBodyMeshFile(const TCHAR* Filename, TArray<FBodyMesh>& OutBodySet);
 
-inline void CopyValue(const uint8* Source, int Offset, uint8 Size, bool bIs3D, FVector& Dest, const FMatrix* Matrix = nullptr)
-{
-	if (Source == nullptr)
-	{
-		return;
-	}
-
-	switch (Size)
-	{
-		case sizeof(uint8) :
-		{
-			Dest[0] = ((uint8*)Source)[Offset + 0] / 255.;
-			Dest[1] = ((uint8*)Source)[Offset + 1] / 255.;
-			Dest[2] = bIs3D ? ((uint8*)Source)[Offset + 2] / 255. : 0.;
-			break;
-		}
-		case sizeof(float) :
-		{
-			Dest[0] = ((float*)Source)[Offset + 0];
-			Dest[1] = ((float*)Source)[Offset + 1];
-			Dest[2] = bIs3D ? ((float*)Source)[Offset + 2] : 0.;
-			break;
-		}
-		case sizeof(double) :
-		{
-			Dest[0] = ((double*)Source)[Offset + 0];
-			Dest[1] = ((double*)Source)[Offset + 1];
-			Dest[2] = bIs3D ? ((double*)Source)[Offset + 2] : 0.;
-			break;
-		}
-		default:
-		{
-			Dest[0] = 0.;
-			Dest[1] = 0.;
-			Dest[2] = 0.;
-			break;
-		}
-	}
-
-	if (bIs3D && Matrix != nullptr)
-	{
-		Dest = Matrix->TransformPosition(Dest);
-	}
-}
-
-inline void CopyValue(const void* Source, int Offset, uint8 Size, int32 Dest[3])
-{
-	if (Source == nullptr)
-	{
-		return;
-	}
-
-	switch (Size)
-	{
-		case sizeof(uint8) :
-		{
-			Dest[0] = ((uint8*)Source)[Offset + 0];
-			Dest[1] = ((uint8*)Source)[Offset + 1];
-			Dest[2] = ((uint8*)Source)[Offset + 2];
-			break;
-		}
-		case sizeof(uint16) :
-		{
-			Dest[0] = ((uint16*)Source)[Offset + 0];
-			Dest[1] = ((uint16*)Source)[Offset + 1];
-			Dest[2] = ((uint16*)Source)[Offset + 2];
-			break;
-		}
-		case sizeof(uint32) :
-		{
-			Dest[0] = ((uint32*)Source)[Offset + 0];
-			Dest[1] = ((uint32*)Source)[Offset + 1];
-			Dest[2] = ((uint32*)Source)[Offset + 2];
-			break;
-		}
-		default:
-		{
-			Dest[0] = 0.;
-			Dest[1] = 0.;
-			Dest[2] = 0.;
-			break;
-		}
-	}
-}
-
 }
 
diff --git a/Engine/Plugins/Enterprise/DatasmithCADImporter/Source/CADTools/Public/CADOptions.h b/Engine/Plugins/Enterprise/DatasmithCADImporter/Source/CADTools/Public/CADOptions.h
--- a/Engine/Plugins/Enterprise/DatasmithCADImporter/Source/CADTools/Public/CADOptions.h
+++ b/Engine/Plugins/Enterprise/DatasmithCADImporter/Source/CADTools/Public/CADOptions.h
@@ -144,7 +144,7 @@ namespace CADLibrary
 	{
 		bool bNeedSwapOrientation = false;
 		bool bIsSymmetric = false;
-		FVector SymmetricOrigin;
-		FVector SymmetricNormal;
+		FVector3f SymmetricOrigin;
+		FVector3f SymmetricNormal;
 	};
 }
diff --git a/Engine/Plugins/Enterprise/DatasmithCADImporter/Source/DatasmithOpenNurbsTranslator/Private/DatasmithOpenNurbsTranslator.cpp b/Engine/Plugins/Enterprise/DatasmithCADImporter/Source/DatasmithOpenNurbsTranslator/Private/DatasmithOpenNurbsTranslator.cpp
--- a/Engine/Plugins/Enterprise/DatasmithCADImporter/Source/DatasmithOpenNurbsTranslator/Private/DatasmithOpenNurbsTranslator.cpp
+++ b/Engine/Plugins/Enterprise/DatasmithCADImporter/Source/DatasmithOpenNurbsTranslator/Private/DatasmithOpenNurbsTranslator.cpp
@@ -242,26 +242,26 @@ namespace DatasmithOpenNurbsTranslatorUtils
 		{ }
 
 		FNode(float x, float y, float z)
-			: Vertex(FVector(x, y, z))
+			: Vertex(x, y, z)
 		{ }
 
 		void SetVertex(float x, float y, float z)
 		{
-			Vertex = FVector(x, y, z);
+			Vertex = FVector3f(x, y, z);
 		}
 
 		void SetNormal(float x, float y, float z)
 		{
-			Normal = FVector(x, y, z);
+			Normal = FVector3f(x, y, z);
 		}
 
 		void SetNormal(const ON_3fVector& InNormal)
 		{
-			Normal = FVector(InNormal[0], InNormal[1], InNormal[2]);
+			Normal = FVector3f(InNormal[0], InNormal[1], InNormal[2]);
 		}
 
-		FVector Vertex;
-		FVector Normal;
+		FVector3f Vertex;
+		FVector3f Normal;
 	};
 
 	struct FFace
@@ -281,7 +281,7 @@ namespace DatasmithOpenNurbsTranslatorUtils
 		}
 
 		const FNode* Nodes[3];
-		FVector2D UVCoords[3];
+		FVector2f UVCoords[3];
 	};
 
 	bool TranslateMesh(const ON_Mesh** Meshes, int MeshCount, FMeshDescription& MeshDescription, bool& bHasNormal, double ScalingFactor, const ON_3dVector& Offset, bool bHasFaceMaterialChannel, int32* FaceMaterialChannel)
@@ -416,7 +416,7 @@ namespace DatasmithOpenNurbsTranslatorUtils
 					FFace& Face = Faces[Index];
 					for (int iTriangleNode = 0; iTriangleNode < 3; ++iTriangleNode)
 					{
-						Face.UVCoords[iTriangleNode] = FVector2D(UvCoords[iTriangleNode + 3 * Index][0], UvCoords[iTriangleNode + 3 * Index][1]);
+						Face.UVCoords[iTriangleNode] = FVector2f(UvCoords[iTriangleNode + 3 * Index][0], UvCoords[iTriangleNode + 3 * Index][1]);
 					}
 				}
 			}
@@ -450,7 +450,7 @@ namespace DatasmithOpenNurbsTranslatorUtils
 			for (int Index = 0; Index < VertexCount; ++Index)
 			{
 				const FNode& Node = Nodes[Index];
-				const FVector& Pos = Node.Vertex;
+				const FVector3f& Pos = Node.Vertex;
 				NodeToIndex.Add(&Node, VertexIndexBase + Index);
 
 				// Fill the vertex array
@@ -495,7 +495,7 @@ namespace DatasmithOpenNurbsTranslatorUtils
 					const FNode& FaceNode = *Face.Nodes[CornerIndex];
 
 					// Set the normal
-					FVector UENormal = FDatasmithUtils::ConvertVector(FDatasmithUtils::EModelCoordSystem::ZUp_RightHanded_FBXLegacy, FaceNode.Normal);
+					FVector3f UENormal = FDatasmithUtils::ConvertVector(FDatasmithUtils::EModelCoordSystem::ZUp_RightHanded_FBXLegacy, FaceNode.Normal);
 					UENormal = UENormal.GetSafeNormal();
 
 					// Check to see if normal is correct. If not replace by face's normal
@@ -511,7 +511,7 @@ namespace DatasmithOpenNurbsTranslatorUtils
 					// Set the UV
 					if (bHasUVData)
 					{
-						const FVector2D& UVValues = Face.UVCoords[CornerIndex];
+						const FVector2f& UVValues = Face.UVCoords[CornerIndex];
 						if (!UVValues.ContainsNaN())
 						{
 							// Convert UVs - bottom-left Rhino's texture origin to Unreal's top-left
@@ -2065,7 +2065,8 @@ TSharedPtr<IDatasmithMeshActorElement> FOpenNurbsTranslatorImpl::GetMeshActorEle
 	if (ComputeObjectGeometryCenter(Object, GeometryCenter))
 	{
 		MeshElementToGeometryCenter.Add(MeshElement.ToSharedRef(), GeometryCenter);
-		FVector ActorOffset = GetScaleFactor() * FDatasmithUtils::ConvertVector(FDatasmithUtils::EModelCoordSystem::ZUp_RightHanded_FBXLegacy, GeometryCenter);
+		FVector ActorOffset(GeometryCenter.x, GeometryCenter.y, GeometryCenter.z);
+		ActorOffset = FDatasmithUtils::ConvertVector(FDatasmithUtils::EModelCoordSystem::ZUp_RightHanded_FBXLegacy, ActorOffset) * GetScaleFactor();
 		ActorElement->SetTranslation(ActorOffset);
 	}
 
diff --git a/Engine/Plugins/Enterprise/DatasmithCADImporter/Source/DatasmithWireTranslator/Private/OpenModelUtils.cpp b/Engine/Plugins/Enterprise/DatasmithCADImporter/Source/DatasmithWireTranslator/Private/OpenModelUtils.cpp
--- a/Engine/Plugins/Enterprise/DatasmithCADImporter/Source/DatasmithWireTranslator/Private/OpenModelUtils.cpp
+++ b/Engine/Plugins/Enterprise/DatasmithCADImporter/Source/DatasmithWireTranslator/Private/OpenModelUtils.cpp
@@ -91,7 +91,7 @@ bool OpenModelUtils::TransferAlMeshToMeshDescription(const AlMesh& AliasMesh, co
 	}
 
 	int32 NbStep = 1;
-	FMatrix SymmetricMatrix;
+	FMatrix44f SymmetricMatrix;
 	bool bIsSymmetricMesh = MeshParameters.bIsSymmetric;
 	if (bIsSymmetricMesh)
 	{
@@ -165,7 +165,7 @@ bool OpenModelUtils::TransferAlMeshToMeshDescription(const AlMesh& AliasMesh, co
 				const float* CurVertex = AlVertices + 3 * Index;
 				*VertexPositionIDPtr = MeshDescription.CreateVertex();
 				// ConvertVector_ZUp_RightHanded
-				VertexPositions[*VertexPositionIDPtr] = FVector4f(SymmetricMatrix.TransformPosition(FVector(-CurVertex[0], CurVertex[1], CurVertex[2])));
+				VertexPositions[*VertexPositionIDPtr] = SymmetricMatrix.TransformPosition(FVector3f(-CurVertex[0], CurVertex[1], CurVertex[2]));
 			}
 		}
 
@@ -195,7 +195,7 @@ bool OpenModelUtils::TransferAlMeshToMeshDescription(const AlMesh& AliasMesh, co
 					// Set the normal
 					const float* CurNormal = &AlNormals[3 * Triangles[TIndex]];
 					// ConvertVector_ZUp_RightHanded
-					FVector UENormal(-CurNormal[0], CurNormal[1], CurNormal[2]);
+					FVector3f UENormal(-CurNormal[0], CurNormal[1], CurNormal[2]);
 					UENormal = UENormal.GetSafeNormal();
 					if (Step > 0)
 					{
@@ -205,7 +205,7 @@ bool OpenModelUtils::TransferAlMeshToMeshDescription(const AlMesh& AliasMesh, co
 					{
 						UENormal *= -1.;
 					}
-					VertexInstanceNormals[CornerVertexInstanceIDs[VertexIndex]] = (FVector3f)UENormal;
+					VertexInstanceNormals[CornerVertexInstanceIDs[VertexIndex]] = UENormal;
 				}
 				if (CornerVertexIDs[0] == CornerVertexIDs[1] || CornerVertexIDs[0] == CornerVertexIDs[2] || CornerVertexIDs[1] == CornerVertexIDs[2])
 				{
@@ -242,7 +242,7 @@ bool OpenModelUtils::TransferAlMeshToMeshDescription(const AlMesh& AliasMesh, co
 					const float* CurNormal = &AlNormals[3 * Triangles[VertexIndex]];
 
 					// ConvertVector_ZUp_RightHanded
-					FVector UENormal(-CurNormal[0], CurNormal[1], CurNormal[2]);
+					FVector3f UENormal(-CurNormal[0], CurNormal[1], CurNormal[2]);
 					UENormal = UENormal.GetSafeNormal();
 					if (Step > 0)
 					{
diff --git a/Engine/Plugins/Enterprise/DatasmithCADImporter/Source/ParametricSurface/Private/ParametricSurfaceData.cpp b/Engine/Plugins/Enterprise/DatasmithCADImporter/Source/ParametricSurface/Private/ParametricSurfaceData.cpp
--- a/Engine/Plugins/Enterprise/DatasmithCADImporter/Source/ParametricSurface/Private/ParametricSurfaceData.cpp
+++ b/Engine/Plugins/Enterprise/DatasmithCADImporter/Source/ParametricSurface/Private/ParametricSurfaceData.cpp
@@ -32,8 +32,8 @@ void UParametricSurfaceData::SetMeshParameters(const CADLibrary::FMeshParameters
 {
 	MeshParameters.bNeedSwapOrientation = InMeshParameters.bNeedSwapOrientation;
 	MeshParameters.bIsSymmetric = InMeshParameters.bIsSymmetric;
-	MeshParameters.SymmetricNormal = InMeshParameters.SymmetricNormal;
-	MeshParameters.SymmetricOrigin = InMeshParameters.SymmetricOrigin;
+	MeshParameters.SymmetricNormal = (FVector) InMeshParameters.SymmetricNormal;
+	MeshParameters.SymmetricOrigin = (FVector) InMeshParameters.SymmetricOrigin;
 }
 
 void UParametricSurfaceData::Serialize(FArchive& Ar)
diff --git a/Engine/Plugins/Enterprise/DatasmithCADImporter/Source/ParametricSurface/Public/ParametricSurfaceData.h b/Engine/Plugins/Enterprise/DatasmithCADImporter/Source/ParametricSurface/Public/ParametricSurfaceData.h
--- a/Engine/Plugins/Enterprise/DatasmithCADImporter/Source/ParametricSurface/Public/ParametricSurfaceData.h
+++ b/Engine/Plugins/Enterprise/DatasmithCADImporter/Source/ParametricSurface/Public/ParametricSurfaceData.h
@@ -53,8 +53,8 @@ struct PARAMETRICSURFACE_API FParametricMeshParameters
 
 		Parameters.bNeedSwapOrientation = bNeedSwapOrientation;
 		Parameters.bIsSymmetric = bIsSymmetric;
-		Parameters.SymmetricNormal = SymmetricNormal;
-		Parameters.SymmetricOrigin = SymmetricOrigin;
+		Parameters.SymmetricNormal = (FVector3f) SymmetricNormal;
+		Parameters.SymmetricOrigin = (FVector3f) SymmetricOrigin;
 
 		return Parameters;
 	}
diff --git a/Engine/Source/Runtime/Datasmith/CADKernel/Private/CADKernel/Mesh/Structure/FaceMesh.cpp b/Engine/Source/Runtime/Datasmith/CADKernel/Private/CADKernel/Mesh/Structure/FaceMesh.cpp
--- a/Engine/Source/Runtime/Datasmith/CADKernel/Private/CADKernel/Mesh/Structure/FaceMesh.cpp
+++ b/Engine/Source/Runtime/Datasmith/CADKernel/Private/CADKernel/Mesh/Structure/FaceMesh.cpp
@@ -62,7 +62,7 @@ void FFaceMesh::InverseOrientation()
 		Swap(TrianglesVerticesIndex[Index], TrianglesVerticesIndex[Index+1]);
 	}
 
-	for (FVector& Normal : Normals)
+	for (FVector3f& Normal : Normals)
 	{
 		Normal *= -1.;
 	}
diff --git a/Engine/Source/Runtime/Datasmith/CADKernel/Private/CADKernel/Mesh/Structure/ModelMesh.cpp b/Engine/Source/Runtime/Datasmith/CADKernel/Private/CADKernel/Mesh/Structure/ModelMesh.cpp
--- a/Engine/Source/Runtime/Datasmith/CADKernel/Private/CADKernel/Mesh/Structure/ModelMesh.cpp
+++ b/Engine/Source/Runtime/Datasmith/CADKernel/Private/CADKernel/Mesh/Structure/ModelMesh.cpp
@@ -109,7 +109,7 @@ void FModelMesh::GetNodeCoordinates(TArray<FPoint>& NodeCoordinates) const
 	}
 }
 
-void FModelMesh::GetNodeCoordinates(TArray<FVector>& NodeCoordinates) const
+void FModelMesh::GetNodeCoordinates(TArray<FVector3f>& NodeCoordinates) const
 {
 	NodeCoordinates.Reserve(LastIdUsed);
 
@@ -117,7 +117,7 @@ void FModelMesh::GetNodeCoordinates(TArray<FVector>& NodeCoordinates) const
 	{
 		for (const FPoint& Point : *PointArray)
 		{
-			NodeCoordinates.Emplace((float)Point.X, (float)Point.Y, (float)Point.Z);
+			NodeCoordinates.Emplace(Point.X, Point.Y, Point.Z);
 		}
 	}
 }
diff --git a/Engine/Source/Runtime/Datasmith/CADKernel/Private/CADKernel/Mesh/Structure/ModelMeshAnalyzer.cpp b/Engine/Source/Runtime/Datasmith/CADKernel/Private/CADKernel/Mesh/Structure/ModelMeshAnalyzer.cpp
--- a/Engine/Source/Runtime/Datasmith/CADKernel/Private/CADKernel/Mesh/Structure/ModelMeshAnalyzer.cpp
+++ b/Engine/Source/Runtime/Datasmith/CADKernel/Private/CADKernel/Mesh/Structure/ModelMeshAnalyzer.cpp
@@ -53,7 +53,7 @@ void FModelMeshAnalyzer::BuildMesh()
 	{
 		TArray<int32>& TrianglesVerticesIndex = FaceMesh->TrianglesVerticesIndex;
 		TArray<int32>& VerticesGlobalIndex = FaceMesh->VerticesGlobalIndex;
-		TArray<FVector>& FaceNormals = FaceMesh->Normals;
+		TArray<FVector3f>& FaceNormals = FaceMesh->Normals;
 
 		for (int32 Index = 0; Index < FaceMesh->TrianglesVerticesIndex.Num(); Index += 3)
 		{
@@ -67,7 +67,7 @@ void FModelMeshAnalyzer::BuildMesh()
 			TriangleEdges[1] = FindOrBuildEdge(Vertices[1], Vertices[2]);
 			TriangleEdges[2] = FindOrBuildEdge(Vertices[2], Vertices[0]);
 
-			FVector Normals[3];
+			FVector3f Normals[3];
 			for (uint32 NormalI = 0; NormalI < 3; ++NormalI)
 			{
 				Normals[NormalI] = FaceNormals[TrianglesVerticesIndex[Index + NormalI]];
diff --git a/Engine/Source/Runtime/Datasmith/CADKernel/Private/CADKernel/Mesh/Structure/ModelMeshAnalyzer.h b/Engine/Source/Runtime/Datasmith/CADKernel/Private/CADKernel/Mesh/Structure/ModelMeshAnalyzer.h
--- a/Engine/Source/Runtime/Datasmith/CADKernel/Private/CADKernel/Mesh/Structure/ModelMeshAnalyzer.h
+++ b/Engine/Source/Runtime/Datasmith/CADKernel/Private/CADKernel/Mesh/Structure/ModelMeshAnalyzer.h
@@ -20,12 +20,12 @@ struct FEdge;
 struct FTriangle
 {
 	uint32 Vertices[3];
-	FVector Normals[3];
+	FVector3f Normals[3];
 	FEdge* Edges[3];
 	const FFaceMesh* FaceMesh;
 
 
-	FTriangle(uint32* InVertices, FEdge** InEdges, FVector* InNormals, const FFaceMesh* InFaceMesh)
+	FTriangle(uint32* InVertices, FEdge** InEdges, FVector3f* InNormals, const FFaceMesh* InFaceMesh)
 	: FaceMesh(InFaceMesh)
 	{
 		for (int32 Index = 0; Index < 3; ++Index)
diff --git a/Engine/Source/Runtime/Datasmith/CADKernel/Public/CADKernel/Mesh/Structure/FaceMesh.h b/Engine/Source/Runtime/Datasmith/CADKernel/Public/CADKernel/Mesh/Structure/FaceMesh.h
--- a/Engine/Source/Runtime/Datasmith/CADKernel/Public/CADKernel/Mesh/Structure/FaceMesh.h
+++ b/Engine/Source/Runtime/Datasmith/CADKernel/Public/CADKernel/Mesh/Structure/FaceMesh.h
@@ -30,10 +30,10 @@ namespace CADKernel
 		TArray<int32> VerticesGlobalIndex;
 
 		/** Normals of the surface mesh vertices*/
-		TArray<FVector> Normals;
+		TArray<FVector3f> Normals;
 
 		/** UV coordinates of the surface mesh vertices*/
-		TArray<FVector2D> UVMap;
+		TArray<FVector2f> UVMap;
 
 	public:
 		FFaceMesh(FModelMesh& InMeshModel, FTopologicalEntity& InTopologicalEntity)
diff --git a/Engine/Source/Runtime/Datasmith/CADKernel/Public/CADKernel/Mesh/Structure/ModelMesh.h b/Engine/Source/Runtime/Datasmith/CADKernel/Public/CADKernel/Mesh/Structure/ModelMesh.h
--- a/Engine/Source/Runtime/Datasmith/CADKernel/Public/CADKernel/Mesh/Structure/ModelMesh.h
+++ b/Engine/Source/Runtime/Datasmith/CADKernel/Public/CADKernel/Mesh/Structure/ModelMesh.h
@@ -150,7 +150,7 @@ namespace CADKernel
 		const TSharedPtr<FVertexMesh> GetMeshOfVertexNodeId(const int32 Id) const;
 
 		void GetNodeCoordinates(TArray<FPoint>& NodeCoordinates) const;
-		void GetNodeCoordinates(TArray<FVector>& NodeCoordinates) const;
+		void GetNodeCoordinates(TArray<FVector3f>& NodeCoordinates) const;
 
 		const TArray<TSharedPtr<FMesh>>& GetMeshes() const;
 
diff --git a/Engine/Source/Runtime/Datasmith/DatasmithCore/Private/DatasmithUtils.cpp b/Engine/Source/Runtime/Datasmith/DatasmithCore/Private/DatasmithUtils.cpp
--- a/Engine/Source/Runtime/Datasmith/DatasmithCore/Private/DatasmithUtils.cpp
+++ b/Engine/Source/Runtime/Datasmith/DatasmithCore/Private/DatasmithUtils.cpp
@@ -1699,46 +1699,4 @@ FTransform FDatasmithUtils::ConvertTransform(EModelCoordSystem SourceCoordSystem
 }
 
 
-FMatrix FDatasmithUtils::GetSymmetricMatrix(const FVector& Origin, const FVector& Normal)
-{
-	//Calculate symmetry matrix
-	//(Px, Py, Pz) = normal
-	// -Px² + Pz² + Py²  |  - 2 * Px * Py     |  - 2 * Px * Pz
-	// - 2 * Py * Px     |  - Py² + Px² + Pz² |  - 2 * Py * Pz
-	// - 2 * Pz * Px     |  - 2 * Pz * Py     |  - Pz² + Py² + Px²
-
-	FVector LocOrigin = Origin;
-
-	float NormalXSqr, NormalYSqr, NormalZSqr;
-	NormalXSqr = Normal.X * Normal.X;
-	NormalYSqr = Normal.Y * Normal.Y;
-	NormalZSqr = Normal.Z * Normal.Z;
-
-	FMatrix OSymmetricMatrix;
-	OSymmetricMatrix.SetIdentity();
-	FVector Axis0(-NormalXSqr + NormalZSqr + NormalYSqr, -2 * Normal.X * Normal.Y, -2 * Normal.X * Normal.Z);
-	FVector Axis1(-2 * Normal.Y * Normal.X, -NormalYSqr + NormalXSqr + NormalZSqr, -2 * Normal.Y * Normal.Z);
-	FVector Axis2(-2 * Normal.Z * Normal.X, -2 * Normal.Z * Normal.Y, -NormalZSqr + NormalYSqr + NormalXSqr);
-	OSymmetricMatrix.SetAxes(&Axis0, &Axis1, &Axis2);
-
-	FMatrix SymmetricMatrix;
-	SymmetricMatrix.SetIdentity();
-
-	//Translate to 0, 0, 0
-	LocOrigin *= -1.;
-	SymmetricMatrix.SetOrigin(LocOrigin);
-
-	//// Apply Symmetric
-	SymmetricMatrix *= OSymmetricMatrix;
-
-	//Translate to original position
-	LocOrigin *= -1.;
-	FMatrix OrigTranslation;
-	OrigTranslation.SetIdentity();
-	OrigTranslation.SetOrigin(LocOrigin);
-	SymmetricMatrix *= OrigTranslation;
-
-	return SymmetricMatrix;
-}
-
 
diff --git a/Engine/Source/Runtime/Datasmith/DatasmithCore/Public/DatasmithUtils.h b/Engine/Source/Runtime/Datasmith/DatasmithCore/Public/DatasmithUtils.h
--- a/Engine/Source/Runtime/Datasmith/DatasmithCore/Public/DatasmithUtils.h
+++ b/Engine/Source/Runtime/Datasmith/DatasmithCore/Public/DatasmithUtils.h
@@ -74,33 +74,34 @@ public:
 		ZUp_RightHanded_FBXLegacy,
 	};
 
-	static void ConvertVectorArray(EModelCoordSystem ModelCoordSys, TArray<FVector>& Array)
+	template<typename VecType>
+	static void ConvertVectorArray(EModelCoordSystem ModelCoordSys, TArray<VecType>& Array)
 	{
 		switch (ModelCoordSys)
 		{
 		case EModelCoordSystem::YUp_LeftHanded:
-			for (FVector& Vector : Array)
+			for (VecType& Vector : Array)
 			{
 				Vector.Set(Vector[2], Vector[0], Vector[1]);
 			}
 			break;
 
 		case EModelCoordSystem::YUp_RightHanded:
-			for (FVector& Vector : Array)
+			for (VecType& Vector : Array)
 			{
 				Vector.Set(-Vector[2], Vector[0], Vector[1]);
 			}
 			break;
 
 		case EModelCoordSystem::ZUp_RightHanded:
-			for (FVector& Vector : Array)
+			for (VecType& Vector : Array)
 			{
 				Vector.Set(-Vector[0], Vector[1], Vector[2]);
 			}
 			break;
 
 		case EModelCoordSystem::ZUp_RightHanded_FBXLegacy:
-			for (FVector& Vector : Array)
+			for (VecType& Vector : Array)
 			{
 				Vector.Set(Vector[0], -Vector[1], Vector[2]);
 			}
@@ -113,31 +114,72 @@ public:
 	}
 
 	template<typename VecType>
-	static FVector ConvertVector(EModelCoordSystem ModelCoordSys, const VecType& V)
+	static VecType ConvertVector(EModelCoordSystem ModelCoordSys, const VecType& V)
 	{
 		switch (ModelCoordSys)
 		{
 		case EModelCoordSystem::YUp_LeftHanded:
-			return FVector(V[2], V[0], V[1]);
+			return VecType(V[2], V[0], V[1]);
 
 		case EModelCoordSystem::YUp_RightHanded:
-			return FVector(-V[2], V[0], V[1]);
+			return VecType(-V[2], V[0], V[1]);
 
 		case EModelCoordSystem::ZUp_RightHanded:
-			return FVector(-V[0], V[1], V[2]);
+			return VecType(-V[0], V[1], V[2]);
 
 		case EModelCoordSystem::ZUp_RightHanded_FBXLegacy:
-			return FVector(V[0], -V[1], V[2]);
+			return VecType(V[0], -V[1], V[2]);
 
 		case EModelCoordSystem::ZUp_LeftHanded:
 		default:
-			return FVector(V[0], V[1], V[2]);
+			return VecType(V[0], V[1], V[2]);
 		}
 	}
 
 	FTransform static ConvertTransform(EModelCoordSystem SourceCoordSystem, const FTransform& LocalTransform);
 
-	FMatrix static GetSymmetricMatrix(const FVector& Origin, const FVector& Normal);
+	template<typename Type>
+	static UE::Math::TMatrix<Type> GetSymmetricMatrix(const UE::Math::TVector<Type>& Origin, const UE::Math::TVector<Type>& Normal)
+	{
+		using namespace UE::Math;
+		//Calculate symmetry matrix
+		//(Px, Py, Pz) = normal
+		// -Px*Px + Pz*Pz + Py*Py  |  - 2 * Px * Py           |  - 2 * Px * Pz
+		// - 2 * Py * Px           |  - Py*Py + Px*Px + Pz*Pz |  - 2 * Py * Pz
+		// - 2 * Pz * Px           |  - 2 * Pz * Py           |  - Pz*Pz + Py*Py + Px*Px
+
+		TVector<Type> LocOrigin = Origin;
+
+		Type NormalXSqr = Normal.X * Normal.X;
+		Type NormalYSqr = Normal.Y * Normal.Y;
+		Type NormalZSqr = Normal.Z * Normal.Z;
+
+		TMatrix<Type> OSymmetricMatrix;
+		OSymmetricMatrix.SetIdentity();
+		TVector<Type> Axis0(-NormalXSqr + NormalZSqr + NormalYSqr, -2 * Normal.X * Normal.Y, -2 * Normal.X * Normal.Z);
+		TVector<Type> Axis1(-2 * Normal.Y * Normal.X, -NormalYSqr + NormalXSqr + NormalZSqr, -2 * Normal.Y * Normal.Z);
+		TVector<Type> Axis2(-2 * Normal.Z * Normal.X, -2 * Normal.Z * Normal.Y, -NormalZSqr + NormalYSqr + NormalXSqr);
+		OSymmetricMatrix.SetAxes(&Axis0, &Axis1, &Axis2);
+
+		TMatrix<Type> SymmetricMatrix;
+		SymmetricMatrix.SetIdentity();
+
+		//Translate to 0, 0, 0
+		LocOrigin *= -1.;
+		SymmetricMatrix.SetOrigin(LocOrigin);
+
+		//// Apply Symmetric
+		SymmetricMatrix *= OSymmetricMatrix;
+
+		//Translate to original position
+		LocOrigin *= -1.;
+		TMatrix<Type> OrigTranslation;
+		OrigTranslation.SetIdentity();
+		OrigTranslation.SetOrigin(LocOrigin);
+		SymmetricMatrix *= OrigTranslation;
+
+		return SymmetricMatrix;
+	}
 };
 
 class DATASMITHCORE_API FDatasmithMeshUtils
