To avoid duplicated triangle, a map VertexIds -> Triangle is used.
The time cost of the use of this map is about 20% for [REDACTED_PHONE] triangles

The CVar ds.CADTranslator.RemoveDuplicatedTriangle (false by default) enable the process
As this problem is rare, and to avoid the cost of this process, this option is disable by default. To fix StaticMesh with duplicated triangles, the StaticMesh can be remesh with this option.

diff --git a/Engine/Plugins/Enterprise/DatasmithCADImporter/Source/CADLibrary/Private/CADKernelTools.cpp b/Engine/Plugins/Enterprise/DatasmithCADImporter/Source/CADLibrary/Private/CADKernelTools.cpp
--- a/Engine/Plugins/Enterprise/DatasmithCADImporter/Source/CADLibrary/Private/CADKernelTools.cpp
+++ b/Engine/Plugins/Enterprise/DatasmithCADImporter/Source/CADLibrary/Private/CADKernelTools.cpp
@@ -53,7 +53,7 @@ namespace CADLibrary
 
 		for (FVector3f& Vertex : VertexArray)
 		{
-			Vertex *= 0.1; // mm (CADKernel unit) to cm (UE unit)
+			Vertex *= 0.1f; // mm (CADKernel unit) to cm (UE unit)
 		}
 
 		int32 VertexCount = VertexArray.Num();
diff --git a/Engine/Plugins/Enterprise/DatasmithCADImporter/Source/CADLibrary/Private/MeshDescriptionHelper.cpp b/Engine/Plugins/Enterprise/DatasmithCADImporter/Source/CADLibrary/Private/MeshDescriptionHelper.cpp
--- a/Engine/Plugins/Enterprise/DatasmithCADImporter/Source/CADLibrary/Private/MeshDescriptionHelper.cpp
+++ b/Engine/Plugins/Enterprise/DatasmithCADImporter/Source/CADLibrary/Private/MeshDescriptionHelper.cpp
@@ -66,6 +66,8 @@ bool IsTriangleDegenerated(const int32_t* Indices, const TArray<FVertexID>& Rema
 
 void MergeCoincidentVertices(TArray<FVector3f>& VertexArray, TArray<int32>& VertexIdSet)
 {
+	TRACE_CPUPROFILER_EVENT_SCOPE(FDatasmithMeshBuilder::MergeCoincidentVertices);
+
 	const double CoincidenceTolerance = 0.001;
 
 	// Create a list of vertex Z/index pairs
@@ -115,6 +117,8 @@ void MergeCoincidentVertices(TArray<FVector3f>& VertexArray, TArray<int32>& Vert
 
 void FillVertexPosition(const FImportParameters& ImportParams, const FMeshParameters& MeshParameters, FBodyMesh& Body, FMeshDescription& MeshDescription)
 {
+	TRACE_CPUPROFILER_EVENT_SCOPE(CADLibrary::FillVertexPosition);
+
 	int32 TriangleCount = Body.TriangleCount;
 	TArray<FTessellationData>& FaceTessellationSet = Body.Faces;
 
@@ -258,6 +262,8 @@ void CopyMaterialSlotNames(FMeshDescription& MeshSource, FMeshDescription& MeshD
 
 bool FillMesh(const FMeshParameters& MeshParameters, const FImportParameters& ImportParams, FBodyMesh& BodyTessellation, FMeshDescription& MeshDescription)
 {
+	TRACE_CPUPROFILER_EVENT_SCOPE(CADLibrary::FillMesh);
+
 	const int32 UVChannel = 0;
 	const int32 VertexCountPerFace = 3;
 	const TriangleIndex Clockwise = { 0, 1, 2 };
@@ -283,6 +289,9 @@ bool FillMesh(const FMeshParameters& MeshParameters, const FImportParameters& Im
 		return false;
 	}
 
+	// To avoid duplicated triangle
+	TSet<TTuple<FVertexInstanceID, FVertexInstanceID, FVertexInstanceID>> VertexIdsToTriangle;
+
 	// Find all the materials used
 	TMap<uint32, FPolygonGroupID> MaterialToPolygonGroupMapping;
 	for (const FTessellationData& FaceTessellation : BodyTessellation.Faces)
@@ -380,8 +389,19 @@ bool FillMesh(const FMeshParameters& MeshParameters, const FImportParameters& Im
 				MeshVertexInstanceIDs.Add(TriangleVertexInstanceIDs[1] = MeshDescription.CreateVertexInstance((FVertexID)FaceVertexIDs[1]));
 				MeshVertexInstanceIDs.Add(TriangleVertexInstanceIDs[2] = MeshDescription.CreateVertexInstance((FVertexID)FaceVertexIDs[2]));
 
+				if(FImportParameters::bGRemoveDuplicatedTriangle)
+				{
+					Sort(FaceVertexIDs, 3);
+					bool bIsAlreadyInSet;
+					VertexIdsToTriangle.Emplace(TTuple<FVertexInstanceID, FVertexInstanceID, FVertexInstanceID>(FaceVertexIDs[0], FaceVertexIDs[1], FaceVertexIDs[2]), &bIsAlreadyInSet);
+					if (bIsAlreadyInSet)
+					{
+						continue;
+					}
+				}
+
 				// Add the triangle as a polygon to the mesh description
-				const FPolygonID PolygonID = MeshDescription.CreatePolygon(*PolygonGroupID, TriangleVertexInstanceIDs);
+				FPolygonID PolygonID = MeshDescription.CreatePolygon(*PolygonGroupID, TriangleVertexInstanceIDs);
 
 				// Set patch id attribute
 				PatchGroups[PolygonID] = Tessellation.PatchId;
@@ -461,6 +481,8 @@ bool FillMesh(const FMeshParameters& MeshParameters, const FImportParameters& Im
 
 bool ConvertBodyMeshToMeshDescription(const FImportParameters& ImportParams, const FMeshParameters& MeshParameters, FBodyMesh& Body, FMeshDescription& MeshDescription)
 {
+	TRACE_CPUPROFILER_EVENT_SCOPE(CADLibrary::ConvertBodyMeshToMeshDescription);
+
 	// in a closed big mesh VertexCount ~ TriangleCount / 2, EdgeCount ~ 1.5* TriangleCount
 	MeshDescription.ReserveNewVertexInstances(Body.VertexArray.Num());
 	MeshDescription.ReserveNewPolygons(Body.TriangleCount);
diff --git a/Engine/Plugins/Enterprise/DatasmithCADImporter/Source/CADTools/Private/CADOption.cpp b/Engine/Plugins/Enterprise/DatasmithCADImporter/Source/CADTools/Private/CADOption.cpp
--- a/Engine/Plugins/Enterprise/DatasmithCADImporter/Source/CADTools/Private/CADOption.cpp
+++ b/Engine/Plugins/Enterprise/DatasmithCADImporter/Source/CADTools/Private/CADOption.cpp
@@ -79,6 +79,17 @@ This welding is performed respecting the ds.CADTranslator.StitchingTolerance\n\
 Default value is true\n"),
 ECVF_Default);
 
+bool FImportParameters::bGRemoveDuplicatedTriangle = false;
+FAutoConsoleVariableRef GCADTranslatorRemoveDuplicatedTriangle(
+	TEXT("ds.CADTranslator.RemoveDuplicatedTriangle"),
+	FImportParameters::bGRemoveDuplicatedTriangle,
+	TEXT("Avoid duplicated triangles in a mesh (i.e. two triangle with the same triplet of vertices)\n\
+This problem appears when a surface patch is duplicated. Due to the two identical mesh of the patches, the global mesh has duplicated triangles\n\
+The time cost of this process is about 20% for 10e6 triangles with a complexity in nlog(n) due to the use of a map\n\
+As this problem is rare, and to avoid the cost of the process, this option is disable by default. To fix StaticMesh with duplicated triangles, the StaticMesh can be remesh with this option.\n\
+Default value is false. \n"),
+ECVF_Default);
+
 float FImportParameters::GMeshingParameterFactor = 1;
 FAutoConsoleVariableRef GCADTranslatorMeshingParameterFactor(
 	TEXT("ds.CADTranslator.MeshingParameterFactor"),
diff --git a/Engine/Plugins/Enterprise/DatasmithCADImporter/Source/CADTools/Public/CADOptions.h b/Engine/Plugins/Enterprise/DatasmithCADImporter/Source/CADTools/Public/CADOptions.h
--- a/Engine/Plugins/Enterprise/DatasmithCADImporter/Source/CADTools/Public/CADOptions.h
+++ b/Engine/Plugins/Enterprise/DatasmithCADImporter/Source/CADTools/Public/CADOptions.h
@@ -35,6 +35,7 @@ namespace CADLibrary
 		CADTOOLS_API static bool bGOverwriteCache;
 		CADTOOLS_API static bool bGPreferJtFileEmbeddedTessellation;
 		CADTOOLS_API static bool bGSewMeshIfNeeded;
+		CADTOOLS_API static bool bGRemoveDuplicatedTriangle;
 		CADTOOLS_API static float GStitchingTolerance;
 		CADTOOLS_API static float GUnitScale;
 		CADTOOLS_API static float GMeshingParameterFactor;
