Fix Datasmith retessellate

diff --git a/Engine/Plugins/Enterprise/DatasmithCADImporter/Source/CADInterfaces/Private/CADSceneGraph.cpp b/Engine/Plugins/Enterprise/DatasmithCADImporter/Source/CADInterfaces/Private/CADSceneGraph.cpp
--- a/Engine/Plugins/Enterprise/DatasmithCADImporter/Source/CADInterfaces/Private/CADSceneGraph.cpp
+++ b/Engine/Plugins/Enterprise/DatasmithCADImporter/Source/CADInterfaces/Private/CADSceneGraph.cpp
@@ -63,7 +63,6 @@ FArchive& operator<<(FArchive& Ar, FArchiveBody& Body)
 	Ar << Body.ColorFaceSet;
 	Ar << Body.ParentId;
 	Ar << Body.MeshActorUId;
-	Ar << Body.BodyUnit;
 
 	return Ar;
 }
diff --git a/Engine/Plugins/Enterprise/DatasmithCADImporter/Source/CADInterfaces/Private/TechSoftFileParser.cpp b/Engine/Plugins/Enterprise/DatasmithCADImporter/Source/CADInterfaces/Private/TechSoftFileParser.cpp
--- a/Engine/Plugins/Enterprise/DatasmithCADImporter/Source/CADInterfaces/Private/TechSoftFileParser.cpp
+++ b/Engine/Plugins/Enterprise/DatasmithCADImporter/Source/CADInterfaces/Private/TechSoftFileParser.cpp
@@ -525,7 +525,7 @@ void FTechSoftFileParser::GenerateBodyMesh(A3DRiRepresentationItem* Representati
 
 			// Save body unit and default color and material attributes in a json string
 			// This will be used when the file is reloaded
-			JsonObject->SetNumberField(JSON_ENTRY_BODY_UNIT, Body.BodyUnit);
+			JsonObject->SetNumberField(JSON_ENTRY_BODY_UNIT, Body.Unit);
 
 			if(Body.ColorUId)
 			{
diff --git a/Engine/Plugins/Enterprise/DatasmithCADImporter/Source/CADInterfaces/Public/CADSceneGraph.h b/Engine/Plugins/Enterprise/DatasmithCADImporter/Source/CADInterfaces/Public/CADSceneGraph.h
--- a/Engine/Plugins/Enterprise/DatasmithCADImporter/Source/CADInterfaces/Public/CADSceneGraph.h
+++ b/Engine/Plugins/Enterprise/DatasmithCADImporter/Source/CADInterfaces/Public/CADSceneGraph.h
@@ -123,7 +123,6 @@ public:
 public:
 	FCadId ParentId = 0;
 	FCadUuid MeshActorUId = 0;
-	double BodyUnit = 1.;
 
 	TSet<FMaterialUId> MaterialFaceSet;
 	TSet<FMaterialUId> ColorFaceSet;
