[NO NON-PII COMMIT MESSAGE]

diff --git a/Engine/Plugins/Enterprise/DatasmithCADImporter/Source/CADInterfaces/Public/CADSceneGraph.h b/Engine/Plugins/Enterprise/DatasmithCADImporter/Source/CADInterfaces/Public/CADSceneGraph.h
--- a/Engine/Plugins/Enterprise/DatasmithCADImporter/Source/CADInterfaces/Public/CADSceneGraph.h
+++ b/Engine/Plugins/Enterprise/DatasmithCADImporter/Source/CADInterfaces/Public/CADSceneGraph.h
@@ -130,6 +130,7 @@ public:
 	void Delete();
 
 	// non serialized fields
+	bool bIsFromCad = true;
 	bool bIsASolid = true;
 };
 
diff --git a/Engine/Plugins/Enterprise/DatasmithCADImporter/Source/CADLibrary/Private/MeshDescriptionHelper.cpp b/Engine/Plugins/Enterprise/DatasmithCADImporter/Source/CADLibrary/Private/MeshDescriptionHelper.cpp
--- a/Engine/Plugins/Enterprise/DatasmithCADImporter/Source/CADLibrary/Private/MeshDescriptionHelper.cpp
+++ b/Engine/Plugins/Enterprise/DatasmithCADImporter/Source/CADLibrary/Private/MeshDescriptionHelper.cpp
@@ -19,35 +19,27 @@ typedef uint32 TriangleIndex[3];
 namespace CADLibrary
 {
 
-struct FVertexData
+template<typename VectorType>
+struct TGeometricData
 {
-	float Z;
+	float W;
 	int32 Index;
-	FVector3f Coordinates;
+	VectorType Coordinates;
 	bool bIsMerged;
-	FVertexID VertexID;
-	FVertexID SymVertexID;
 
 	/** Default constructor. */
-	FVertexData() {}
+	TGeometricData() {}
 
 	/** Initialization constructor. */
-	FVertexData(int32 InIndex, const FVector3f& V)
+	TGeometricData(int32 InIndex, const VectorType& V, const VectorType& OneVector)
 	{
-		//Z = V.X + V.Y + V.Z;
-		Z = 0.30f * V.X + 0.33f * V.Y + 0.37f * V.Z;
+		W = V | OneVector;  // V.X + V.Y or V.X + V.Y + V.Z according to the vector dimension
 		Index = InIndex;
 		Coordinates = V;
 		bIsMerged = false;
-		VertexID = INDEX_NONE;
 	}
 };
 
-struct FCompareVertexZ
-{
-	FORCEINLINE bool operator()(FVertexData const& A, FVertexData const& B) const { return A.Z < B.Z; }
-};
-
 // Verify the 3 input indices are not defining a degenerated triangle and fill up the corresponding FVertexIDs
 bool IsTriangleDegenerated(const int32_t* Indices, const TArray<FVertexID>& RemapVertexPosition, FVertexID VertexIDs[3])
 {
@@ -64,52 +56,49 @@ bool IsTriangleDegenerated(const int32_t* Indices, const TArray<FVertexID>& Rema
 	return (VertexIDs[0] == VertexIDs[1] || VertexIDs[0] == VertexIDs[2] || VertexIDs[1] == VertexIDs[2]);
 }
 
-void MergeCoincidentVertices(TArray<FVector3f>& VertexArray, TArray<int32>& VertexIdSet)
+template<typename VectorType>
+void MergeCoincidents(const TArray<VectorType>& DataArray, const VectorType& OneVector, const double CoincidenceTolerance, TArray<int32>& InOutIndices)
 {
-	TRACE_CPUPROFILER_EVENT_SCOPE(FDatasmithMeshBuilder::MergeCoincidentVertices);
+	// Create a list of Data W/index pairs
+	TArray<TGeometricData<VectorType>> GeometicDataArray;
+	GeometicDataArray.Reserve(DataArray.Num());
 
-	const double CoincidenceTolerance = 0.001;
-
-	// Create a list of vertex Z/index pairs
-	TArray<FVertexData> VertexDataSet;
-	VertexDataSet.Reserve(VertexArray.Num());
-
-	const FVector3f* Position = VertexArray.GetData();
-	const int32* VertexId = VertexIdSet.GetData();
-	for (int32 Index = 0; Index < VertexArray.Num(); ++Index, ++Position, ++VertexId)
+	const VectorType* Data = DataArray.GetData();
+	const int32* DataId = InOutIndices.GetData();
+	for (int32 Index = 0; Index < DataArray.Num(); ++Index, ++Data, ++DataId)
 	{
-		VertexDataSet.Emplace(*VertexId, *Position);
+		GeometicDataArray.Emplace(*DataId, *Data, OneVector);
 	}
 
 	// Sort the vertices by z value
-	VertexDataSet.Sort(FCompareVertexZ());
+	GeometicDataArray.Sort([](TGeometricData<VectorType> const& A, TGeometricData<VectorType> const& B) { return A.W < B.W; });
 
 	// Search for duplicates
-	for (int32 Index = 0; Index < VertexDataSet.Num(); Index++)
+	for (int32 Index = 0; Index < GeometicDataArray.Num(); Index++)
 	{
-		if (VertexDataSet[Index].bIsMerged)
+		if (GeometicDataArray[Index].bIsMerged)
 		{
 			continue;
 		}
 
-		VertexDataSet[Index].bIsMerged = true;
-		int32 NewIndex = VertexDataSet[Index].Index;
+		GeometicDataArray[Index].bIsMerged = true;
+		int32 NewIndex = GeometicDataArray[Index].Index;
 
-		const FVector3f& PositionA = VertexDataSet[Index].Coordinates;
+		const VectorType& PositionA = GeometicDataArray[Index].Coordinates;
 
 		// only need to search forward, since we add pairs both ways
-		for (int32 Andex = Index + 1; Andex < VertexDataSet.Num(); Andex++)
+		for (int32 Bndex = Index + 1; Bndex < GeometicDataArray.Num(); Bndex++)
 		{
-			if (FMath::Abs(VertexDataSet[Andex].Z - VertexDataSet[Index].Z) > 3. * CoincidenceTolerance)
+			if ((GeometicDataArray[Bndex].W - GeometicDataArray[Index].W) > CoincidenceTolerance)
 			{
 				break; // can't be any more duplicated
 			}
 
-			const FVector3f& PositionB = VertexDataSet[Andex].Coordinates;
+			const VectorType& PositionB = GeometicDataArray[Bndex].Coordinates;
 			if (PositionA.Equals(PositionB, CoincidenceTolerance))
 			{
-				VertexDataSet[Andex].bIsMerged = true;
-				VertexIdSet[VertexDataSet[Andex].Index] = NewIndex;
+				GeometicDataArray[Bndex].bIsMerged = true;
+				InOutIndices[GeometicDataArray[Bndex].Index] = NewIndex;
 			}
 		}
 	}
@@ -148,7 +137,8 @@ void FillVertexPosition(const FImportParameters& ImportParams, const FMeshParame
 		VertexIdSet[VertexIndex] = VertexID;
 	}
 
-	MergeCoincidentVertices(VertexArray, VertexIdSet);
+	const float GeometricTolerance = 0.001f; // cm
+	MergeCoincidents(VertexArray, FVector3f::OneVector, GeometricTolerance, VertexIdSet);
 
 	// if Symmetric mesh, the symmetric side of the mesh have to be generated
 	if (MeshParameters.bIsSymmetric)
@@ -273,7 +263,6 @@ bool FillMesh(const FMeshParameters& MeshParameters, const FImportParameters& Im
 	TriangleVertexInstanceIDs.SetNum(VertexCountPerFace);
 
 	TArray<FVertexInstanceID> MeshVertexInstanceIDs;
-	TArray<uint32> NewFaceIndex;  // new CT face index to remove degenerated face
 
 	// Gather all array data
 	FStaticMeshAttributes Attributes(MeshDescription);
@@ -349,45 +338,142 @@ bool FillMesh(const FMeshParameters& MeshParameters, const FImportParameters& Im
 				continue;
 			}
 
+			if (!Step)
+			{
+				FDatasmithUtils::ConvertVectorArray(ImportParams.GetModelCoordSys(), Tessellation.NormalArray);
+				for (FVector3f& Normal : Tessellation.NormalArray)
+				{
+					Normal = Normal.GetSafeNormal();
+				}
+			}
+
+			TArray<int32> UVIndices;
+			TArray<int32> NormalIndices;
+			TMap<TTuple<FVertexID, int32, int32>, FVertexInstanceID> VertexIDToInstanceIDForMesh;
+			TFunction<FVertexInstanceID(FVertexID, int32)> FindOrAddVertexInstanceIDForMesh = [&](FVertexID VertexID, int32 VertexIndex) ->FVertexInstanceID
+			{
+				int32 NormalIndex = NormalIndices[VertexIndex];
+				int32 UVIndex = UVIndices[VertexIndex];
+
+				FVertexInstanceID& VertexInstanceID = VertexIDToInstanceIDForMesh.FindOrAdd(TTuple<FVertexID, int32, int32>(VertexID, NormalIndices[VertexIndex], UVIndices[VertexIndex]));
+				if (VertexInstanceID == -1)
+				{
+					VertexInstanceID = MeshDescription.CreateVertexInstance(VertexID);
+					VertexInstanceColors[VertexInstanceID] = FLinearColor::White;
+					VertexInstanceTangents[VertexInstanceID] = FVector3f(ForceInitToZero);
+					VertexInstanceBinormalSigns[VertexInstanceID] = 0.0f;
+
+					MeshVertexInstanceIDs.Add(VertexInstanceID);
+					VertexInstanceNormals[VertexInstanceID] = (FVector3f)Tessellation.NormalArray[NormalIndex];
+					if (!Tessellation.TexCoordArray.IsEmpty())
+					{
+						VertexInstanceUVs.Set(VertexInstanceID, UVChannel, FVector2f(Tessellation.TexCoordArray[UVIndex]));
+					}
+				}
+				return VertexInstanceID;
+			};
+
+			TMap<FVertexID, FVertexInstanceID> VertexIDToInstanceIDForCad;
+			TFunction<FVertexInstanceID(FVertexID, int32)> FindOrAddVertexInstanceIDForCad = [&](FVertexID VertexID, int32 VertexIndex) ->FVertexInstanceID
+			{
+				FVertexInstanceID& VertexInstanceID = VertexIDToInstanceIDForCad.FindOrAdd(VertexID);
+				if (VertexInstanceID == -1)
+				{
+					VertexInstanceID = MeshDescription.CreateVertexInstance(VertexID);
+					VertexInstanceColors[VertexInstanceID] = FLinearColor::White;
+					VertexInstanceTangents[VertexInstanceID] = FVector3f(ForceInitToZero);
+					VertexInstanceBinormalSigns[VertexInstanceID] = 0.0f;
+
+					MeshVertexInstanceIDs.Add(VertexInstanceID);
+					VertexInstanceNormals[VertexInstanceID] = (FVector3f)Tessellation.NormalArray[VertexIndex];
+					if (!Tessellation.TexCoordArray.IsEmpty())
+					{
+						VertexInstanceUVs.Set(VertexInstanceID, UVChannel, FVector2f(Tessellation.TexCoordArray[VertexIndex]));
+					}
+				}
+				return VertexInstanceID;
+			};
+
+			TFunction<FVertexInstanceID(FVertexID, int32)> FindOrAddVertexInstanceID = BodyTessellation.bIsFromCad ? FindOrAddVertexInstanceIDForCad : FindOrAddVertexInstanceIDForMesh;
+
+
+			if (BodyTessellation.bIsFromCad)
+			{
+				VertexIDToInstanceIDForMesh.Reserve(Tessellation.VertexIndices.Num());
+			}
+			else
+			{
+				VertexIDToInstanceIDForCad.Reserve(Tessellation.VertexIndices.Num());
+
+				if (Tessellation.TexCoordArray.IsEmpty())
+				{
+					UVIndices.Init(0, Tessellation.VertexIndices.Num());
+				}
+				else
+				{
+					UVIndices.SetNum(Tessellation.VertexIndices.Num());
+					for (int32 Index = 0; Index < Tessellation.VertexIndices.Num(); ++Index)
+					{
+						UVIndices[Index] = Index;
+					}
+					const FVector2f OneVector(1.f, 1.f);
+					MergeCoincidents(Tessellation.TexCoordArray, OneVector, KINDA_SMALL_NUMBER, UVIndices);
+				}
+
+				if (Tessellation.NormalArray.Num() == 1)
+				{
+					NormalIndices.Init(0, Tessellation.VertexIndices.Num());
+				}
+				else
+				{
+					NormalIndices.SetNum(Tessellation.VertexIndices.Num());
+					for (int32 Index = 0; Index < Tessellation.VertexIndices.Num(); ++Index)
+					{
+						NormalIndices[Index] = Index;
+					}
+					MergeCoincidents(Tessellation.NormalArray, FVector3f::OneVector, KINDA_SMALL_NUMBER, NormalIndices);
+				}
+			}
+
 			int32 FaceVertexIDs[3];
 			int32 FaceVertexIndices[3];
+			int32 FaceVertexPositionIndices[3];
 			FVector Temp3D = { 0, 0, 0 };
 			FVector2D TexCoord2D = { 0, 0 };
 
 			MeshVertexInstanceIDs.Empty(Tessellation.VertexIndices.Num());
-			NewFaceIndex.Empty(Tessellation.VertexIndices.Num());
 
 			PatchIndex++;
 
 			// build each valid face i.e. 3 different indexes
 			for (int32 FaceIndex = 0; FaceIndex < Tessellation.VertexIndices.Num(); FaceIndex += VertexCountPerFace)
 			{
-				FaceVertexIndices[0] = Tessellation.PositionIndices[Tessellation.VertexIndices[FaceIndex + Orientation[0]]];
-				FaceVertexIndices[1] = Tessellation.PositionIndices[Tessellation.VertexIndices[FaceIndex + Orientation[1]]];
-				FaceVertexIndices[2] = Tessellation.PositionIndices[Tessellation.VertexIndices[FaceIndex + Orientation[2]]];
+				FaceVertexIndices[0] = Tessellation.VertexIndices[FaceIndex + Orientation[0]];
+				FaceVertexIndices[1] = Tessellation.VertexIndices[FaceIndex + Orientation[1]];
+				FaceVertexIndices[2] = Tessellation.VertexIndices[FaceIndex + Orientation[2]];
+
+				FaceVertexPositionIndices[0] = Tessellation.PositionIndices[FaceVertexIndices[0]];
+				FaceVertexPositionIndices[1] = Tessellation.PositionIndices[FaceVertexIndices[1]];
+				FaceVertexPositionIndices[2] = Tessellation.PositionIndices[FaceVertexIndices[2]];
 
-				if (FaceVertexIndices[0] == INDEX_NONE || FaceVertexIndices[1] == INDEX_NONE || FaceVertexIndices[2] == INDEX_NONE)
+				if (FaceVertexPositionIndices[0] == INDEX_NONE || FaceVertexPositionIndices[1] == INDEX_NONE || FaceVertexPositionIndices[2] == INDEX_NONE)
 				{
 					continue;
 				}
 
+				FaceVertexIDs[0] = (FVertexID) VertexIdSet[FaceVertexPositionIndices[0]];
+				FaceVertexIDs[1] = (FVertexID) VertexIdSet[FaceVertexPositionIndices[1]];
+				FaceVertexIDs[2] = (FVertexID) VertexIdSet[FaceVertexPositionIndices[2]];
+
 				// Verify the 3 input indices are not defining a degenerated triangle
-				if (FaceVertexIndices[0] == FaceVertexIndices[1] || FaceVertexIndices[0] == FaceVertexIndices[2] || FaceVertexIndices[1] == FaceVertexIndices[2])
+				if (FaceVertexIDs[0] == FaceVertexIDs[1] || FaceVertexIDs[0] == FaceVertexIDs[2] || FaceVertexIDs[1] == FaceVertexIDs[2])
 				{
 					continue;
 				}
 
-				FaceVertexIDs[0] = VertexIdSet[FaceVertexIndices[0]];
-				FaceVertexIDs[1] = VertexIdSet[FaceVertexIndices[1]];
-				FaceVertexIDs[2] = VertexIdSet[FaceVertexIndices[2]];
-
-				NewFaceIndex.Add(Tessellation.VertexIndices[FaceIndex + Orientation[0]]);
-				NewFaceIndex.Add(Tessellation.VertexIndices[FaceIndex + Orientation[1]]);
-				NewFaceIndex.Add(Tessellation.VertexIndices[FaceIndex + Orientation[2]]);
-
-				MeshVertexInstanceIDs.Add(TriangleVertexInstanceIDs[0] = MeshDescription.CreateVertexInstance((FVertexID)FaceVertexIDs[0]));
-				MeshVertexInstanceIDs.Add(TriangleVertexInstanceIDs[1] = MeshDescription.CreateVertexInstance((FVertexID)FaceVertexIDs[1]));
-				MeshVertexInstanceIDs.Add(TriangleVertexInstanceIDs[2] = MeshDescription.CreateVertexInstance((FVertexID)FaceVertexIDs[2]));
+				TriangleVertexInstanceIDs[0] = FindOrAddVertexInstanceID(FaceVertexIDs[0], FaceVertexIndices[0]);
+				TriangleVertexInstanceIDs[1] = FindOrAddVertexInstanceID(FaceVertexIDs[1], FaceVertexIndices[1]);
+				TriangleVertexInstanceIDs[2] = FindOrAddVertexInstanceID(FaceVertexIDs[2], FaceVertexIndices[2]);
 
 				if(FImportParameters::bGRemoveDuplicatedTriangle)
 				{
@@ -401,62 +487,12 @@ bool FillMesh(const FMeshParameters& MeshParameters, const FImportParameters& Im
 				}
 
 				// Add the triangle as a polygon to the mesh description
-				FPolygonID PolygonID = MeshDescription.CreatePolygon(*PolygonGroupID, TriangleVertexInstanceIDs);
+				const FPolygonID PolygonID = MeshDescription.CreatePolygon(*PolygonGroupID, TriangleVertexInstanceIDs);
 
 				// Set patch id attribute
 				PatchGroups[PolygonID] = Tessellation.PatchId;
 			}
 
-			if (!Tessellation.TexCoordArray.IsEmpty())
-			{
-				for (int32 FaceIndex = 0; FaceIndex < MeshVertexInstanceIDs.Num(); FaceIndex += VertexCountPerFace)
-				{
-					for (int32 VertexIndex = 0; VertexIndex < VertexCountPerFace; VertexIndex++)
-					{
-						const FVertexInstanceID& VertexInstanceID = MeshVertexInstanceIDs[FaceIndex + Orientation[VertexIndex]];
-						VertexInstanceUVs.Set(VertexInstanceID, UVChannel, FVector2f(Tessellation.TexCoordArray[NewFaceIndex[FaceIndex + Orientation[VertexIndex]]]));
-					}
-				}
-			}
-
-			for (const FVertexInstanceID& VertexInstanceID : MeshVertexInstanceIDs)
-			{
-				VertexInstanceColors[VertexInstanceID] = FLinearColor::White;
-				VertexInstanceTangents[VertexInstanceID] = FVector3f(ForceInitToZero);
-				VertexInstanceBinormalSigns[VertexInstanceID] = 0.0f;
-			}
-
-			if (!Step)
-			{
-				FDatasmithUtils::ConvertVectorArray(ImportParams.GetModelCoordSys(), Tessellation.NormalArray);
-				for (FVector3f& Normal : Tessellation.NormalArray)
-				{
-					Normal = Normal.GetSafeNormal();
-				}
-			}
-
-			if (Tessellation.NormalArray.Num() == 1)
-			{
-				const FVector3f& Normal = Tessellation.NormalArray[0];
-				for (int32 Index = 0; Index < NewFaceIndex.Num(); Index++)
-				{
-					const FVertexInstanceID VertexInstanceID = MeshVertexInstanceIDs[Index];
-					VertexInstanceNormals[VertexInstanceID] = (FVector3f)Normal;
-				}
-			}
-			else
-			{
-				for (int32 IndexFace = 0; IndexFace < NewFaceIndex.Num(); IndexFace += 3)
-				{
-					for (int32 Index = 0; Index < 3; Index++)
-					{
-						const FVertexInstanceID VertexInstanceID = MeshVertexInstanceIDs[IndexFace + Orientation[Index]];
-						VertexInstanceNormals[VertexInstanceID] = Tessellation.NormalArray[NewFaceIndex[IndexFace + Orientation[Index]]];
-					}
-				}
-			}
-
-			// compute normals
 			if (Step)
 			{
 				// compute normals of Symmetric vertex
diff --git a/Engine/Plugins/Enterprise/DatasmithCADImporter/Source/CADTools/Private/CADData.cpp b/Engine/Plugins/Enterprise/DatasmithCADImporter/Source/CADTools/Private/CADData.cpp
--- a/Engine/Plugins/Enterprise/DatasmithCADImporter/Source/CADTools/Private/CADData.cpp
+++ b/Engine/Plugins/Enterprise/DatasmithCADImporter/Source/CADTools/Private/CADData.cpp
@@ -95,9 +95,10 @@ FArchive& operator<<(FArchive& Ar, FTessellationData& TessellationData)
 
 FArchive& operator<<(FArchive& Ar, FBodyMesh& BodyMesh)
 {
+	Ar << BodyMesh.bIsFromCad;
+
 	Ar << BodyMesh.VertexArray;
 	Ar << BodyMesh.Faces;
-	Ar << BodyMesh.BBox;
 
 	Ar << BodyMesh.TriangleCount;
 	Ar << BodyMesh.BodyID;
diff --git a/Engine/Plugins/Enterprise/DatasmithCADImporter/Source/CADTools/Private/CADToolsModule.cpp b/Engine/Plugins/Enterprise/DatasmithCADImporter/Source/CADTools/Private/CADToolsModule.cpp
--- a/Engine/Plugins/Enterprise/DatasmithCADImporter/Source/CADTools/Private/CADToolsModule.cpp
+++ b/Engine/Plugins/Enterprise/DatasmithCADImporter/Source/CADTools/Private/CADToolsModule.cpp
@@ -18,7 +18,7 @@ bool FCADToolsModule::IsAvailable()
 
 uint32 FCADToolsModule::GetCacheVersion()
 {
-	return 502;
+	return 503;
 }
 
 IMPLEMENT_MODULE(FCADToolsModule, CADTools);
diff --git a/Engine/Plugins/Enterprise/DatasmithCADImporter/Source/CADTools/Public/CADData.h b/Engine/Plugins/Enterprise/DatasmithCADImporter/Source/CADTools/Public/CADData.h
--- a/Engine/Plugins/Enterprise/DatasmithCADImporter/Source/CADTools/Public/CADData.h
+++ b/Engine/Plugins/Enterprise/DatasmithCADImporter/Source/CADTools/Public/CADData.h
@@ -317,7 +317,6 @@ class CADTOOLS_API FBodyMesh
 public:
 	FBodyMesh(FCadId InBodyID = 0) : BodyID(InBodyID)
 	{
-		BBox.Init();
 	}
 
 	friend FArchive& operator<<(FArchive& Ar, FBodyMesh& BodyMesh);
@@ -336,9 +335,10 @@ public:
 		}
 	}
 
+	bool bIsFromCad = true;
+
 	TArray<FVector3f> VertexArray; // set by CADKernel, filled by FillKioVertexPosition that merges coincident vertices (CoreTechHelper)
 	TArray<FTessellationData> Faces;
-	FBox BBox;
 
 	uint32 TriangleCount = 0;
 	FCadId BodyID = 0;
