[NO NON-PII COMMIT MESSAGE]

diff --git a/Engine/Plugins/Enterprise/DatasmithCADImporter/Source/CADInterfaces/Public/CADFileData.h b/Engine/Plugins/Enterprise/DatasmithCADImporter/Source/CADInterfaces/Public/CADFileData.h
--- a/Engine/Plugins/Enterprise/DatasmithCADImporter/Source/CADInterfaces/Public/CADFileData.h
+++ b/Engine/Plugins/Enterprise/DatasmithCADImporter/Source/CADInterfaces/Public/CADFileData.h
@@ -141,6 +141,7 @@ public:
 	{
 		FBodyMesh& BodyMesh = BodyMeshes.Emplace_GetRef(BodyId);
 		BodyMesh.MeshActorUId = GetStaticMeshHash(BodyId);
+		BodyMesh.bIsFromCad = Body.bIsFromCad;
 		Body.MeshActorUId = BodyMesh.MeshActorUId;
 		return BodyMesh;
 	}
