CAD Import, Improve meta attribut name

diff --git a/Engine/Plugins/Enterprise/DatasmithCADImporter/Source/CADInterfaces/Private/CADFileReader.cpp b/Engine/Plugins/Enterprise/DatasmithCADImporter/Source/CADInterfaces/Private/CADFileReader.cpp
--- a/Engine/Plugins/Enterprise/DatasmithCADImporter/Source/CADInterfaces/Private/CADFileReader.cpp
+++ b/Engine/Plugins/Enterprise/DatasmithCADImporter/Source/CADInterfaces/Private/CADFileReader.cpp
@@ -117,7 +117,7 @@ namespace CADLibrary
 				{
 					bNeedToProceed = false;
 				}
-				else // the file has been converted into CT file but meshed with different parameters
+				else // the file has been converted into TechSoft file (".prc") but meshed with different parameters
 				{
 					CADFileData.GetCADFileDescription().SetCacheFile(CADFileCachePath);
 				}
diff --git a/Engine/Plugins/Enterprise/DatasmithCADImporter/Source/CADInterfaces/Private/TechSoftFileParser.cpp b/Engine/Plugins/Enterprise/DatasmithCADImporter/Source/CADInterfaces/Private/TechSoftFileParser.cpp
--- a/Engine/Plugins/Enterprise/DatasmithCADImporter/Source/CADInterfaces/Private/TechSoftFileParser.cpp
+++ b/Engine/Plugins/Enterprise/DatasmithCADImporter/Source/CADInterfaces/Private/TechSoftFileParser.cpp
@@ -139,20 +139,33 @@ inline void RemoveUnwantedChar(FString& StringToClean, TCHAR UnwantedChar)
 
 void TraverseAttribute(const A3DMiscAttributeData& AttributeData, TMap<FString, FString>& OutMetaData)
 {
-	FString AttributeName;
+	FString AttributeFamillyName;
 	if (AttributeData.m_bTitleIsInt)
 	{
 		A3DUns32 UnsignedValue = 0;
 		memcpy(&UnsignedValue, AttributeData.m_pcTitle, sizeof(A3DUns32));
-		AttributeName = FString::Printf(TEXT("%u"), UnsignedValue);
+		AttributeFamillyName = FString::Printf(TEXT("%u"), UnsignedValue);
 	}
 	else if (AttributeData.m_pcTitle && AttributeData.m_pcTitle[0] != '\0')
 	{
-		AttributeName = UTF8_TO_TCHAR(AttributeData.m_pcTitle);
+		AttributeFamillyName = UTF8_TO_TCHAR(AttributeData.m_pcTitle);
 	}
 
 	for (A3DUns32 Index = 0; Index < AttributeData.m_uiSize; ++Index)
 	{
+		FString AttributeName = AttributeFamillyName;
+		{
+			FString AttributeTitle = UTF8_TO_TCHAR(AttributeData.m_asSingleAttributesData[Index].m_pcTitle);
+			if (AttributeTitle.Len())
+			{
+				AttributeName += TEXT(" ") + AttributeTitle;
+			}
+			else if(Index > 0)
+			{
+				AttributeName += TEXT(" ") + FString::FromInt((int32)Index);
+			}
+		}
+
 		FString AttributeValue;
 		switch (AttributeData.m_asSingleAttributesData[Index].m_eType)
 		{
@@ -188,14 +201,7 @@ void TraverseAttribute(const A3DMiscAttributeData& AttributeData, TMap<FString,
 
 		if (AttributeName.Len())
 		{
-			if (Index)
-			{
-				OutMetaData.Emplace(FString::Printf(TEXT("%s_%u"), *AttributeName, (int32)Index), AttributeValue);
-			}
-			else
-			{
-				OutMetaData.Emplace(AttributeName, AttributeValue);
-			}
+			OutMetaData.Emplace(AttributeName, AttributeValue);
 		}
 	}
 }
diff --git a/Engine/Plugins/Enterprise/DatasmithCADImporter/Source/DatasmithCADTranslator/Private/DatasmithSceneGraphBuilder.cpp b/Engine/Plugins/Enterprise/DatasmithCADImporter/Source/DatasmithCADTranslator/Private/DatasmithSceneGraphBuilder.cpp
--- a/Engine/Plugins/Enterprise/DatasmithCADImporter/Source/DatasmithCADTranslator/Private/DatasmithSceneGraphBuilder.cpp
+++ b/Engine/Plugins/Enterprise/DatasmithCADImporter/Source/DatasmithCADTranslator/Private/DatasmithSceneGraphBuilder.cpp
@@ -234,24 +234,13 @@ bool FDatasmithSceneBaseGraphBuilder::Build()
 	// application_name is something like "Catia V5"
 	DatasmithScene->SetVendor(TEXT("Techsoft"));
 
-	if (const FString* ProductVersion = Reference.MetaData.Find(TEXT("TechsoftVersion")))
+	FString ProductVersion;
+	if (Reference.MetaData.RemoveAndCopyValue(TEXT("TechsoftVersion"), ProductVersion))
 	{
-		DatasmithScene->SetProductVersion(**ProductVersion);
+		DatasmithScene->SetProductVersion(*ProductVersion);
 	}
 
-	FString ProductName;
-	const FString* ProductNamePtr = Reference.MetaData.Find(TEXT("Input_Format_and_Emitter"));
-	if (ProductNamePtr)
-	{
-		ProductName = *ProductNamePtr;
-		ProductName.TrimStartAndEndInline();
-		if (!ProductName.IsEmpty())
-		{
-			DatasmithScene->SetProductName(*ProductName);
-		}
-	}
-
-	if (ProductName.IsEmpty())
+	// SetProductName
 	{
 		switch (RootFileDescription.GetFileFormat())
 		{
@@ -606,6 +595,7 @@ void FDatasmithSceneBaseGraphBuilder::AddMetaData(TSharedPtr<IDatasmithActorElem
 		}
 
 		// Add name
+		if(!Component.Label.IsEmpty())
 		{
 			FString MetaName = PostName;
 			MetaName += TEXT(" Name");
