Rename CADKernel namespace into UE::CADKernel

diff --git a/Engine/Plugins/Enterprise/DatasmithCADImporter/Source/CADInterfaces/Private/TechSoftBridge.cpp b/Engine/Plugins/Enterprise/DatasmithCADImporter/Source/CADInterfaces/Private/TechSoftBridge.cpp
--- a/Engine/Plugins/Enterprise/DatasmithCADImporter/Source/CADInterfaces/Private/TechSoftBridge.cpp
+++ b/Engine/Plugins/Enterprise/DatasmithCADImporter/Source/CADInterfaces/Private/TechSoftBridge.cpp
@@ -5,7 +5,6 @@
 #ifdef USE_TECHSOFT_SDK
 
 #include "CADSceneGraph.h"
-#include "CADFileReport.h"
 #include "TechSoftInterface.h"
 #include "TechSoftFileParser.h"
 #include "TechSoftUtils.h"
@@ -31,6 +30,10 @@
 
 #include "CADKernel/UI/Display.h"
 
+#ifdef CADKERNEL_DEV
+#include "CADFileReport.h"
+#endif
+
 namespace CADLibrary
 {
 namespace TechSoftUtils
@@ -48,11 +51,11 @@ A3DStatus GetSurfaceAsNurbs(const A3DSurfBase* A3DSurface, A3DSurfNurbsData* Dat
 	return TechSoftInterface::GetSurfaceAsNurbs(A3DSurface, DataPtr, Forward<InArgTypes>(Args)...);
 }; 
 
-CADKernel::FMatrixH CreateCoordinateSystem(const A3DMiscCartesianTransformationData& Transformation, double UnitScale = 1.0)
+UE::CADKernel::FMatrixH CreateCoordinateSystem(const A3DMiscCartesianTransformationData& Transformation, double UnitScale = 1.0)
 {
-	CADKernel::FPoint Origin(&Transformation.m_sOrigin.m_dX);
-	CADKernel::FPoint Ox(&Transformation.m_sXVector.m_dX);
-	CADKernel::FPoint Oy(&Transformation.m_sYVector.m_dX);
+	UE::CADKernel::FPoint Origin(&Transformation.m_sOrigin.m_dX);
+	UE::CADKernel::FPoint Ox(&Transformation.m_sXVector.m_dX);
+	UE::CADKernel::FPoint Oy(&Transformation.m_sYVector.m_dX);
 
 	Ox.Normalize();
 	Oy.Normalize();
@@ -61,13 +64,13 @@ CADKernel::FMatrixH CreateCoordinateSystem(const A3DMiscCartesianTransformationD
 	{
 		Origin *= UnitScale;
 	}
-	CADKernel::FPoint Oz = Ox ^ Oy;
+	UE::CADKernel::FPoint Oz = Ox ^ Oy;
 
-	CADKernel::FMatrixH Matrix = CADKernel::FMatrixH(Origin, Ox, Oy, Oz);
+	UE::CADKernel::FMatrixH Matrix = UE::CADKernel::FMatrixH(Origin, Ox, Oy, Oz);
 
 	if (!FMath::IsNearlyEqual(Transformation.m_sScale.m_dX, 1.) || !FMath::IsNearlyEqual(Transformation.m_sScale.m_dY, 1.) || !FMath::IsNearlyEqual(Transformation.m_sScale.m_dZ, 1.))
 	{
-		CADKernel::FMatrixH Scale = CADKernel::FMatrixH::MakeScaleMatrix(Transformation.m_sScale.m_dX, Transformation.m_sScale.m_dY, Transformation.m_sScale.m_dZ);
+		UE::CADKernel::FMatrixH Scale = UE::CADKernel::FMatrixH::MakeScaleMatrix(Transformation.m_sScale.m_dX, Transformation.m_sScale.m_dY, Transformation.m_sScale.m_dZ);
 		Matrix *= Scale;
 	}
 	return Matrix;
@@ -104,7 +107,7 @@ void FillDoubleArray(const int32 UCount, const int32 VCount, const double* Value
 	}
 }
 
-void FillPointArray(const int32 Count, const A3DVector3dData* Points, TArray<CADKernel::FPoint>& OutPointsArray, double UnitScale = 1.0)
+void FillPointArray(const int32 Count, const A3DVector3dData* Points, TArray<UE::CADKernel::FPoint>& OutPointsArray, double UnitScale = 1.0)
 {
 	OutPointsArray.Reserve(Count);
 	for (int32 Index = 0; Index < Count; Index++)
@@ -114,14 +117,14 @@ void FillPointArray(const int32 Count, const A3DVector3dData* Points, TArray<CAD
 
 	if (!FMath::IsNearlyEqual(UnitScale, 1.))
 	{
-		for (CADKernel::FPoint& Point : OutPointsArray)
+		for (UE::CADKernel::FPoint& Point : OutPointsArray)
 		{
 			Point *= UnitScale;
 		}
 	}
 };
 
-void FillPointArray(const int32 UCount, const int32 VCount, const A3DVector3dData* Points, TArray<CADKernel::FPoint>& OutPointsArray, double UnitScale = 1.0)
+void FillPointArray(const int32 UCount, const int32 VCount, const A3DVector3dData* Points, TArray<UE::CADKernel::FPoint>& OutPointsArray, double UnitScale = 1.0)
 {
 	OutPointsArray.SetNum(UCount * VCount);
 	for (int32 Undex = 0, PointIndex = 0; Undex < UCount; ++Undex)
@@ -135,18 +138,18 @@ void FillPointArray(const int32 UCount, const int32 VCount, const A3DVector3dDat
 
 	if (!FMath::IsNearlyEqual(UnitScale, 1.))
 	{
-		for (CADKernel::FPoint& Point : OutPointsArray)
+		for (UE::CADKernel::FPoint& Point : OutPointsArray)
 		{
 			Point *= UnitScale;
 		}
 	}
 };
 
-CADKernel::FSurfacicBoundary GetSurfacicBoundary(A3DDomainData& Domain, const TechSoftUtils::FUVReparameterization& UVReparameterization)
+UE::CADKernel::FSurfacicBoundary GetSurfacicBoundary(A3DDomainData& Domain, const TechSoftUtils::FUVReparameterization& UVReparameterization)
 {
 
-	CADKernel::FPoint2D Min(Domain.m_sMin.m_dX, Domain.m_sMin.m_dY);
-	CADKernel::FPoint2D Max(Domain.m_sMax.m_dX, Domain.m_sMax.m_dY);
+	UE::CADKernel::FPoint2D Min(Domain.m_sMin.m_dX, Domain.m_sMin.m_dY);
+	UE::CADKernel::FPoint2D Max(Domain.m_sMax.m_dX, Domain.m_sMax.m_dY);
 
 	if(UVReparameterization.GetNeedApply())
 	{
@@ -154,10 +157,10 @@ CADKernel::FSurfacicBoundary GetSurfacicBoundary(A3DDomainData& Domain, const Te
 		UVReparameterization.Apply(Max);
 	}
 
-	CADKernel::EIso UIndex = UVReparameterization.GetSwapUV() ? CADKernel::EIso::IsoV : CADKernel::EIso::IsoU;
-	CADKernel::EIso VIndex = UVReparameterization.GetSwapUV() ? CADKernel::EIso::IsoU : CADKernel::EIso::IsoV;
+	UE::CADKernel::EIso UIndex = UVReparameterization.GetSwapUV() ? UE::CADKernel::EIso::IsoV : UE::CADKernel::EIso::IsoU;
+	UE::CADKernel::EIso VIndex = UVReparameterization.GetSwapUV() ? UE::CADKernel::EIso::IsoU : UE::CADKernel::EIso::IsoV;
 
-	CADKernel::FSurfacicBoundary Boundary;
+	UE::CADKernel::FSurfacicBoundary Boundary;
 	Boundary[UIndex].Min = Min.U;
 	Boundary[VIndex].Min = Min.V;
 	Boundary[UIndex].Max = Max.U;
@@ -166,33 +169,32 @@ CADKernel::FSurfacicBoundary GetSurfacicBoundary(A3DDomainData& Domain, const Te
 	return Boundary;
 }
 
-CADKernel::FLinearBoundary GetLinearBoundary(A3DIntervalData& A3DDomain)
+UE::CADKernel::FLinearBoundary GetLinearBoundary(A3DIntervalData& A3DDomain)
 {
-	CADKernel::FLinearBoundary Domain(A3DDomain.m_dMin, A3DDomain.m_dMax);
+	UE::CADKernel::FLinearBoundary Domain(A3DDomain.m_dMin, A3DDomain.m_dMax);
 	return Domain;
 }
 
-CADKernel::FLinearBoundary GetLinearBoundary(const A3DCrvBase* A3DCurve)
+UE::CADKernel::FLinearBoundary GetLinearBoundary(const A3DCrvBase* A3DCurve)
 {
 	TUniqueTSObj<A3DIntervalData> A3DDomain(A3DCurve);
-	CADKernel::FLinearBoundary Domain(A3DDomain->m_dMin, A3DDomain->m_dMax);
+	UE::CADKernel::FLinearBoundary Domain(A3DDomain->m_dMin, A3DDomain->m_dMax);
 	return Domain;
 }
 
 } // ns TechSoftUtils
 
-FTechSoftBridge::FTechSoftBridge(FTechSoftFileParser& InParser, CADKernel::FSession& InSession, CADKernel::FCADFileReport& InReport)
+FTechSoftBridge::FTechSoftBridge(FTechSoftFileParser& InParser, UE::CADKernel::FSession& InSession)
 	: Parser(InParser)
 	, Session(InSession)
 	, Model(InSession.GetModel())
-	, Report(InReport)
 	, GeometricTolerance(Session.GetGeometricTolerance())
 	, SquareGeometricTolerance(FMath::Square(Session.GetGeometricTolerance()))
 	, SquareJoiningVertexTolerance(SquareGeometricTolerance * 2)
 {
 }
 
-const A3DRiBrepModel* FTechSoftBridge::GetA3DBody(CADKernel::FBody* BRepModel)
+const A3DRiBrepModel* FTechSoftBridge::GetA3DBody(UE::CADKernel::FBody* BRepModel)
 {
 	const A3DEntity** BodyPtr = CADKernelToTechSoft.Find(BRepModel);
 	if (BodyPtr)
@@ -202,9 +204,9 @@ const A3DRiBrepModel* FTechSoftBridge::GetA3DBody(CADKernel::FBody* BRepModel)
 	return nullptr;
 }
 
-CADKernel::FBody* FTechSoftBridge::GetBody(A3DRiBrepModel* A3DBRepModel)
+UE::CADKernel::FBody* FTechSoftBridge::GetBody(A3DRiBrepModel* A3DBRepModel)
 {
-	CADKernel::FBody** BodyPtr = TechSoftToCADKernel.Find(A3DBRepModel);
+	UE::CADKernel::FBody** BodyPtr = TechSoftToCADKernel.Find(A3DBRepModel);
 	if (BodyPtr &&!(*BodyPtr)->IsDeleted())
 	{
 		return *BodyPtr;
@@ -212,11 +214,13 @@ CADKernel::FBody* FTechSoftBridge::GetBody(A3DRiBrepModel* A3DBRepModel)
 	return nullptr;
 }
 
-CADKernel::FBody* FTechSoftBridge::AddBody(A3DRiBrepModel* A3DBRepModel, TMap<FString, FString> MetaData, const double InBodyScale)
+UE::CADKernel::FBody* FTechSoftBridge::AddBody(A3DRiBrepModel* A3DBRepModel, TMap<FString, FString> MetaData, const double InBodyScale)
 {
+#ifdef CADKERNEL_DEV
 	Report.BodyCount++;
+#endif
 
-	//CADKernel working unit is mm
+	//UE::CADKernel working unit is mm
 	BodyScale = InBodyScale * 10.;
 
 	FArchiveCADObject BRepMetaData;
@@ -227,7 +231,7 @@ CADKernel::FBody* FTechSoftBridge::AddBody(A3DRiBrepModel* A3DBRepModel, TMap<FS
 		BRepMetaData.MetaData.FindOrAdd(TEXT("Name")) = *Name;
 	}
 	
-	CADKernel::FBody** BodyPtr = TechSoftToCADKernel.Find(A3DBRepModel);
+	UE::CADKernel::FBody** BodyPtr = TechSoftToCADKernel.Find(A3DBRepModel);
 	if (BodyPtr)
 	{
 		if((*BodyPtr)->IsDeleted())
@@ -237,7 +241,7 @@ CADKernel::FBody* FTechSoftBridge::AddBody(A3DRiBrepModel* A3DBRepModel, TMap<FS
 		return *BodyPtr;
 	}
 
-	TSharedRef<CADKernel::FBody> Body = CADKernel::FEntity::MakeShared<CADKernel::FBody>();
+	TSharedRef<UE::CADKernel::FBody> Body = UE::CADKernel::FEntity::MakeShared<UE::CADKernel::FBody>();
 	AddMetadata(BRepMetaData, *Body);
 
 	TUniqueTSObj<A3DRiBrepModelData> BRepModelData(A3DBRepModel);
@@ -259,7 +263,7 @@ CADKernel::FBody* FTechSoftBridge::AddBody(A3DRiBrepModel* A3DBRepModel, TMap<FS
 	return &*Body;
 }
 
-void FTechSoftBridge::TraverseBrepData(const A3DTopoBrepData* A3DBrepData, TSharedRef<CADKernel::FBody>& Body)
+void FTechSoftBridge::TraverseBrepData(const A3DTopoBrepData* A3DBrepData, TSharedRef<UE::CADKernel::FBody>& Body)
 {
 	FArchiveCADObject MetaData;
 	Parser.ExtractMetaData(A3DBrepData, MetaData);
@@ -292,7 +296,7 @@ void FTechSoftBridge::TraverseBrepData(const A3DTopoBrepData* A3DBrepData, TShar
 	}
 }
 
-void FTechSoftBridge::TraverseConnex(const A3DTopoConnex* A3DTopoConnex, TSharedRef<CADKernel::FBody>& Body)
+void FTechSoftBridge::TraverseConnex(const A3DTopoConnex* A3DTopoConnex, TSharedRef<UE::CADKernel::FBody>& Body)
 {
 	FArchiveCADObject MetaData;
 	Parser.ExtractMetaData(A3DTopoConnex, MetaData);
@@ -307,15 +311,16 @@ void FTechSoftBridge::TraverseConnex(const A3DTopoConnex* A3DTopoConnex, TShared
 	}
 }
 
-void FTechSoftBridge::TraverseShell(const A3DTopoShell* A3DShell, TSharedRef<CADKernel::FBody>& Body)
+void FTechSoftBridge::TraverseShell(const A3DTopoShell* A3DShell, TSharedRef<UE::CADKernel::FBody>& Body)
 {
 	FArchiveCADObject MetaData;
 	Parser.ExtractMetaData(A3DShell, MetaData);
 
-	TSharedRef<CADKernel::FShell> Shell = CADKernel::FEntity::MakeShared<CADKernel::FShell>();
+	TSharedRef<UE::CADKernel::FShell> Shell = UE::CADKernel::FEntity::MakeShared<UE::CADKernel::FShell>();
 	Body->AddShell(Shell);
+#ifdef CADKERNEL_DEV
 	Report.ShellCount++;
-
+#endif
 	AddMetadata(MetaData, *Shell);
 
 	TUniqueTSObj<A3DTopoShellData> ShellData(A3DShell);
@@ -325,21 +330,22 @@ void FTechSoftBridge::TraverseShell(const A3DTopoShell* A3DShell, TSharedRef<CAD
 		A3DEdgeToEdge.Empty();
 		for (A3DUns32 Index = 0; Index < ShellData->m_uiFaceSize; ++Index)
 		{
-			AddFace(ShellData->m_ppFaces[Index], ShellData->m_pucOrientationWithShell[Index] == 1 ? CADKernel::Front : CADKernel::Back, Shell, Index);
+			AddFace(ShellData->m_ppFaces[Index], ShellData->m_pucOrientationWithShell[Index] == 1 ? UE::CADKernel::Front : UE::CADKernel::Back, Shell, Index);
 		}
 	}
 }
 
 static bool bUseCurveAsNurbs = true;
-TSharedPtr<CADKernel::FCurve> FTechSoftBridge::AddCurve(const A3DCrvBase* A3DCurve, const TechSoftUtils::FUVReparameterization& UVReparameterization)
+TSharedPtr<UE::CADKernel::FCurve> FTechSoftBridge::AddCurve(const A3DCrvBase* A3DCurve, const TechSoftUtils::FUVReparameterization& UVReparameterization)
 {
-	TSharedPtr<CADKernel::FCurve> Curve = TSharedPtr<CADKernel::FCurve>();
+	TSharedPtr<UE::CADKernel::FCurve> Curve = TSharedPtr<UE::CADKernel::FCurve>();
 	A3DEEntityType eType;
 	A3DInt32 Ret = TechSoftInterface::GetEntityType(A3DCurve, &eType);
 	if (Ret == A3D_SUCCESS)
 	{
+#ifdef CADKERNEL_DEV
 		Report.CurveCount++;
-
+#endif
 		switch (eType)
 		{
 		case kA3DTypeCrvNurbs:
@@ -375,14 +381,16 @@ TSharedPtr<CADKernel::FCurve> FTechSoftBridge::AddCurve(const A3DCrvBase* A3DCur
 		}
 	}
 
-	CADKernel::FLinearBoundary Boundary = TechSoftUtils::GetLinearBoundary(A3DCurve);
+	UE::CADKernel::FLinearBoundary Boundary = TechSoftUtils::GetLinearBoundary(A3DCurve);
 
 	return Curve;
 }
 
-TSharedPtr<CADKernel::FCurve> FTechSoftBridge::AddCurveLine(const A3DCrvLine* A3DCurve, const TechSoftUtils::FUVReparameterization& UVReparameterization)
+TSharedPtr<UE::CADKernel::FCurve> FTechSoftBridge::AddCurveLine(const A3DCrvLine* A3DCurve, const TechSoftUtils::FUVReparameterization& UVReparameterization)
 {
+#ifdef CADKERNEL_DEV
 	Report.CurveLineCount++;
+#endif
 
 	if (bUseCurveAsNurbs)
 	{
@@ -393,17 +401,19 @@ TSharedPtr<CADKernel::FCurve> FTechSoftBridge::AddCurveLine(const A3DCrvLine* A3
 	TUniqueTSObj<A3DCrvLineData> CrvLineData(A3DCurve);
 	if (!CrvLineData.IsValid())
 	{
-		return TSharedPtr<CADKernel::FCurve>();
+		return TSharedPtr<UE::CADKernel::FCurve>();
 	}
 
 	bool bIs2D = (bool)CrvLineData->m_bIs2D;
 	// Todo
-	return TSharedPtr<CADKernel::FCurve>();
+	return TSharedPtr<UE::CADKernel::FCurve>();
 }
 
-TSharedPtr<CADKernel::FCurve> FTechSoftBridge::AddCurveCircle(const A3DCrvCircle* A3DCurve, const TechSoftUtils::FUVReparameterization& UVReparameterization)
+TSharedPtr<UE::CADKernel::FCurve> FTechSoftBridge::AddCurveCircle(const A3DCrvCircle* A3DCurve, const TechSoftUtils::FUVReparameterization& UVReparameterization)
 {
+#ifdef CADKERNEL_DEV
 	Report.CurveCircleCount++;
+#endif
 
 	if (bUseCurveAsNurbs)
 	{
@@ -413,17 +423,19 @@ TSharedPtr<CADKernel::FCurve> FTechSoftBridge::AddCurveCircle(const A3DCrvCircle
 	TUniqueTSObj<A3DCrvCircleData> CrvCircleData(A3DCurve);
 	if (!CrvCircleData.IsValid())
 	{
-		return TSharedPtr<CADKernel::FCurve>();
+		return TSharedPtr<UE::CADKernel::FCurve>();
 	}
 
 	bool bIs2D = (bool)CrvCircleData->m_bIs2D;
 	// Todo
-	return TSharedPtr<CADKernel::FCurve>();
+	return TSharedPtr<UE::CADKernel::FCurve>();
 }
 
-TSharedPtr<CADKernel::FCurve> FTechSoftBridge::AddCurveEllipse(const A3DCrvEllipse* A3DCurve, const TechSoftUtils::FUVReparameterization& UVReparameterization)
+TSharedPtr<UE::CADKernel::FCurve> FTechSoftBridge::AddCurveEllipse(const A3DCrvEllipse* A3DCurve, const TechSoftUtils::FUVReparameterization& UVReparameterization)
 {
+#ifdef CADKERNEL_DEV
 	Report.CurveEllipseCount++;
+#endif
 
 	if (bUseCurveAsNurbs)
 	{
@@ -433,17 +445,19 @@ TSharedPtr<CADKernel::FCurve> FTechSoftBridge::AddCurveEllipse(const A3DCrvEllip
 	TUniqueTSObj<A3DCrvEllipseData> CrvEllipseData(A3DCurve);
 	if (!CrvEllipseData.IsValid())
 	{
-		return TSharedPtr<CADKernel::FCurve>();
+		return TSharedPtr<UE::CADKernel::FCurve>();
 	}
 
 	bool bIs2D = (bool)CrvEllipseData->m_bIs2D;
 	// Todo
-	return TSharedPtr<CADKernel::FCurve>();
+	return TSharedPtr<UE::CADKernel::FCurve>();
 }
 
-TSharedPtr<CADKernel::FCurve> FTechSoftBridge::AddCurveParabola(const A3DCrvParabola* A3DCurve, const TechSoftUtils::FUVReparameterization& UVReparameterization)
+TSharedPtr<UE::CADKernel::FCurve> FTechSoftBridge::AddCurveParabola(const A3DCrvParabola* A3DCurve, const TechSoftUtils::FUVReparameterization& UVReparameterization)
 {
+#ifdef CADKERNEL_DEV
 	Report.CurveParabolaCount++;
+#endif
 
 	if (bUseCurveAsNurbs)
 	{
@@ -453,17 +467,19 @@ TSharedPtr<CADKernel::FCurve> FTechSoftBridge::AddCurveParabola(const A3DCrvPara
 	TUniqueTSObj<A3DCrvParabolaData> CrvParabolaData(A3DCurve);
 	if (!CrvParabolaData.IsValid())
 	{
-		return TSharedPtr<CADKernel::FCurve>();
+		return TSharedPtr<UE::CADKernel::FCurve>();
 	}
 
 	bool bIs2D = (bool)CrvParabolaData->m_bIs2D;
 	// Todo
-	return TSharedPtr<CADKernel::FCurve>();
+	return TSharedPtr<UE::CADKernel::FCurve>();
 }
 
-TSharedPtr<CADKernel::FCurve> FTechSoftBridge::AddCurveHyperbola(const A3DCrvHyperbola* A3DCurve, const TechSoftUtils::FUVReparameterization& UVReparameterization)
+TSharedPtr<UE::CADKernel::FCurve> FTechSoftBridge::AddCurveHyperbola(const A3DCrvHyperbola* A3DCurve, const TechSoftUtils::FUVReparameterization& UVReparameterization)
 {
+#ifdef CADKERNEL_DEV
 	Report.CurveHyperbolaCount++;
+#endif
 
 	if (bUseCurveAsNurbs)
 	{
@@ -473,17 +489,19 @@ TSharedPtr<CADKernel::FCurve> FTechSoftBridge::AddCurveHyperbola(const A3DCrvHyp
 	TUniqueTSObj<A3DCrvHyperbolaData> CrvHyperbolaData(A3DCurve);
 	if (!CrvHyperbolaData.IsValid())
 	{
-		return TSharedPtr<CADKernel::FCurve>();
+		return TSharedPtr<UE::CADKernel::FCurve>();
 	}
 
 	bool bIs2D = (bool)CrvHyperbolaData->m_bIs2D;
 	// Todo
-	return TSharedPtr<CADKernel::FCurve>();
+	return TSharedPtr<UE::CADKernel::FCurve>();
 }
 
-TSharedPtr<CADKernel::FCurve> FTechSoftBridge::AddCurveHelix(const A3DCrvHelix* A3DCurve, const TechSoftUtils::FUVReparameterization& UVReparameterization)
+TSharedPtr<UE::CADKernel::FCurve> FTechSoftBridge::AddCurveHelix(const A3DCrvHelix* A3DCurve, const TechSoftUtils::FUVReparameterization& UVReparameterization)
 {
+#ifdef CADKERNEL_DEV
 	Report.CurveHelixCount++;
+#endif
 
 	if (bUseCurveAsNurbs)
 	{
@@ -493,17 +511,19 @@ TSharedPtr<CADKernel::FCurve> FTechSoftBridge::AddCurveHelix(const A3DCrvHelix*
 	TUniqueTSObj<A3DCrvHelixData> CrvHelixData(A3DCurve);
 	if (!CrvHelixData.IsValid())
 	{
-		return TSharedPtr<CADKernel::FCurve>();
+		return TSharedPtr<UE::CADKernel::FCurve>();
 	}
 
 	bool bIs2D = (bool)CrvHelixData->m_bIs2D;
 
-	return TSharedPtr<CADKernel::FCurve>();
+	return TSharedPtr<UE::CADKernel::FCurve>();
 }
 
-TSharedPtr<CADKernel::FCurve> FTechSoftBridge::AddCurvePolyLine(const A3DCrvPolyLine* A3DCurve, const TechSoftUtils::FUVReparameterization& UVReparameterization)
+TSharedPtr<UE::CADKernel::FCurve> FTechSoftBridge::AddCurvePolyLine(const A3DCrvPolyLine* A3DCurve, const TechSoftUtils::FUVReparameterization& UVReparameterization)
 {
+#ifdef CADKERNEL_DEV
 	Report.CurvePolyLineCount++;
+#endif
 
 	if (bUseCurveAsNurbs)
 	{
@@ -513,17 +533,19 @@ TSharedPtr<CADKernel::FCurve> FTechSoftBridge::AddCurvePolyLine(const A3DCrvPoly
 	TUniqueTSObj<A3DCrvPolyLineData> CrvPolyLineData(A3DCurve);
 	if (!CrvPolyLineData.IsValid())
 	{
-		return TSharedPtr<CADKernel::FCurve>();
+		return TSharedPtr<UE::CADKernel::FCurve>();
 	}
 
 	bool bIs2D = (bool)CrvPolyLineData->m_bIs2D;
 	// Todo
-	return TSharedPtr<CADKernel::FCurve>();
+	return TSharedPtr<UE::CADKernel::FCurve>();
 }
 
-TSharedPtr<CADKernel::FCurve> FTechSoftBridge::AddCurveComposite(const A3DCrvComposite* A3DCurve, const TechSoftUtils::FUVReparameterization& UVReparameterization)
+TSharedPtr<UE::CADKernel::FCurve> FTechSoftBridge::AddCurveComposite(const A3DCrvComposite* A3DCurve, const TechSoftUtils::FUVReparameterization& UVReparameterization)
 {
+#ifdef CADKERNEL_DEV
 	Report.CurveCompositeCount++;
+#endif
 
 	if (bUseCurveAsNurbs)
 	{
@@ -533,17 +555,17 @@ TSharedPtr<CADKernel::FCurve> FTechSoftBridge::AddCurveComposite(const A3DCrvCom
 	TUniqueTSObj<A3DCrvCompositeData> CrvCompositeData(A3DCurve);
 	if (!CrvCompositeData.IsValid())
 	{
-		return TSharedPtr<CADKernel::FCurve>();
+		return TSharedPtr<UE::CADKernel::FCurve>();
 	}
 
 	bool bIs2D = (bool)CrvCompositeData->m_bIs2D;
 	// Todo
-	return TSharedPtr<CADKernel::FCurve>();
+	return TSharedPtr<UE::CADKernel::FCurve>();
 }
 
-TSharedPtr<CADKernel::FCurve> AddCurveNurbsFromData(A3DCrvNurbsData& A3DNurbs, const TechSoftUtils::FUVReparameterization& UVReparameterization)
+TSharedPtr<UE::CADKernel::FCurve> AddCurveNurbsFromData(A3DCrvNurbsData& A3DNurbs, const TechSoftUtils::FUVReparameterization& UVReparameterization)
 {
-	CADKernel::FNurbsCurveData Nurbs;
+	UE::CADKernel::FNurbsCurveData Nurbs;
 	Nurbs.Dimension = A3DNurbs.m_bIs2D ? 2 : 3;
 	Nurbs.bIsRational = (bool)A3DNurbs.m_bRational;
 	Nurbs.Degree = A3DNurbs.m_uiDegree;
@@ -562,25 +584,29 @@ TSharedPtr<CADKernel::FCurve> AddCurveNurbsFromData(A3DCrvNurbsData& A3DNurbs, c
 
 	A3DCrvNurbsGet(NULL, &A3DNurbs);
 
-	return CADKernel::FCurve::MakeNurbsCurve(Nurbs);
+	return UE::CADKernel::FCurve::MakeNurbsCurve(Nurbs);
 }
 
-TSharedPtr<CADKernel::FCurve> FTechSoftBridge::AddCurveNurbs(const A3DCrvNurbs* A3DNurbs, const TechSoftUtils::FUVReparameterization& UVReparameterization)
+TSharedPtr<UE::CADKernel::FCurve> FTechSoftBridge::AddCurveNurbs(const A3DCrvNurbs* A3DNurbs, const TechSoftUtils::FUVReparameterization& UVReparameterization)
 {
+#ifdef CADKERNEL_DEV
 	Report.CurveLineCount++;
+#endif
 
 	TUniqueTSObj<A3DCrvNurbsData> CrvNurbsData(A3DNurbs);
 	if (!CrvNurbsData.IsValid())
 	{
-		return TSharedPtr<CADKernel::FCurve>();
+		return TSharedPtr<UE::CADKernel::FCurve>();
 	}
 
 	return AddCurveNurbsFromData(*CrvNurbsData, UVReparameterization);
 }
 
-TSharedPtr<CADKernel::FCurve> FTechSoftBridge::AddCurveAsNurbs(const A3DCrvBase* A3DCurve, const TechSoftUtils::FUVReparameterization& UVReparameterization)
+TSharedPtr<UE::CADKernel::FCurve> FTechSoftBridge::AddCurveAsNurbs(const A3DCrvBase* A3DCurve, const TechSoftUtils::FUVReparameterization& UVReparameterization)
 {
+#ifdef CADKERNEL_DEV
 	Report.CurveAsNurbsCount++;
+#endif
 
 	TUniqueTSObj<A3DCrvNurbsData> NurbsData;
 
@@ -590,44 +616,46 @@ TSharedPtr<CADKernel::FCurve> FTechSoftBridge::AddCurveAsNurbs(const A3DCrvBase*
 
 	if (!NurbsData.IsValid())
 	{
-		return TSharedPtr<CADKernel::FCurve>();
+		return TSharedPtr<UE::CADKernel::FCurve>();
 	}
 
 	return AddCurveNurbsFromData(*NurbsData, UVReparameterization);
 }
 
-TSharedPtr<CADKernel::FTopologicalEdge> FTechSoftBridge::AddEdge(const A3DTopoCoEdge* A3DCoedge, const TSharedRef<CADKernel::FSurface>& Surface, const TechSoftUtils::FUVReparameterization& UVReparameterization, CADKernel::EOrientation& OutOrientation)
+TSharedPtr<UE::CADKernel::FTopologicalEdge> FTechSoftBridge::AddEdge(const A3DTopoCoEdge* A3DCoedge, const TSharedRef<UE::CADKernel::FSurface>& Surface, const TechSoftUtils::FUVReparameterization& UVReparameterization, UE::CADKernel::EOrientation& OutOrientation)
 {
+#ifdef CADKERNEL_DEV
 	Report.EdgeCount++;
+#endif
 
 	TUniqueTSObj<A3DTopoCoEdgeData> CoEdgeData(A3DCoedge);
 	if (!CoEdgeData.IsValid())
 	{
-		return TSharedPtr<CADKernel::FTopologicalEdge>();
+		return TSharedPtr<UE::CADKernel::FTopologicalEdge>();
 	}
 
 	if (CoEdgeData->m_pUVCurve == nullptr)
 	{
-		return TSharedPtr<CADKernel::FTopologicalEdge>();
+		return TSharedPtr<UE::CADKernel::FTopologicalEdge>();
 	}
 
-	TSharedPtr<CADKernel::FCurve> Curve = AddCurve(CoEdgeData->m_pUVCurve, UVReparameterization);
+	TSharedPtr<UE::CADKernel::FCurve> Curve = AddCurve(CoEdgeData->m_pUVCurve, UVReparameterization);
 	if (!Curve.IsValid())
 	{
-		return TSharedPtr<CADKernel::FTopologicalEdge>();
+		return TSharedPtr<UE::CADKernel::FTopologicalEdge>();
 	}
 
-	TSharedRef<CADKernel::FRestrictionCurve> RestrictionCurve = CADKernel::FEntity::MakeShared<CADKernel::FRestrictionCurve>(Surface, Curve.ToSharedRef());
+	TSharedRef<UE::CADKernel::FRestrictionCurve> RestrictionCurve = UE::CADKernel::FEntity::MakeShared<UE::CADKernel::FRestrictionCurve>(Surface, Curve.ToSharedRef());
 
-	TSharedPtr<CADKernel::FTopologicalEdge> Edge = CADKernel::FTopologicalEdge::Make(RestrictionCurve);
+	TSharedPtr<UE::CADKernel::FTopologicalEdge> Edge = UE::CADKernel::FTopologicalEdge::Make(RestrictionCurve);
 	if (!Edge.IsValid())
 	{
-		return TSharedPtr<CADKernel::FTopologicalEdge>();
+		return TSharedPtr<UE::CADKernel::FTopologicalEdge>();
 	}
 
 	A3DEdgeToEdge.Emplace(A3DCoedge, Edge);
 
-	OutOrientation = CoEdgeData->m_ucOrientationUVWithLoop > 0 ? CADKernel::EOrientation::Front : CADKernel::EOrientation::Back;
+	OutOrientation = CoEdgeData->m_ucOrientationUVWithLoop > 0 ? UE::CADKernel::EOrientation::Front : UE::CADKernel::EOrientation::Back;
 
 	// Link edges
 	if (CoEdgeData->m_pNeighbor)
@@ -635,7 +663,7 @@ TSharedPtr<CADKernel::FTopologicalEdge> FTechSoftBridge::AddEdge(const A3DTopoCo
 		const A3DTopoCoEdge* Neighbor = CoEdgeData->m_pNeighbor;
 		while (Neighbor && Neighbor != A3DCoedge)
 		{
-			TSharedPtr<CADKernel::FTopologicalEdge>* TwinEdge = A3DEdgeToEdge.Find(Neighbor);
+			TSharedPtr<UE::CADKernel::FTopologicalEdge>* TwinEdge = A3DEdgeToEdge.Find(Neighbor);
 			if (TwinEdge != nullptr)
 			{
 				Edge->Link(*TwinEdge->Get(), SquareJoiningVertexTolerance);
@@ -656,28 +684,34 @@ TSharedPtr<CADKernel::FTopologicalEdge> FTechSoftBridge::AddEdge(const A3DTopoCo
 	return Edge;
 }
 
-TSharedPtr<CADKernel::FTopologicalLoop> FTechSoftBridge::AddLoop(const A3DTopoLoop* A3DLoop, const TSharedRef<CADKernel::FSurface>& Surface, const TechSoftUtils::FUVReparameterization& UVReparameterization, const bool bIsExternalLoop)
+TSharedPtr<UE::CADKernel::FTopologicalLoop> FTechSoftBridge::AddLoop(const A3DTopoLoop* A3DLoop, const TSharedRef<UE::CADKernel::FSurface>& Surface, const TechSoftUtils::FUVReparameterization& UVReparameterization, const bool bIsExternalLoop)
 {
+#ifdef CADKERNEL_DEV
 	Report.LoopCount++;
+#endif
 
-	TArray<TSharedPtr<CADKernel::FTopologicalEdge>> Edges;
-	TArray<CADKernel::EOrientation> Directions;
+	TArray<TSharedPtr<UE::CADKernel::FTopologicalEdge>> Edges;
+	TArray<UE::CADKernel::EOrientation> Directions;
 
 	TUniqueTSObj<A3DTopoLoopData> TopoLoopData(A3DLoop);
 	if (!TopoLoopData.IsValid())
 	{
+#ifdef CADKERNEL_DEV
 		Report.DegeneratedLoopCount++;
-		return TSharedPtr<CADKernel::FTopologicalLoop>();
+#endif
+		return TSharedPtr<UE::CADKernel::FTopologicalLoop>();
 	}
 
 	bool bLoopOrientation = (bool)TopoLoopData->m_ucOrientationWithSurface;
 	for (A3DUns32 Index = 0; Index < TopoLoopData->m_uiCoEdgeSize; ++Index)
 	{
-		CADKernel::EOrientation Orientation;
-		TSharedPtr<CADKernel::FTopologicalEdge> Edge = AddEdge(TopoLoopData->m_ppCoEdges[Index], Surface, UVReparameterization, Orientation);
+		UE::CADKernel::EOrientation Orientation;
+		TSharedPtr<UE::CADKernel::FTopologicalEdge> Edge = AddEdge(TopoLoopData->m_ppCoEdges[Index], Surface, UVReparameterization, Orientation);
 		if (!Edge.IsValid())
 		{
+#ifdef CADKERNEL_DEV
 			Report.DegeneratedEdgeCount++;
+#endif
 			continue;
 		}
 
@@ -687,16 +721,20 @@ TSharedPtr<CADKernel::FTopologicalLoop> FTechSoftBridge::AddLoop(const A3DTopoLo
 
 	if (Edges.Num() == 0)
 	{
+#ifdef CADKERNEL_DEV
 		Report.DegeneratedLoopCount++;
-		return TSharedPtr<CADKernel::FTopologicalLoop>();
+#endif
+		return TSharedPtr<UE::CADKernel::FTopologicalLoop>();
 	}
 
-	return CADKernel::FTopologicalLoop::Make(Edges, Directions, bIsExternalLoop, GeometricTolerance);
+	return UE::CADKernel::FTopologicalLoop::Make(Edges, Directions, bIsExternalLoop, GeometricTolerance);
 }
 
-void FTechSoftBridge::AddFace(const A3DTopoFace* A3DFace, CADKernel::EOrientation Orientation, TSharedRef<CADKernel::FShell>& Shell, uint32 ShellIndex)
+void FTechSoftBridge::AddFace(const A3DTopoFace* A3DFace, UE::CADKernel::EOrientation Orientation, TSharedRef<UE::CADKernel::FShell>& Shell, uint32 ShellIndex)
 {
+#ifdef CADKERNEL_DEV
 	Report.FaceCount++;
+#endif
 
 	FArchiveCADObject MetaData;
 	Parser.ExtractMetaData(A3DFace, MetaData);
@@ -704,17 +742,21 @@ void FTechSoftBridge::AddFace(const A3DTopoFace* A3DFace, CADKernel::EOrientatio
 	TUniqueTSObj<A3DTopoFaceData> TopoFaceData(A3DFace);
 	if (!TopoFaceData.IsValid())
 	{
+#ifdef CADKERNEL_DEV
 		Report.FailedFaceCount++;
+#endif
 		return;
 	}
 
 	const A3DSurfBase* A3DSurface = TopoFaceData->m_pSurface;
 	TechSoftUtils::FUVReparameterization UVReparameterization;
-	TSharedPtr<CADKernel::FSurface> SurfacePtr = AddSurface(A3DSurface, UVReparameterization);
+	TSharedPtr<UE::CADKernel::FSurface> SurfacePtr = AddSurface(A3DSurface, UVReparameterization);
 	if (!SurfacePtr.IsValid())
 	{
+#ifdef CADKERNEL_DEV
 		Report.DegeneratedSurfaceCount++;
-		Report.FailedFaceCount++;
+		Report.FailedFaceCount++; 
+#endif
 		return;
 	}
 
@@ -723,12 +765,12 @@ void FTechSoftBridge::AddFace(const A3DTopoFace* A3DFace, CADKernel::EOrientatio
 		SwapOrientation(Orientation);
 	}
 
-	TSharedRef<CADKernel::FSurface> Surface = SurfacePtr.ToSharedRef();
-	TSharedRef<CADKernel::FTopologicalFace> Face = CADKernel::FEntity::MakeShared<CADKernel::FTopologicalFace>(Surface);
+	TSharedRef<UE::CADKernel::FSurface> Surface = SurfacePtr.ToSharedRef();
+	TSharedRef<UE::CADKernel::FTopologicalFace> Face = UE::CADKernel::FEntity::MakeShared<UE::CADKernel::FTopologicalFace>(Surface);
 
 	if (TopoFaceData->m_bHasTrimDomain)
 	{
-		CADKernel::FSurfacicBoundary SurfaceBoundary = TechSoftUtils::GetSurfacicBoundary(TopoFaceData->m_sSurfaceDomain, UVReparameterization);
+		UE::CADKernel::FSurfacicBoundary SurfaceBoundary = TechSoftUtils::GetSurfacicBoundary(TopoFaceData->m_sSurfaceDomain, UVReparameterization);
 		Surface->TrimBoundaryTo(SurfaceBoundary);
 	}
 
@@ -738,26 +780,26 @@ void FTechSoftBridge::AddFace(const A3DTopoFace* A3DFace, CADKernel::EOrientatio
 	}
 	else
 	{
-		TArray<TSharedPtr<CADKernel::FTopologicalLoop>> Loops;
+		TArray<TSharedPtr<UE::CADKernel::FTopologicalLoop>> Loops;
 
 		for (A3DUns32 Index = 0; Index < TopoFaceData->m_uiLoopSize; ++Index)
 		{
 			const bool bIsExternalLoop = (Index == TopoFaceData->m_uiOuterLoopIndex);
-			TSharedPtr<CADKernel::FTopologicalLoop> Loop = AddLoop(TopoFaceData->m_ppLoops[Index], Surface, UVReparameterization, bIsExternalLoop);
+			TSharedPtr<UE::CADKernel::FTopologicalLoop> Loop = AddLoop(TopoFaceData->m_ppLoops[Index], Surface, UVReparameterization, bIsExternalLoop);
 			if (!Loop.IsValid())
 			{
 				continue;
 			}
 
-			TArray<CADKernel::FPoint2D> LoopSampling;
+			TArray<UE::CADKernel::FPoint2D> LoopSampling;
 			Loop->Get2DSampling(LoopSampling);
-			CADKernel::FAABB2D Boundary;
+			UE::CADKernel::FAABB2D Boundary;
 			Boundary += LoopSampling;
 			Loop->Boundary.Set(Boundary.GetMin(), Boundary.GetMax());
 
 			// Check if the loop is not composed with only degenerated edge
 			bool bDegeneratedLoop = true;
-			for (const CADKernel::FOrientedEdge& Edge : Loop->GetEdges())
+			for (const UE::CADKernel::FOrientedEdge& Edge : Loop->GetEdges())
 			{
 				if (!Edge.Entity->IsDegenerated())
 				{
@@ -775,14 +817,20 @@ void FTechSoftBridge::AddFace(const A3DTopoFace* A3DFace, CADKernel::EOrientatio
 
 		if (Loops.Num() == 0)
 		{
+#ifdef CADKERNEL_DEV
 			Report.FailedFaceCount++;
+#endif
 			Face->SetAsDegenerated();
 			Face->SetDeleted();
 			return;
 		}
 		else
 		{
-			Face->AddLoops(Loops, Report.DoubtfulLoopOrientationCount);
+			int32 DoubtfulLoopOrientationCount = 0;
+			Face->AddLoops(Loops, DoubtfulLoopOrientationCount);
+#ifdef CADKERNEL_DEV
+			Report.DoubtfulLoopOrientationCount += DoubtfulLoopOrientationCount;
+#endif
 		}
 	}
 
@@ -793,9 +841,11 @@ void FTechSoftBridge::AddFace(const A3DTopoFace* A3DFace, CADKernel::EOrientatio
 
 static bool bUseSurfaceAsNurbs = true;
 
-TSharedPtr<CADKernel::FSurface> FTechSoftBridge::AddSurface(const A3DSurfBase* A3DSurface, TechSoftUtils::FUVReparameterization& OutUVReparameterization)
+TSharedPtr<UE::CADKernel::FSurface> FTechSoftBridge::AddSurface(const A3DSurfBase* A3DSurface, TechSoftUtils::FUVReparameterization& OutUVReparameterization)
 {
+#ifdef CADKERNEL_DEV
 	Report.SurfaceCount++;
+#endif
 
 	FArchiveCADObject MetaData;
 	Parser.ExtractMetaData(A3DSurface, MetaData);
@@ -865,17 +915,19 @@ TSharedPtr<CADKernel::FSurface> FTechSoftBridge::AddSurface(const A3DSurfBase* A
 	{
 		return AddSurfaceAsNurbs(A3DSurface, OutUVReparameterization);
 	}
-	return TSharedPtr<CADKernel::FSurface>();
+	return TSharedPtr<UE::CADKernel::FSurface>();
 }
 
-TSharedPtr<CADKernel::FSurface> FTechSoftBridge::AddConeSurface(const A3DSurfBase* Surface, TechSoftUtils::FUVReparameterization& OutUVReparameterization)
+TSharedPtr<UE::CADKernel::FSurface> FTechSoftBridge::AddConeSurface(const A3DSurfBase* Surface, TechSoftUtils::FUVReparameterization& OutUVReparameterization)
 {
+#ifdef CADKERNEL_DEV
 	Report.ConeSurfaceCount++;
+#endif
 
 	TUniqueTSObj<A3DSurfConeData> A3DConeData(Surface);
 	if (!A3DConeData.IsValid())
 	{
-		return TSharedPtr<CADKernel::FSurface>();
+		return TSharedPtr<UE::CADKernel::FSurface>();
 	}
 
 	OutUVReparameterization.AddUVTransform(A3DConeData->m_sParam);
@@ -885,87 +937,99 @@ TSharedPtr<CADKernel::FSurface> FTechSoftBridge::AddConeSurface(const A3DSurfBas
 		OutUVReparameterization.SetNeedSwapOrientation();
 	}
 
-	CADKernel::FMatrixH CoordinateSystem = TechSoftUtils::CreateCoordinateSystem(A3DConeData->m_sTrsf, BodyScale);
-	CADKernel::FSurfacicBoundary Boundary = TechSoftUtils::GetSurfacicBoundary(A3DConeData->m_sParam.m_sUVDomain, OutUVReparameterization);
-	return CADKernel::FSurface::MakeConeSurface(GeometricTolerance, CoordinateSystem, A3DConeData->m_dRadius * BodyScale, A3DConeData->m_dSemiAngle, Boundary);
+	UE::CADKernel::FMatrixH CoordinateSystem = TechSoftUtils::CreateCoordinateSystem(A3DConeData->m_sTrsf, BodyScale);
+	UE::CADKernel::FSurfacicBoundary Boundary = TechSoftUtils::GetSurfacicBoundary(A3DConeData->m_sParam.m_sUVDomain, OutUVReparameterization);
+	return UE::CADKernel::FSurface::MakeConeSurface(GeometricTolerance, CoordinateSystem, A3DConeData->m_dRadius * BodyScale, A3DConeData->m_dSemiAngle, Boundary);
 }
 
-TSharedPtr<CADKernel::FSurface> FTechSoftBridge::AddCylinderSurface(const A3DSurfBase* Surface, TechSoftUtils::FUVReparameterization& OutUVReparameterization)
+TSharedPtr<UE::CADKernel::FSurface> FTechSoftBridge::AddCylinderSurface(const A3DSurfBase* Surface, TechSoftUtils::FUVReparameterization& OutUVReparameterization)
 {
+#ifdef CADKERNEL_DEV
 	Report.CylinderSurfaceCount++;
+#endif
 
 	TUniqueTSObj<A3DSurfCylinderData> A3DCylinderData(Surface);
 	if (!A3DCylinderData.IsValid())
 	{
-		return TSharedPtr<CADKernel::FSurface>();
+		return TSharedPtr<UE::CADKernel::FSurface>();
 	}
 
 	OutUVReparameterization.AddUVTransform(A3DCylinderData->m_sParam);
 	OutUVReparameterization.ScaleUVTransform(1, BodyScale);
 
-	CADKernel::FMatrixH CoordinateSystem = TechSoftUtils::CreateCoordinateSystem(A3DCylinderData->m_sTrsf, BodyScale);
-	CADKernel::FSurfacicBoundary Boundary = TechSoftUtils::GetSurfacicBoundary(A3DCylinderData->m_sParam.m_sUVDomain, OutUVReparameterization);
-	return CADKernel::FSurface::MakeCylinderSurface(GeometricTolerance, CoordinateSystem, A3DCylinderData->m_dRadius * BodyScale, Boundary);
+	UE::CADKernel::FMatrixH CoordinateSystem = TechSoftUtils::CreateCoordinateSystem(A3DCylinderData->m_sTrsf, BodyScale);
+	UE::CADKernel::FSurfacicBoundary Boundary = TechSoftUtils::GetSurfacicBoundary(A3DCylinderData->m_sParam.m_sUVDomain, OutUVReparameterization);
+	return UE::CADKernel::FSurface::MakeCylinderSurface(GeometricTolerance, CoordinateSystem, A3DCylinderData->m_dRadius * BodyScale, Boundary);
 }
 
-TSharedPtr<CADKernel::FSurface> FTechSoftBridge::AddLinearTransfoSurface(const A3DSurfBase* Surface, TechSoftUtils::FUVReparameterization& OutUVReparameterization)
+TSharedPtr<UE::CADKernel::FSurface> FTechSoftBridge::AddLinearTransfoSurface(const A3DSurfBase* Surface, TechSoftUtils::FUVReparameterization& OutUVReparameterization)
 {
+#ifdef CADKERNEL_DEV
 	Report.LinearTransfoSurfaceCount++;
+#endif
 
 	if (bUseSurfaceAsNurbs)
 	{
 		return AddSurfaceAsNurbs(Surface, OutUVReparameterization);
 	}
 	// Todo
-	return TSharedPtr<CADKernel::FSurface>();
+	return TSharedPtr<UE::CADKernel::FSurface>();
 }
 
-TSharedPtr<CADKernel::FSurface> FTechSoftBridge::AddNurbsSurface(const A3DSurfNurbs* Nurbs, TechSoftUtils::FUVReparameterization& OutUVReparameterization)
+TSharedPtr<UE::CADKernel::FSurface> FTechSoftBridge::AddNurbsSurface(const A3DSurfNurbs* Nurbs, TechSoftUtils::FUVReparameterization& OutUVReparameterization)
 {
+#ifdef CADKERNEL_DEV
 	Report.NurbsSurfaceCount++;
+#endif
 
 	TUniqueTSObj<A3DSurfNurbsData> A3DNurbsData(Nurbs);
 	if (!A3DNurbsData.IsValid())
 	{
-		return TSharedPtr<CADKernel::FSurface>();
+		return TSharedPtr<UE::CADKernel::FSurface>();
 	}
 
 	return AddSurfaceNurbs(*A3DNurbsData, OutUVReparameterization);
 }
 
-TSharedPtr<CADKernel::FSurface> FTechSoftBridge::AddOffsetSurface(const A3DSurfBase* Surface, TechSoftUtils::FUVReparameterization& OutUVReparameterization)
+TSharedPtr<UE::CADKernel::FSurface> FTechSoftBridge::AddOffsetSurface(const A3DSurfBase* Surface, TechSoftUtils::FUVReparameterization& OutUVReparameterization)
 {
+#ifdef CADKERNEL_DEV
 	Report.OffsetSurfaceCount++;
+#endif
 
 	if (bUseSurfaceAsNurbs)
 	{
 		return AddSurfaceAsNurbs(Surface, OutUVReparameterization);
 	}
 
-	return TSharedPtr<CADKernel::FSurface>();
+	return TSharedPtr<UE::CADKernel::FSurface>();
 }
 
-TSharedPtr<CADKernel::FSurface> FTechSoftBridge::AddPlaneSurface(const A3DSurfPlane* Surface, TechSoftUtils::FUVReparameterization& OutUVReparameterization)
+TSharedPtr<UE::CADKernel::FSurface> FTechSoftBridge::AddPlaneSurface(const A3DSurfPlane* Surface, TechSoftUtils::FUVReparameterization& OutUVReparameterization)
 {
+#ifdef CADKERNEL_DEV
 	Report.PlaneSurfaceCount++;
+#endif
 
 	TUniqueTSObj<A3DSurfPlaneData> A3DPlaneData(Surface);
 	if (!A3DPlaneData.IsValid())
 	{
-		return TSharedPtr<CADKernel::FSurface>();
+		return TSharedPtr<UE::CADKernel::FSurface>();
 	}
 
 	OutUVReparameterization.AddUVTransform(A3DPlaneData->m_sParam);
 	OutUVReparameterization.ScaleUVTransform(BodyScale, BodyScale);
 	
-	CADKernel::FMatrixH CoordinateSystem = TechSoftUtils::CreateCoordinateSystem(A3DPlaneData->m_sTrsf, BodyScale);
-	CADKernel::FSurfacicBoundary Boundary = TechSoftUtils::GetSurfacicBoundary(A3DPlaneData->m_sParam.m_sUVDomain, OutUVReparameterization);
-	return CADKernel::FSurface::MakePlaneSurface(GeometricTolerance, CoordinateSystem, Boundary);
+	UE::CADKernel::FMatrixH CoordinateSystem = TechSoftUtils::CreateCoordinateSystem(A3DPlaneData->m_sTrsf, BodyScale);
+	UE::CADKernel::FSurfacicBoundary Boundary = TechSoftUtils::GetSurfacicBoundary(A3DPlaneData->m_sParam.m_sUVDomain, OutUVReparameterization);
+	return UE::CADKernel::FSurface::MakePlaneSurface(GeometricTolerance, CoordinateSystem, Boundary);
 }
 
-TSharedPtr<CADKernel::FSurface> FTechSoftBridge::AddRevolutionSurface(const A3DSurfRevolution* Surface, TechSoftUtils::FUVReparameterization& OutUVReparameterization)
+TSharedPtr<UE::CADKernel::FSurface> FTechSoftBridge::AddRevolutionSurface(const A3DSurfRevolution* Surface, TechSoftUtils::FUVReparameterization& OutUVReparameterization)
 {
+#ifdef CADKERNEL_DEV
 	Report.RevolutionSurfaceCount++;
+#endif
 
 	if (bUseSurfaceAsNurbs)
 	{
@@ -975,15 +1039,17 @@ TSharedPtr<CADKernel::FSurface> FTechSoftBridge::AddRevolutionSurface(const A3DS
 	TUniqueTSObj<A3DSurfRevolutionData> A3DRevolutionData(Surface);
 	if (!A3DRevolutionData.IsValid())
 	{
-		return TSharedPtr<CADKernel::FSurface>();
+		return TSharedPtr<UE::CADKernel::FSurface>();
 	}
 
-	return TSharedPtr<CADKernel::FSurface>();
+	return TSharedPtr<UE::CADKernel::FSurface>();
 }
 
-TSharedPtr<CADKernel::FSurface> FTechSoftBridge::AddRuledSurface(const A3DSurfBase* Surface, TechSoftUtils::FUVReparameterization& OutUVReparameterization)
+TSharedPtr<UE::CADKernel::FSurface> FTechSoftBridge::AddRuledSurface(const A3DSurfBase* Surface, TechSoftUtils::FUVReparameterization& OutUVReparameterization)
 {
+#ifdef CADKERNEL_DEV
 	Report.RuledSurfaceCount++;
+#endif
 
 	if (bUseSurfaceAsNurbs)
 	{
@@ -993,48 +1059,54 @@ TSharedPtr<CADKernel::FSurface> FTechSoftBridge::AddRuledSurface(const A3DSurfBa
 	TUniqueTSObj<A3DSurfRuledData> A3DRuledData(Surface);
 	if (!A3DRuledData.IsValid())
 	{
-		return TSharedPtr<CADKernel::FSurface>();
+		return TSharedPtr<UE::CADKernel::FSurface>();
 	}
 
-	return TSharedPtr<CADKernel::FSurface>();
+	return TSharedPtr<UE::CADKernel::FSurface>();
 }
 
-TSharedPtr<CADKernel::FSurface> FTechSoftBridge::AddSphereSurface(const A3DSurfBase* Surface, TechSoftUtils::FUVReparameterization& OutUVReparameterization)
+TSharedPtr<UE::CADKernel::FSurface> FTechSoftBridge::AddSphereSurface(const A3DSurfBase* Surface, TechSoftUtils::FUVReparameterization& OutUVReparameterization)
 {
+#ifdef CADKERNEL_DEV
 	Report.SphereSurfaceCount++;
+#endif
 
 	TUniqueTSObj<A3DSurfSphereData> A3DSphereData(Surface);
 	if (!A3DSphereData.IsValid())
 	{
-		return TSharedPtr<CADKernel::FSurface>();
+		return TSharedPtr<UE::CADKernel::FSurface>();
 	}
 
 	OutUVReparameterization.AddUVTransform(A3DSphereData->m_sParam);
 
-	CADKernel::FMatrixH CoordinateSystem = TechSoftUtils::CreateCoordinateSystem(A3DSphereData->m_sTrsf, BodyScale);
-	CADKernel::FSurfacicBoundary Boundary = TechSoftUtils::GetSurfacicBoundary(A3DSphereData->m_sParam.m_sUVDomain, OutUVReparameterization);
-	return CADKernel::FSurface::MakeSphericalSurface(GeometricTolerance, CoordinateSystem, A3DSphereData->m_dRadius * BodyScale, Boundary);
+	UE::CADKernel::FMatrixH CoordinateSystem = TechSoftUtils::CreateCoordinateSystem(A3DSphereData->m_sTrsf, BodyScale);
+	UE::CADKernel::FSurfacicBoundary Boundary = TechSoftUtils::GetSurfacicBoundary(A3DSphereData->m_sParam.m_sUVDomain, OutUVReparameterization);
+	return UE::CADKernel::FSurface::MakeSphericalSurface(GeometricTolerance, CoordinateSystem, A3DSphereData->m_dRadius * BodyScale, Boundary);
 }
 
-TSharedPtr<CADKernel::FSurface> FTechSoftBridge::AddTorusSurface(const A3DSurfBase* Surface, TechSoftUtils::FUVReparameterization& OutUVReparameterization)
+TSharedPtr<UE::CADKernel::FSurface> FTechSoftBridge::AddTorusSurface(const A3DSurfBase* Surface, TechSoftUtils::FUVReparameterization& OutUVReparameterization)
 {
+#ifdef CADKERNEL_DEV
 	Report.TorusSurfaceCount++;
+#endif
 
 	TUniqueTSObj<A3DSurfTorusData> A3DTorusData(Surface);
 	if (!A3DTorusData.IsValid())
 	{
-		return TSharedPtr<CADKernel::FSurface>();
+		return TSharedPtr<UE::CADKernel::FSurface>();
 	}
 
 	OutUVReparameterization.AddUVTransform(A3DTorusData->m_sParam);
-	CADKernel::FMatrixH CoordinateSystem = TechSoftUtils::CreateCoordinateSystem(A3DTorusData->m_sTrsf, BodyScale);
-	CADKernel::FSurfacicBoundary Boundary = TechSoftUtils::GetSurfacicBoundary(A3DTorusData->m_sParam.m_sUVDomain, OutUVReparameterization);
-	return CADKernel::FSurface::MakeTorusSurface(GeometricTolerance, CoordinateSystem, A3DTorusData->m_dMajorRadius * BodyScale, A3DTorusData->m_dMinorRadius * BodyScale, Boundary);
+	UE::CADKernel::FMatrixH CoordinateSystem = TechSoftUtils::CreateCoordinateSystem(A3DTorusData->m_sTrsf, BodyScale);
+	UE::CADKernel::FSurfacicBoundary Boundary = TechSoftUtils::GetSurfacicBoundary(A3DTorusData->m_sParam.m_sUVDomain, OutUVReparameterization);
+	return UE::CADKernel::FSurface::MakeTorusSurface(GeometricTolerance, CoordinateSystem, A3DTorusData->m_dMajorRadius * BodyScale, A3DTorusData->m_dMinorRadius * BodyScale, Boundary);
 }
 
-TSharedPtr<CADKernel::FSurface> FTechSoftBridge::AddBlend01Surface(const A3DSurfBase* Surface, TechSoftUtils::FUVReparameterization& OutUVReparameterization)
+TSharedPtr<UE::CADKernel::FSurface> FTechSoftBridge::AddBlend01Surface(const A3DSurfBase* Surface, TechSoftUtils::FUVReparameterization& OutUVReparameterization)
 {
+#ifdef CADKERNEL_DEV
 	Report.Blend01SurfaceCount++;
+#endif
 
 	if (bUseSurfaceAsNurbs)
 	{
@@ -1044,15 +1116,17 @@ TSharedPtr<CADKernel::FSurface> FTechSoftBridge::AddBlend01Surface(const A3DSurf
 	TUniqueTSObj<A3DSurfBlend01Data> A3DBlend01Data(Surface);
 	if (!A3DBlend01Data.IsValid())
 	{
-		return TSharedPtr<CADKernel::FSurface>();
+		return TSharedPtr<UE::CADKernel::FSurface>();
 	}
 
-	return TSharedPtr<CADKernel::FSurface>();
+	return TSharedPtr<UE::CADKernel::FSurface>();
 }
 
-TSharedPtr<CADKernel::FSurface> FTechSoftBridge::AddBlend02Surface(const A3DSurfBase* Surface, TechSoftUtils::FUVReparameterization& OutUVReparameterization)
+TSharedPtr<UE::CADKernel::FSurface> FTechSoftBridge::AddBlend02Surface(const A3DSurfBase* Surface, TechSoftUtils::FUVReparameterization& OutUVReparameterization)
 {
+#ifdef CADKERNEL_DEV
 	Report.Blend02SurfaceCount++;
+#endif
 
 	if (bUseSurfaceAsNurbs)
 	{
@@ -1062,27 +1136,31 @@ TSharedPtr<CADKernel::FSurface> FTechSoftBridge::AddBlend02Surface(const A3DSurf
 	TUniqueTSObj<A3DSurfBlend02Data> A3DBlend02Data(Surface);
 	if (!A3DBlend02Data.IsValid())
 	{
-		return TSharedPtr<CADKernel::FSurface>();
+		return TSharedPtr<UE::CADKernel::FSurface>();
 	}
 
-	return TSharedPtr<CADKernel::FSurface>();
+	return TSharedPtr<UE::CADKernel::FSurface>();
 }
 
-TSharedPtr<CADKernel::FSurface> FTechSoftBridge::AddBlend03Surface(const A3DSurfBase* Surface, TechSoftUtils::FUVReparameterization& OutUVReparameterization)
+TSharedPtr<UE::CADKernel::FSurface> FTechSoftBridge::AddBlend03Surface(const A3DSurfBase* Surface, TechSoftUtils::FUVReparameterization& OutUVReparameterization)
 {
+#ifdef CADKERNEL_DEV
 	Report.Blend03SurfaceCount++;
+#endif
 
 	if (bUseSurfaceAsNurbs)
 	{
 		return AddSurfaceAsNurbs(Surface, OutUVReparameterization);
 	}
 
-	return TSharedPtr<CADKernel::FSurface>();
+	return TSharedPtr<UE::CADKernel::FSurface>();
 }
 
-TSharedPtr<CADKernel::FSurface> FTechSoftBridge::AddCylindricalSurface(const A3DSurfBase* Surface, TechSoftUtils::FUVReparameterization& OutUVReparameterization)
+TSharedPtr<UE::CADKernel::FSurface> FTechSoftBridge::AddCylindricalSurface(const A3DSurfBase* Surface, TechSoftUtils::FUVReparameterization& OutUVReparameterization)
 {
+#ifdef CADKERNEL_DEV
 	Report.CylindricalSurfaceCount++;
+#endif
 
 	if (bUseSurfaceAsNurbs)
 	{
@@ -1092,15 +1170,17 @@ TSharedPtr<CADKernel::FSurface> FTechSoftBridge::AddCylindricalSurface(const A3D
 	TUniqueTSObj<A3DSurfCylindricalData> A3DCylindricalData(Surface);
 	if (!A3DCylindricalData.IsValid())
 	{
-		return TSharedPtr<CADKernel::FSurface>();
+		return TSharedPtr<UE::CADKernel::FSurface>();
 	}
 
-	return TSharedPtr<CADKernel::FSurface>();
+	return TSharedPtr<UE::CADKernel::FSurface>();
 }
 
-TSharedPtr<CADKernel::FSurface> FTechSoftBridge::AddPipeSurface(const A3DSurfBase* Surface, TechSoftUtils::FUVReparameterization& OutUVReparameterization)
+TSharedPtr<UE::CADKernel::FSurface> FTechSoftBridge::AddPipeSurface(const A3DSurfBase* Surface, TechSoftUtils::FUVReparameterization& OutUVReparameterization)
 {
+#ifdef CADKERNEL_DEV
 	Report.PipeSurfaceCount++;
+#endif
 
 	if (bUseSurfaceAsNurbs)
 	{
@@ -1110,15 +1190,17 @@ TSharedPtr<CADKernel::FSurface> FTechSoftBridge::AddPipeSurface(const A3DSurfBas
 	TUniqueTSObj<A3DSurfPipeData> A3DPipeData(Surface);
 	if (!A3DPipeData.IsValid())
 	{
-		return TSharedPtr<CADKernel::FSurface>();
+		return TSharedPtr<UE::CADKernel::FSurface>();
 	}
 
-	return TSharedPtr<CADKernel::FSurface>();
+	return TSharedPtr<UE::CADKernel::FSurface>();
 }
 
-TSharedPtr<CADKernel::FSurface> FTechSoftBridge::AddExtrusionSurface(const A3DSurfBase* Surface, TechSoftUtils::FUVReparameterization& OutUVReparameterization)
+TSharedPtr<UE::CADKernel::FSurface> FTechSoftBridge::AddExtrusionSurface(const A3DSurfBase* Surface, TechSoftUtils::FUVReparameterization& OutUVReparameterization)
 {
+#ifdef CADKERNEL_DEV
 	Report.ExtrusionSurfaceCount++;
+#endif
 
 	if (bUseSurfaceAsNurbs)
 	{
@@ -1128,15 +1210,17 @@ TSharedPtr<CADKernel::FSurface> FTechSoftBridge::AddExtrusionSurface(const A3DSu
 	TUniqueTSObj<A3DSurfExtrusionData> A3DExtrusionData(Surface);
 	if (!A3DExtrusionData.IsValid())
 	{
-		return TSharedPtr<CADKernel::FSurface>();
+		return TSharedPtr<UE::CADKernel::FSurface>();
 	}
 
-	return TSharedPtr<CADKernel::FSurface>();
+	return TSharedPtr<UE::CADKernel::FSurface>();
 }
 
-TSharedPtr<CADKernel::FSurface> FTechSoftBridge::AddSurfaceFromCurves(const A3DSurfBase* Surface, TechSoftUtils::FUVReparameterization& OutUVReparameterization)
+TSharedPtr<UE::CADKernel::FSurface> FTechSoftBridge::AddSurfaceFromCurves(const A3DSurfBase* Surface, TechSoftUtils::FUVReparameterization& OutUVReparameterization)
 {
+#ifdef CADKERNEL_DEV
 	Report.SurfaceFromCurvesCount++;
+#endif
 
 	if (bUseSurfaceAsNurbs)
 	{
@@ -1146,15 +1230,17 @@ TSharedPtr<CADKernel::FSurface> FTechSoftBridge::AddSurfaceFromCurves(const A3DS
 	TUniqueTSObj<A3DSurfFromCurvesData> A3DFromCurvesData(Surface);
 	if (!A3DFromCurvesData.IsValid())
 	{
-		return TSharedPtr<CADKernel::FSurface>();
+		return TSharedPtr<UE::CADKernel::FSurface>();
 	}
 
-	return TSharedPtr<CADKernel::FSurface>();
+	return TSharedPtr<UE::CADKernel::FSurface>();
 }
 
-TSharedPtr<CADKernel::FSurface> FTechSoftBridge::AddTransformSurface(const A3DSurfBase* Surface, TechSoftUtils::FUVReparameterization& OutUVReparameterization)
+TSharedPtr<UE::CADKernel::FSurface> FTechSoftBridge::AddTransformSurface(const A3DSurfBase* Surface, TechSoftUtils::FUVReparameterization& OutUVReparameterization)
 {
+#ifdef CADKERNEL_DEV
 	Report.TransformSurfaceCount++;
+#endif
 
 	if (bUseSurfaceAsNurbs)
 	{
@@ -1164,15 +1250,15 @@ TSharedPtr<CADKernel::FSurface> FTechSoftBridge::AddTransformSurface(const A3DSu
 	TUniqueTSObj<A3DSurfFromCurvesData> A3DTransformData(Surface);
 	if (!A3DTransformData.IsValid())
 	{
-		return TSharedPtr<CADKernel::FSurface>();
+		return TSharedPtr<UE::CADKernel::FSurface>();
 	}
 
-	return TSharedPtr<CADKernel::FSurface>();
+	return TSharedPtr<UE::CADKernel::FSurface>();
 }
 
-TSharedPtr<CADKernel::FSurface> FTechSoftBridge::AddSurfaceNurbs(const A3DSurfNurbsData& A3DNurbsData, TechSoftUtils::FUVReparameterization& OutUVReparameterization)
+TSharedPtr<UE::CADKernel::FSurface> FTechSoftBridge::AddSurfaceNurbs(const A3DSurfNurbsData& A3DNurbsData, TechSoftUtils::FUVReparameterization& OutUVReparameterization)
 {
-	CADKernel::FNurbsSurfaceData NurbsData;
+	UE::CADKernel::FNurbsSurfaceData NurbsData;
 
 	NurbsData.PoleUCount = A3DNurbsData.m_uiUCtrlSize;
 	NurbsData.PoleVCount = A3DNurbsData.m_uiVCtrlSize;
@@ -1184,11 +1270,11 @@ TSharedPtr<CADKernel::FSurface> FTechSoftBridge::AddSurfaceNurbs(const A3DSurfNu
 	TechSoftUtils::FillDoubleArray(A3DNurbsData.m_uiUKnotSize, A3DNurbsData.m_pdUKnots, NurbsData.UNodalVector);
 	TechSoftUtils::FillDoubleArray(A3DNurbsData.m_uiVKnotSize, A3DNurbsData.m_pdVKnots, NurbsData.VNodalVector);
 
-	TArray<CADKernel::FPoint> Poles;
+	TArray<UE::CADKernel::FPoint> Poles;
 	TechSoftUtils::FillPointArray(NurbsData.PoleUCount, NurbsData.PoleVCount, A3DNurbsData.m_pCtrlPts, NurbsData.Poles);
 	if (!FMath::IsNearlyEqual(BodyScale, 1.))
 	{
-		for (CADKernel::FPoint& Point : NurbsData.Poles)
+		for (UE::CADKernel::FPoint& Point : NurbsData.Poles)
 		{
 			Point *= BodyScale;
 		}
@@ -1201,12 +1287,14 @@ TSharedPtr<CADKernel::FSurface> FTechSoftBridge::AddSurfaceNurbs(const A3DSurfNu
 		TechSoftUtils::FillDoubleArray(NurbsData.PoleUCount, NurbsData.PoleVCount, A3DNurbsData.m_pdWeights, NurbsData.Weights);
 	}
 
-	return CADKernel::FSurface::MakeNurbsSurface(GeometricTolerance, NurbsData);
+	return UE::CADKernel::FSurface::MakeNurbsSurface(GeometricTolerance, NurbsData);
 }
 
-TSharedPtr<CADKernel::FSurface> FTechSoftBridge::AddSurfaceAsNurbs(const A3DSurfBase* Surface, TechSoftUtils::FUVReparameterization& OutUVReparameterization)
+TSharedPtr<UE::CADKernel::FSurface> FTechSoftBridge::AddSurfaceAsNurbs(const A3DSurfBase* Surface, TechSoftUtils::FUVReparameterization& OutUVReparameterization)
 {
+#ifdef CADKERNEL_DEV
 	Report.SurfaceAsNurbsCount++;
+#endif
 
 	TUniqueTSObj<A3DSurfNurbsData> A3DNurbsData;
 
@@ -1216,14 +1304,14 @@ TSharedPtr<CADKernel::FSurface> FTechSoftBridge::AddSurfaceAsNurbs(const A3DSurf
 
 	if (!A3DNurbsData.IsValid())
 	{
-		return TSharedPtr<CADKernel::FSurface>();
+		return TSharedPtr<UE::CADKernel::FSurface>();
 	}
 
 	return AddSurfaceNurbs(*A3DNurbsData, OutUVReparameterization);
 
 }
 
-void FTechSoftBridge::AddMetadata(FArchiveCADObject& MetaData, CADKernel::FTopologicalShapeEntity& Entity)
+void FTechSoftBridge::AddMetadata(FArchiveCADObject& MetaData, UE::CADKernel::FTopologicalShapeEntity& Entity)
 {
 	FString* Name = MetaData.MetaData.Find(TEXT("Name"));
 	if(Name != nullptr)
diff --git a/Engine/Plugins/Enterprise/DatasmithCADImporter/Source/CADInterfaces/Private/TechSoftBridge.h b/Engine/Plugins/Enterprise/DatasmithCADImporter/Source/CADInterfaces/Private/TechSoftBridge.h
--- a/Engine/Plugins/Enterprise/DatasmithCADImporter/Source/CADInterfaces/Private/TechSoftBridge.h
+++ b/Engine/Plugins/Enterprise/DatasmithCADImporter/Source/CADInterfaces/Private/TechSoftBridge.h
@@ -4,14 +4,16 @@
 
 #ifdef USE_TECHSOFT_SDK
 
-#include "CADFileReport.h"
-
 #include "CADKernel/Geo/GeoEnum.h"
 #include "CADKernel/Math/MatrixH.h"
 
 #include "TechSoftInterface.h"
 
-namespace CADKernel
+#ifdef CADKERNEL_DEV
+#include "CADFileReport.h"
+#endif
+
+namespace UE::CADKernel
 {
 
 class FBody;
@@ -58,10 +60,10 @@ public:
 
 	void SetCoef(const double InUScale, const double InUOffset, const double InVScale, const double InVOffset)
 	{
-		Scale[CADKernel::EIso::IsoU] = InUScale;
-		Scale[CADKernel::EIso::IsoV] = InVScale;
-		Offset[CADKernel::EIso::IsoU] = InUOffset;
-		Offset[CADKernel::EIso::IsoV] = InVOffset;
+		Scale[UE::CADKernel::EIso::IsoU] = InUScale;
+		Scale[UE::CADKernel::EIso::IsoV] = InVScale;
+		Offset[UE::CADKernel::EIso::IsoU] = InUOffset;
+		Offset[UE::CADKernel::EIso::IsoV] = InVOffset;
 		SetNeedApply();
 	}
 
@@ -87,7 +89,7 @@ public:
 
 	void SetNeedApply()
 	{
-		if (!FMath::IsNearlyEqual(Scale[CADKernel::EIso::IsoU], 1.) || !FMath::IsNearlyEqual(Scale[CADKernel::EIso::IsoV], 1.) || !FMath::IsNearlyEqual(Offset[CADKernel::EIso::IsoU], 0.) || !FMath::IsNearlyEqual(Offset[CADKernel::EIso::IsoV], 0.))
+		if (!FMath::IsNearlyEqual(Scale[UE::CADKernel::EIso::IsoU], 1.) || !FMath::IsNearlyEqual(Scale[UE::CADKernel::EIso::IsoV], 1.) || !FMath::IsNearlyEqual(Offset[UE::CADKernel::EIso::IsoU], 0.) || !FMath::IsNearlyEqual(Offset[UE::CADKernel::EIso::IsoV], 0.))
 		{
 			bNeedApply = true;
 		}
@@ -103,25 +105,25 @@ public:
 		{
 			Swap(InUScale, InVScale);
 		}
-		Scale[CADKernel::EIso::IsoU] *= InUScale;
-		Scale[CADKernel::EIso::IsoV] *= InVScale;
-		Offset[CADKernel::EIso::IsoU] *= InUScale;
-		Offset[CADKernel::EIso::IsoV] *= InVScale;
+		Scale[UE::CADKernel::EIso::IsoU] *= InUScale;
+		Scale[UE::CADKernel::EIso::IsoV] *= InVScale;
+		Offset[UE::CADKernel::EIso::IsoU] *= InUScale;
+		Offset[UE::CADKernel::EIso::IsoV] *= InVScale;
 		SetNeedApply();
 	}
 
-	void Process(TArray<CADKernel::FPoint>& Poles) const
+	void Process(TArray<UE::CADKernel::FPoint>& Poles) const
 	{
 		if(bNeedApply)
 		{
-			for (CADKernel::FPoint& Point : Poles)
+			for (UE::CADKernel::FPoint& Point : Poles)
 			{
 				Apply(Point);
 			}
 		}
 		if (bSwapUV)
 		{
-			for (CADKernel::FPoint& Point : Poles)
+			for (UE::CADKernel::FPoint& Point : Poles)
 			{
 				GetSwapUV(Point);
 			}
@@ -139,20 +141,20 @@ public:
 		SetNeedApply();
 	}
 
-	void Apply(CADKernel::FPoint2D& Point) const
+	void Apply(UE::CADKernel::FPoint2D& Point) const
 	{
-		Point.U = Scale[CADKernel::EIso::IsoU] * Point.U + Offset[CADKernel::EIso::IsoU];
-		Point.V = Scale[CADKernel::EIso::IsoV] * Point.V + Offset[CADKernel::EIso::IsoV];
+		Point.U = Scale[UE::CADKernel::EIso::IsoU] * Point.U + Offset[UE::CADKernel::EIso::IsoU];
+		Point.V = Scale[UE::CADKernel::EIso::IsoV] * Point.V + Offset[UE::CADKernel::EIso::IsoV];
 	}
 
 private:
-	void Apply(CADKernel::FPoint& Point) const
+	void Apply(UE::CADKernel::FPoint& Point) const
 	{
-		Point.X = Scale[CADKernel::EIso::IsoU] * Point.X + Offset[CADKernel::EIso::IsoU];
-		Point.Y = Scale[CADKernel::EIso::IsoV] * Point.Y + Offset[CADKernel::EIso::IsoV];
+		Point.X = Scale[UE::CADKernel::EIso::IsoU] * Point.X + Offset[UE::CADKernel::EIso::IsoU];
+		Point.Y = Scale[UE::CADKernel::EIso::IsoV] * Point.Y + Offset[UE::CADKernel::EIso::IsoV];
 	}
 
-	void GetSwapUV(CADKernel::FPoint& Point) const
+	void GetSwapUV(UE::CADKernel::FPoint& Point) const
 	{
 		Swap(Point.X, Point.Y);
 	}
@@ -165,85 +167,84 @@ class FTechSoftBridge
 private:
 	FTechSoftFileParser& Parser;
 
-	CADKernel::FSession& Session;
-	CADKernel::FModel& Model;
-	CADKernel::FCADFileReport& Report;
+	UE::CADKernel::FSession& Session;
+	UE::CADKernel::FModel& Model;
 
 	const double GeometricTolerance;
 	const double SquareGeometricTolerance;
 	const double SquareJoiningVertexTolerance;
 
-	TMap<const A3DEntity*, CADKernel::FBody*> TechSoftToCADKernel;
-	TMap<CADKernel::FBody*, const A3DEntity*> CADKernelToTechSoft;
-	TMap<const A3DTopoCoEdge*, TSharedPtr<CADKernel::FTopologicalEdge>> A3DEdgeToEdge;
+	TMap<const A3DEntity*, UE::CADKernel::FBody*> TechSoftToCADKernel;
+	TMap<UE::CADKernel::FBody*, const A3DEntity*> CADKernelToTechSoft;
+	TMap<const A3DTopoCoEdge*, TSharedPtr<UE::CADKernel::FTopologicalEdge>> A3DEdgeToEdge;
 
 	double BodyScale = 1;
 
 public:
-	FTechSoftBridge(FTechSoftFileParser& InParser, CADKernel::FSession& InSession, CADKernel::FCADFileReport& InReport);
-	CADKernel::FBody* AddBody(A3DRiBrepModel* A3DBRepModel, TMap<FString, FString> MetaData, const double InBodyScale);
-	CADKernel::FBody* GetBody(A3DRiBrepModel* A3DBRepModel);
-	const A3DRiBrepModel* GetA3DBody(CADKernel::FBody* BRepModel);
+	FTechSoftBridge(FTechSoftFileParser& InParser, UE::CADKernel::FSession& InSession);
+	UE::CADKernel::FBody* AddBody(A3DRiBrepModel* A3DBRepModel, TMap<FString, FString> MetaData, const double InBodyScale);
+	UE::CADKernel::FBody* GetBody(A3DRiBrepModel* A3DBRepModel);
+	const A3DRiBrepModel* GetA3DBody(UE::CADKernel::FBody* BRepModel);
 
 private:
 
-	void TraverseBrepData(const A3DTopoBrepData* A3DBrepData, TSharedRef<CADKernel::FBody>& Body);
-
-	void TraverseConnex(const A3DTopoConnex* A3DConnex, TSharedRef<CADKernel::FBody>& Body);
-	void TraverseShell(const A3DTopoShell* A3DShell, TSharedRef<CADKernel::FBody>& Body);
-	void AddFace(const A3DTopoFace* A3DFace, CADKernel::EOrientation Orientation, TSharedRef<CADKernel::FShell>& Body, uint32 Index);
-
-	TSharedPtr<CADKernel::FTopologicalLoop> AddLoop(const A3DTopoLoop* A3DLoop, const TSharedRef<CADKernel::FSurface>& Surface, const TechSoftUtils::FUVReparameterization& UVReparameterization, const bool bIsExternalLoop);
-	TSharedPtr<CADKernel::FTopologicalEdge> AddEdge(const A3DTopoCoEdge* A3DCoedge, const TSharedRef<CADKernel::FSurface>& Surface, const TechSoftUtils::FUVReparameterization& UVReparameterization, CADKernel::EOrientation& OutOrientation);
-
-	TSharedPtr<CADKernel::FSurface> AddSurface(const A3DSurfBase* A3DSurface, TechSoftUtils::FUVReparameterization& OutUVReparameterization);
-	TSharedPtr<CADKernel::FSurface> AddConeSurface(const A3DSurfCone* A3DSurface, TechSoftUtils::FUVReparameterization& OutUVReparameterization);
-	TSharedPtr<CADKernel::FSurface> AddCylinderSurface(const A3DSurfCylinder* A3DSurface, TechSoftUtils::FUVReparameterization& OutUVReparameterization);
-	TSharedPtr<CADKernel::FSurface> AddLinearTransfoSurface(const A3DSurfBase* A3DSurface, TechSoftUtils::FUVReparameterization& OutUVReparameterization);
-	TSharedPtr<CADKernel::FSurface> AddNurbsSurface(const A3DSurfNurbs* A3DSurface, TechSoftUtils::FUVReparameterization& OutUVReparameterization);
-	TSharedPtr<CADKernel::FSurface> AddOffsetSurface(const A3DSurfOffset* A3DSurface, TechSoftUtils::FUVReparameterization& OutUVReparameterization);
-	TSharedPtr<CADKernel::FSurface> AddPlaneSurface(const A3DSurfPlane* A3DSurface, TechSoftUtils::FUVReparameterization& OutUVReparameterization);
-	TSharedPtr<CADKernel::FSurface> AddRevolutionSurface(const A3DSurfRevolution* A3DSurface, TechSoftUtils::FUVReparameterization& OutUVReparameterization);
-	TSharedPtr<CADKernel::FSurface> AddRuledSurface(const A3DSurfRuled* A3DSurface, TechSoftUtils::FUVReparameterization& OutUVReparameterization);
-	TSharedPtr<CADKernel::FSurface> AddSphereSurface(const A3DSurfSphere* A3DSurface, TechSoftUtils::FUVReparameterization& OutUVReparameterization);
-	TSharedPtr<CADKernel::FSurface> AddTorusSurface(const A3DSurfTorus* A3DSurface, TechSoftUtils::FUVReparameterization& OutUVReparameterization);
-
-	TSharedPtr<CADKernel::FSurface> AddBlend01Surface(const A3DSurfBlend01* A3DSurface, TechSoftUtils::FUVReparameterization& OutUVReparameterization);
-	TSharedPtr<CADKernel::FSurface> AddBlend02Surface(const A3DSurfBlend02* A3DSurface, TechSoftUtils::FUVReparameterization& OutUVReparameterization);
-	TSharedPtr<CADKernel::FSurface> AddBlend03Surface(const A3DSurfBlend03* A3DSurface, TechSoftUtils::FUVReparameterization& OutUVReparameterization);
-	TSharedPtr<CADKernel::FSurface> AddCylindricalSurface(const A3DSurfCylindrical* A3DSurface, TechSoftUtils::FUVReparameterization& OutUVReparameterization);
-	TSharedPtr<CADKernel::FSurface> AddPipeSurface(const A3DSurfPipe* A3DSurface, TechSoftUtils::FUVReparameterization& OutUVReparameterization);
-	TSharedPtr<CADKernel::FSurface> AddExtrusionSurface(const A3DSurfExtrusion* A3DSurface, TechSoftUtils::FUVReparameterization& OutUVReparameterization);
-	TSharedPtr<CADKernel::FSurface> AddSurfaceFromCurves(const A3DSurfFromCurves* A3DSurface, TechSoftUtils::FUVReparameterization& OutUVReparameterization);
-	TSharedPtr<CADKernel::FSurface> AddTransformSurface(const A3DSurfTransform* A3DSurface, TechSoftUtils::FUVReparameterization& OutUVReparameterization);
-
-	TSharedPtr<CADKernel::FSurface> AddSurfaceAsNurbs(const A3DSurfBase* A3DSurface, TechSoftUtils::FUVReparameterization& OutUVReparameterization);
-	TSharedPtr<CADKernel::FSurface> AddSurfaceNurbs(const A3DSurfNurbsData& A3DNurbsData, TechSoftUtils::FUVReparameterization& OutUVReparameterization);
-
-	TSharedPtr<CADKernel::FRestrictionCurve> AddRestrictionCurve(const A3DCrvBase* A3DCurve, const TSharedRef<CADKernel::FSurface>& Surface);
-
-	TSharedPtr<CADKernel::FCurve> AddCurve(const A3DCrvBase* A3DCurve, const TechSoftUtils::FUVReparameterization& UVReparameterization);
-
-	TSharedPtr<CADKernel::FCurve> AddCurveCircle(const A3DCrvCircle* A3DCurve, const TechSoftUtils::FUVReparameterization& UVReparameterization);
-	TSharedPtr<CADKernel::FCurve> AddCurveComposite(const A3DCrvComposite* A3DCurve, const TechSoftUtils::FUVReparameterization& UVReparameterization);
-	TSharedPtr<CADKernel::FCurve> AddCurveEllipse(const A3DCrvEllipse* A3DCurve, const TechSoftUtils::FUVReparameterization& UVReparameterization);
-	TSharedPtr<CADKernel::FCurve> AddCurveHelix(const A3DCrvHelix* A3DCurve, const TechSoftUtils::FUVReparameterization& UVReparameterization);
-	TSharedPtr<CADKernel::FCurve> AddCurveHyperbola(const A3DCrvHyperbola* A3DCurve, const TechSoftUtils::FUVReparameterization& UVReparameterization);
-	TSharedPtr<CADKernel::FCurve> AddCurveLine(const A3DCrvLine* A3DCurve, const TechSoftUtils::FUVReparameterization& UVReparameterization);
-	TSharedPtr<CADKernel::FCurve> AddCurveNurbs(const A3DCrvNurbs* A3DCurve, const TechSoftUtils::FUVReparameterization& UVReparameterization);
-	TSharedPtr<CADKernel::FCurve> AddCurveParabola(const A3DCrvParabola* A3DCurve, const TechSoftUtils::FUVReparameterization& UVReparameterization);
-	TSharedPtr<CADKernel::FCurve> AddCurvePolyLine(const A3DCrvPolyLine* A3DCurve, const TechSoftUtils::FUVReparameterization& UVReparameterization);
-
-	TSharedPtr<CADKernel::FCurve> AddCurveBlend02Boundary(const A3DCrvBlend02Boundary* A3DCurve, const TechSoftUtils::FUVReparameterization& UVReparameterization);
-	TSharedPtr<CADKernel::FCurve> AddCurveEquation(const A3DCrvEquation* A3DCurve, const TechSoftUtils::FUVReparameterization& UVReparameterization);
-	TSharedPtr<CADKernel::FCurve> AddCurveIntersection(const A3DCrvIntersection* A3DCurve, const TechSoftUtils::FUVReparameterization& UVReparameterization);
-	TSharedPtr<CADKernel::FCurve> AddCurveOffset(const A3DCrvOffset* A3DCurve, const TechSoftUtils::FUVReparameterization& UVReparameterization);
-	TSharedPtr<CADKernel::FCurve> AddCurveOnSurf(const A3DCrvOnSurf* A3DCurve, const TechSoftUtils::FUVReparameterization& UVReparameterization);
-	TSharedPtr<CADKernel::FCurve> AddCurveTransform(const A3DCrvTransform* A3DCurve, const TechSoftUtils::FUVReparameterization& UVReparameterization);
-
-	TSharedPtr<CADKernel::FCurve> AddCurveAsNurbs(const A3DCrvBase* A3DCurve, const TechSoftUtils::FUVReparameterization& UVReparameterization);
-
-	void AddMetadata(FArchiveCADObject& EntityData, CADKernel::FTopologicalShapeEntity& Entity);
+	void TraverseBrepData(const A3DTopoBrepData* A3DBrepData, TSharedRef<UE::CADKernel::FBody>& Body);
+
+	void TraverseConnex(const A3DTopoConnex* A3DConnex, TSharedRef<UE::CADKernel::FBody>& Body);
+	void TraverseShell(const A3DTopoShell* A3DShell, TSharedRef<UE::CADKernel::FBody>& Body);
+	void AddFace(const A3DTopoFace* A3DFace, UE::CADKernel::EOrientation Orientation, TSharedRef<UE::CADKernel::FShell>& Body, uint32 Index);
+
+	TSharedPtr<UE::CADKernel::FTopologicalLoop> AddLoop(const A3DTopoLoop* A3DLoop, const TSharedRef<UE::CADKernel::FSurface>& Surface, const TechSoftUtils::FUVReparameterization& UVReparameterization, const bool bIsExternalLoop);
+	TSharedPtr<UE::CADKernel::FTopologicalEdge> AddEdge(const A3DTopoCoEdge* A3DCoedge, const TSharedRef<UE::CADKernel::FSurface>& Surface, const TechSoftUtils::FUVReparameterization& UVReparameterization, UE::CADKernel::EOrientation& OutOrientation);
+
+	TSharedPtr<UE::CADKernel::FSurface> AddSurface(const A3DSurfBase* A3DSurface, TechSoftUtils::FUVReparameterization& OutUVReparameterization);
+	TSharedPtr<UE::CADKernel::FSurface> AddConeSurface(const A3DSurfCone* A3DSurface, TechSoftUtils::FUVReparameterization& OutUVReparameterization);
+	TSharedPtr<UE::CADKernel::FSurface> AddCylinderSurface(const A3DSurfCylinder* A3DSurface, TechSoftUtils::FUVReparameterization& OutUVReparameterization);
+	TSharedPtr<UE::CADKernel::FSurface> AddLinearTransfoSurface(const A3DSurfBase* A3DSurface, TechSoftUtils::FUVReparameterization& OutUVReparameterization);
+	TSharedPtr<UE::CADKernel::FSurface> AddNurbsSurface(const A3DSurfNurbs* A3DSurface, TechSoftUtils::FUVReparameterization& OutUVReparameterization);
+	TSharedPtr<UE::CADKernel::FSurface> AddOffsetSurface(const A3DSurfOffset* A3DSurface, TechSoftUtils::FUVReparameterization& OutUVReparameterization);
+	TSharedPtr<UE::CADKernel::FSurface> AddPlaneSurface(const A3DSurfPlane* A3DSurface, TechSoftUtils::FUVReparameterization& OutUVReparameterization);
+	TSharedPtr<UE::CADKernel::FSurface> AddRevolutionSurface(const A3DSurfRevolution* A3DSurface, TechSoftUtils::FUVReparameterization& OutUVReparameterization);
+	TSharedPtr<UE::CADKernel::FSurface> AddRuledSurface(const A3DSurfRuled* A3DSurface, TechSoftUtils::FUVReparameterization& OutUVReparameterization);
+	TSharedPtr<UE::CADKernel::FSurface> AddSphereSurface(const A3DSurfSphere* A3DSurface, TechSoftUtils::FUVReparameterization& OutUVReparameterization);
+	TSharedPtr<UE::CADKernel::FSurface> AddTorusSurface(const A3DSurfTorus* A3DSurface, TechSoftUtils::FUVReparameterization& OutUVReparameterization);
+
+	TSharedPtr<UE::CADKernel::FSurface> AddBlend01Surface(const A3DSurfBlend01* A3DSurface, TechSoftUtils::FUVReparameterization& OutUVReparameterization);
+	TSharedPtr<UE::CADKernel::FSurface> AddBlend02Surface(const A3DSurfBlend02* A3DSurface, TechSoftUtils::FUVReparameterization& OutUVReparameterization);
+	TSharedPtr<UE::CADKernel::FSurface> AddBlend03Surface(const A3DSurfBlend03* A3DSurface, TechSoftUtils::FUVReparameterization& OutUVReparameterization);
+	TSharedPtr<UE::CADKernel::FSurface> AddCylindricalSurface(const A3DSurfCylindrical* A3DSurface, TechSoftUtils::FUVReparameterization& OutUVReparameterization);
+	TSharedPtr<UE::CADKernel::FSurface> AddPipeSurface(const A3DSurfPipe* A3DSurface, TechSoftUtils::FUVReparameterization& OutUVReparameterization);
+	TSharedPtr<UE::CADKernel::FSurface> AddExtrusionSurface(const A3DSurfExtrusion* A3DSurface, TechSoftUtils::FUVReparameterization& OutUVReparameterization);
+	TSharedPtr<UE::CADKernel::FSurface> AddSurfaceFromCurves(const A3DSurfFromCurves* A3DSurface, TechSoftUtils::FUVReparameterization& OutUVReparameterization);
+	TSharedPtr<UE::CADKernel::FSurface> AddTransformSurface(const A3DSurfTransform* A3DSurface, TechSoftUtils::FUVReparameterization& OutUVReparameterization);
+
+	TSharedPtr<UE::CADKernel::FSurface> AddSurfaceAsNurbs(const A3DSurfBase* A3DSurface, TechSoftUtils::FUVReparameterization& OutUVReparameterization);
+	TSharedPtr<UE::CADKernel::FSurface> AddSurfaceNurbs(const A3DSurfNurbsData& A3DNurbsData, TechSoftUtils::FUVReparameterization& OutUVReparameterization);
+
+	TSharedPtr<UE::CADKernel::FRestrictionCurve> AddRestrictionCurve(const A3DCrvBase* A3DCurve, const TSharedRef<UE::CADKernel::FSurface>& Surface);
+
+	TSharedPtr<UE::CADKernel::FCurve> AddCurve(const A3DCrvBase* A3DCurve, const TechSoftUtils::FUVReparameterization& UVReparameterization);
+
+	TSharedPtr<UE::CADKernel::FCurve> AddCurveCircle(const A3DCrvCircle* A3DCurve, const TechSoftUtils::FUVReparameterization& UVReparameterization);
+	TSharedPtr<UE::CADKernel::FCurve> AddCurveComposite(const A3DCrvComposite* A3DCurve, const TechSoftUtils::FUVReparameterization& UVReparameterization);
+	TSharedPtr<UE::CADKernel::FCurve> AddCurveEllipse(const A3DCrvEllipse* A3DCurve, const TechSoftUtils::FUVReparameterization& UVReparameterization);
+	TSharedPtr<UE::CADKernel::FCurve> AddCurveHelix(const A3DCrvHelix* A3DCurve, const TechSoftUtils::FUVReparameterization& UVReparameterization);
+	TSharedPtr<UE::CADKernel::FCurve> AddCurveHyperbola(const A3DCrvHyperbola* A3DCurve, const TechSoftUtils::FUVReparameterization& UVReparameterization);
+	TSharedPtr<UE::CADKernel::FCurve> AddCurveLine(const A3DCrvLine* A3DCurve, const TechSoftUtils::FUVReparameterization& UVReparameterization);
+	TSharedPtr<UE::CADKernel::FCurve> AddCurveNurbs(const A3DCrvNurbs* A3DCurve, const TechSoftUtils::FUVReparameterization& UVReparameterization);
+	TSharedPtr<UE::CADKernel::FCurve> AddCurveParabola(const A3DCrvParabola* A3DCurve, const TechSoftUtils::FUVReparameterization& UVReparameterization);
+	TSharedPtr<UE::CADKernel::FCurve> AddCurvePolyLine(const A3DCrvPolyLine* A3DCurve, const TechSoftUtils::FUVReparameterization& UVReparameterization);
+
+	TSharedPtr<UE::CADKernel::FCurve> AddCurveBlend02Boundary(const A3DCrvBlend02Boundary* A3DCurve, const TechSoftUtils::FUVReparameterization& UVReparameterization);
+	TSharedPtr<UE::CADKernel::FCurve> AddCurveEquation(const A3DCrvEquation* A3DCurve, const TechSoftUtils::FUVReparameterization& UVReparameterization);
+	TSharedPtr<UE::CADKernel::FCurve> AddCurveIntersection(const A3DCrvIntersection* A3DCurve, const TechSoftUtils::FUVReparameterization& UVReparameterization);
+	TSharedPtr<UE::CADKernel::FCurve> AddCurveOffset(const A3DCrvOffset* A3DCurve, const TechSoftUtils::FUVReparameterization& UVReparameterization);
+	TSharedPtr<UE::CADKernel::FCurve> AddCurveOnSurf(const A3DCrvOnSurf* A3DCurve, const TechSoftUtils::FUVReparameterization& UVReparameterization);
+	TSharedPtr<UE::CADKernel::FCurve> AddCurveTransform(const A3DCrvTransform* A3DCurve, const TechSoftUtils::FUVReparameterization& UVReparameterization);
+
+	TSharedPtr<UE::CADKernel::FCurve> AddCurveAsNurbs(const A3DCrvBase* A3DCurve, const TechSoftUtils::FUVReparameterization& UVReparameterization);
+
+	void AddMetadata(FArchiveCADObject& EntityData, UE::CADKernel::FTopologicalShapeEntity& Entity);
 };
 }
 
diff --git a/Engine/Plugins/Enterprise/DatasmithCADImporter/Source/CADInterfaces/Private/TechSoftFileParser.cpp b/Engine/Plugins/Enterprise/DatasmithCADImporter/Source/CADInterfaces/Private/TechSoftFileParser.cpp
--- a/Engine/Plugins/Enterprise/DatasmithCADImporter/Source/CADInterfaces/Private/TechSoftFileParser.cpp
+++ b/Engine/Plugins/Enterprise/DatasmithCADImporter/Source/CADInterfaces/Private/TechSoftFileParser.cpp
@@ -400,7 +400,7 @@ ECADParsingResult FTechSoftFileParser::Process()
 		}
 	}
 
-	// Adapt BRep to CADKernel
+	// Adapt BRep to UE::CADKernel
 	if (AdaptBRepModel() != A3D_SUCCESS)
 	{
 		return ECADParsingResult::ProcessFailed;
diff --git a/Engine/Plugins/Enterprise/DatasmithCADImporter/Source/CADInterfaces/Private/TechSoftFileParser.h b/Engine/Plugins/Enterprise/DatasmithCADImporter/Source/CADInterfaces/Private/TechSoftFileParser.h
--- a/Engine/Plugins/Enterprise/DatasmithCADImporter/Source/CADInterfaces/Private/TechSoftFileParser.h
+++ b/Engine/Plugins/Enterprise/DatasmithCADImporter/Source/CADInterfaces/Private/TechSoftFileParser.h
@@ -58,7 +58,7 @@ protected:
 
 	/**
 	 * If the tessellator is TechSoft, SewModel call TechSoftInterface::SewModel
-	 * If the tessellator is CADKernel, SewModel do nothing as the file is not yet parsed. In this case, the sew is done in GenerateBodyMeshes.
+	 * If the tessellator is UE::CADKernel, SewModel do nothing as the file is not yet parsed. In this case, the sew is done in GenerateBodyMeshes.
 	 */
 	virtual void SewModel();
 
diff --git a/Engine/Plugins/Enterprise/DatasmithCADImporter/Source/CADInterfaces/Private/TechSoftFileParserCADKernelTessellator.cpp b/Engine/Plugins/Enterprise/DatasmithCADImporter/Source/CADInterfaces/Private/TechSoftFileParserCADKernelTessellator.cpp
--- a/Engine/Plugins/Enterprise/DatasmithCADImporter/Source/CADInterfaces/Private/TechSoftFileParserCADKernelTessellator.cpp
+++ b/Engine/Plugins/Enterprise/DatasmithCADImporter/Source/CADInterfaces/Private/TechSoftFileParserCADKernelTessellator.cpp
@@ -72,14 +72,15 @@ void FTechSoftFileParserCADKernelTessellator::SewAndGenerateBodyMeshes()
 
 void FTechSoftFileParserCADKernelTessellator::SewAndMesh(TArray<A3DRiRepresentationItem*>& Representations)
 {
+	using namespace UE::CADKernel;
+
 	const double GeometricTolerance = 0.01; // mm
 
-	CADKernel::FSession CADKernelSession(GeometricTolerance);
+	FSession CADKernelSession(GeometricTolerance);
 	CADKernelSession.SetFirstNewHostId(LastHostIdUsed);
-	CADKernel::FModel& CADKernelModel = CADKernelSession.GetModel();
+	FModel& CADKernelModel = CADKernelSession.GetModel();
 
-	CADKernel::FCADFileReport Report;
-	FTechSoftBridge TechSoftBridge(*this, CADKernelSession, Report);
+	FTechSoftBridge TechSoftBridge(*this, CADKernelSession);
 
 	for (A3DRiRepresentationItem* Representation : Representations)
 	{
@@ -90,11 +91,11 @@ void FTechSoftFileParserCADKernelTessellator::SewAndMesh(TArray<A3DRiRepresentat
 		}
 		FArchiveBody& Body = SceneGraph.GetBody(*BodyIndex);
 
-		CADKernel::FBody* CADKernelBody = TechSoftBridge.AddBody(Representation, Body.MetaData, Body.Unit);
+		FBody* CADKernelBody = TechSoftBridge.AddBody(Representation, Body.MetaData, Body.Unit);
 	}
 
 	// Sew if needed
-	CADKernel::FTopomaker Topomaker(CADKernelSession, GeometricTolerance);
+	FTopomaker Topomaker(CADKernelSession, GeometricTolerance);
 	Topomaker.Sew();
 	Topomaker.SplitIntoConnectedShells();
 	Topomaker.OrientShells();
@@ -107,13 +108,13 @@ void FTechSoftFileParserCADKernelTessellator::SewAndMesh(TArray<A3DRiRepresentat
 	int32 BodyCount = CADKernelModel.GetBodies().Num();
 	int32 DeletedBodyCount = Representations.Num() - BodyCount;
 
-	TArray<CADKernel::FBody*> NewBodies;
-	TArray<CADKernel::FBody*> ExistingBodies;
+	TArray<FBody*> NewBodies;
+	TArray<FBody*> ExistingBodies;
 	NewBodies.Reserve(BodyCount);
 	ExistingBodies.Reserve(BodyCount);
 
 	// find new and existing bodies
-	for (const TSharedPtr<CADKernel::FBody>& CADKernelBody : CADKernelModel.GetBodies())
+	for (const TSharedPtr<FBody>& CADKernelBody : CADKernelModel.GetBodies())
 	{
 		if (!CADKernelBody.IsValid())
 		{
@@ -138,7 +139,7 @@ void FTechSoftFileParserCADKernelTessellator::SewAndMesh(TArray<A3DRiRepresentat
 
 	for (A3DRiRepresentationItem* Representation : Representations)
 	{
-		CADKernel::FBody* Body = TechSoftBridge.GetBody(Representation);
+		FBody* Body = TechSoftBridge.GetBody(Representation);
 		if (Body == nullptr)
 		{
 			FCadId* ArchiveBodyId = RepresentationItemsCache.Find(Representation);
@@ -151,7 +152,7 @@ void FTechSoftFileParserCADKernelTessellator::SewAndMesh(TArray<A3DRiRepresentat
 	}
 
 	// Process existing bodies
-	for (CADKernel::FBody* Body : ExistingBodies)
+	for (FBody* Body : ExistingBodies)
 	{
 		const A3DRiRepresentationItem* Representation = TechSoftBridge.GetA3DBody(Body);
 
@@ -167,7 +168,7 @@ void FTechSoftFileParserCADKernelTessellator::SewAndMesh(TArray<A3DRiRepresentat
 
 	// Process new bodies
 	int32 Index = 0;
-	for (CADKernel::FBody* Body : NewBodies)
+	for (FBody* Body : NewBodies)
 	{
 		FCadId ArchiveBodyId = ArchiveBodyIdOfDeletedRepresentation[Index++];
 		FArchiveBody& ArchiveBody = SceneGraph.GetBody(ArchiveBodyId);
@@ -185,21 +186,22 @@ void FTechSoftFileParserCADKernelTessellator::SewAndMesh(TArray<A3DRiRepresentat
 
 void FTechSoftFileParserCADKernelTessellator::GenerateBodyMesh(A3DRiRepresentationItem* Representation, FArchiveBody& ArchiveBody)
 {
+	using namespace UE::CADKernel;
+
 	const double GeometricTolerance = 0.01; // mm
 
-	CADKernel::FSession CADKernelSession(GeometricTolerance);
+	FSession CADKernelSession(GeometricTolerance);
 	CADKernelSession.SetFirstNewHostId(LastHostIdUsed);
-	CADKernel::FModel& CADKernelModel = CADKernelSession.GetModel();
+	FModel& CADKernelModel = CADKernelSession.GetModel();
 
-	CADKernel::FCADFileReport Report;
-	FTechSoftBridge TechSoftBridge(*this, CADKernelSession, Report);
+	FTechSoftBridge TechSoftBridge(*this, CADKernelSession);
 
-	CADKernel::FBody* CADKernelBody = TechSoftBridge.AddBody(Representation, ArchiveBody.MetaData, ArchiveBody.Unit);
+	FBody* CADKernelBody = TechSoftBridge.AddBody(Representation, ArchiveBody.MetaData, ArchiveBody.Unit);
 
 	if (CADFileData.GetImportParameters().GetStitchingTechnique() == StitchingHeal)
 	{
 		const double Tolerance = 0.1; // mm
-		CADKernel::FTopomaker Topomaker(CADKernelSession, Tolerance);
+		FTopomaker Topomaker(CADKernelSession, Tolerance);
 		Topomaker.Sew();
 		Topomaker.OrientShells();
 	}
@@ -207,8 +209,10 @@ void FTechSoftFileParserCADKernelTessellator::GenerateBodyMesh(A3DRiRepresentati
 	MeshAndGetTessellation(CADKernelSession, ArchiveBody, *CADKernelBody);
 }
 
-void FTechSoftFileParserCADKernelTessellator::MeshAndGetTessellation(CADKernel::FSession& CADKernelSession, FArchiveBody& ArchiveBody, CADKernel::FBody& CADKernelBody)
+void FTechSoftFileParserCADKernelTessellator::MeshAndGetTessellation(UE::CADKernel::FSession& CADKernelSession, FArchiveBody& ArchiveBody, UE::CADKernel::FBody& CADKernelBody)
 {
+	using namespace UE::CADKernel;
+
 	FBodyMesh& BodyMesh = CADFileData.AddBodyMesh(ArchiveBody.Id, ArchiveBody);
 	ArchiveBody.ColorFaceSet = BodyMesh.ColorSet;
 	ArchiveBody.MaterialFaceSet = BodyMesh.MaterialSet;
@@ -221,11 +225,11 @@ void FTechSoftFileParserCADKernelTessellator::MeshAndGetTessellation(CADKernel::
 	}
 
 	// Tessellate the body
-	TSharedRef<CADKernel::FModelMesh> CADKernelModelMesh = CADKernel::FEntity::MakeShared<CADKernel::FModelMesh>();
+	TSharedRef<FModelMesh> CADKernelModelMesh = FEntity::MakeShared<FModelMesh>();
 
 	FCADKernelTools::DefineMeshCriteria(CADKernelModelMesh.Get(), CADFileData.GetImportParameters(), CADKernelSession.GetGeometricTolerance());
 
-	CADKernel::FParametricMesher Mesher(*CADKernelModelMesh);
+	FParametricMesher Mesher(*CADKernelModelMesh);
 	Mesher.MeshEntity(CADKernelBody);
 
 	FCADKernelTools::GetBodyTessellation(*CADKernelModelMesh, CADKernelBody, BodyMesh);
diff --git a/Engine/Plugins/Enterprise/DatasmithCADImporter/Source/CADInterfaces/Private/TechSoftFileParserCADKernelTessellator.h b/Engine/Plugins/Enterprise/DatasmithCADImporter/Source/CADInterfaces/Private/TechSoftFileParserCADKernelTessellator.h
--- a/Engine/Plugins/Enterprise/DatasmithCADImporter/Source/CADInterfaces/Private/TechSoftFileParserCADKernelTessellator.h
+++ b/Engine/Plugins/Enterprise/DatasmithCADImporter/Source/CADInterfaces/Private/TechSoftFileParserCADKernelTessellator.h
@@ -5,7 +5,7 @@
 
 #include "TechSoftFileParser.h"
 
-namespace CADKernel
+namespace UE::CADKernel
 {
 class FBody;
 class FSession;
@@ -38,7 +38,7 @@ private:
 	// Tessellation methods
 	virtual void GenerateBodyMeshes() override;
 	virtual void GenerateBodyMesh(A3DRiRepresentationItem* Representation, FArchiveBody& Body) override;
-	void MeshAndGetTessellation(CADKernel::FSession& CADKernelSession, FArchiveBody& ArchiveBody, CADKernel::FBody& CADKernelBody);
+	void MeshAndGetTessellation(UE::CADKernel::FSession& CADKernelSession, FArchiveBody& ArchiveBody, UE::CADKernel::FBody& CADKernelBody);
 
 	void SewAndGenerateBodyMeshes();
 	void SewAndMesh(TArray<A3DRiRepresentationItem*>& Representations);
diff --git a/Engine/Plugins/Enterprise/DatasmithCADImporter/Source/CADKernelSurface/Private/CADKernelSurfaceExtension.cpp b/Engine/Plugins/Enterprise/DatasmithCADImporter/Source/CADKernelSurface/Private/CADKernelSurfaceExtension.cpp
--- a/Engine/Plugins/Enterprise/DatasmithCADImporter/Source/CADKernelSurface/Private/CADKernelSurfaceExtension.cpp
+++ b/Engine/Plugins/Enterprise/DatasmithCADImporter/Source/CADKernelSurface/Private/CADKernelSurfaceExtension.cpp
@@ -31,6 +31,8 @@ void UCADKernelParametricSurfaceData::Serialize(FArchive& Ar)
 
 bool UCADKernelParametricSurfaceData::Tessellate(UStaticMesh& StaticMesh, const FDatasmithRetessellationOptions& RetessellateOptions)
 {
+	using namespace UE::CADKernel;
+
 	bool bSuccessfulTessellation = false;
 
 #if WITH_EDITOR
@@ -57,11 +59,11 @@ bool UCADKernelParametricSurfaceData::Tessellate(UStaticMesh& StaticMesh, const
 		}
 
 		const double GeometricTolerance = 0.01; // mm
-		TSharedRef<CADKernel::FSession> CADKernelSession = MakeShared<CADKernel::FSession>(GeometricTolerance);
+		TSharedRef<FSession> CADKernelSession = MakeShared<FSession>(GeometricTolerance);
 		CADKernelSession->AddDatabase(RawData);
 
-		CADKernel::FModel& CADKernelModel = CADKernelSession->GetModel();
-		TArray<TSharedPtr<CADKernel::FBody>> CADKernelBodies = CADKernelModel.GetBodies();
+		FModel& CADKernelModel = CADKernelSession->GetModel();
+		TArray<TSharedPtr<FBody>> CADKernelBodies = CADKernelModel.GetBodies();
 		if (CADKernelBodies.Num() != 1)
 		{
 			return bSuccessfulTessellation;
diff --git a/Engine/Plugins/Enterprise/DatasmithCADImporter/Source/CADKernelSurface/Public/CADModelToCADKernelConverterBase.h b/Engine/Plugins/Enterprise/DatasmithCADImporter/Source/CADKernelSurface/Public/CADModelToCADKernelConverterBase.h
--- a/Engine/Plugins/Enterprise/DatasmithCADImporter/Source/CADKernelSurface/Public/CADModelToCADKernelConverterBase.h
+++ b/Engine/Plugins/Enterprise/DatasmithCADImporter/Source/CADKernelSurface/Public/CADModelToCADKernelConverterBase.h
@@ -41,7 +41,7 @@ public:
 		{
 			// the joining tolerance is set to 0.1 mm until the user can specify it
 			const double JoiningTolerance = 0.1;
-			CADKernel::FTopomaker Topomaker(CADKernelSession, JoiningTolerance);
+			UE::CADKernel::FTopomaker Topomaker(CADKernelSession, JoiningTolerance);
 			Topomaker.Sew();
 			Topomaker.OrientShells();
 		}
@@ -59,7 +59,7 @@ public:
 
 	virtual bool Tessellate(const CADLibrary::FMeshParameters& InMeshParameters, FMeshDescription& OutMeshDescription) override
 	{
-		CADKernel::FModel& Model = CADKernelSession.GetModel();
+		UE::CADKernel::FModel& Model = CADKernelSession.GetModel();
 		return CADLibrary::FCADKernelTools::Tessellate(Model, ImportParameters, InMeshParameters, OutMeshDescription);
 	}
 
@@ -80,7 +80,7 @@ public:
 
 protected:
 
-	CADKernel::FSession CADKernelSession;
+	UE::CADKernel::FSession CADKernelSession;
 
 	CADLibrary::FImportParameters ImportParameters;
 	double GeometricTolerance;
diff --git a/Engine/Plugins/Enterprise/DatasmithCADImporter/Source/CADLibrary/Private/CADKernelTools.cpp b/Engine/Plugins/Enterprise/DatasmithCADImporter/Source/CADLibrary/Private/CADKernelTools.cpp
--- a/Engine/Plugins/Enterprise/DatasmithCADImporter/Source/CADLibrary/Private/CADKernelTools.cpp
+++ b/Engine/Plugins/Enterprise/DatasmithCADImporter/Source/CADLibrary/Private/CADKernelTools.cpp
@@ -46,7 +46,7 @@ namespace CADLibrary
 		}
 	};
 
-	static void FillVertexPosition(FMeshConversionContext& Context, const TSharedRef<CADKernel::FModelMesh>& ModelMesh, FMeshDescription& MeshDescription)
+	static void FillVertexPosition(FMeshConversionContext& Context, const TSharedRef<UE::CADKernel::FModelMesh>& ModelMesh, FMeshDescription& MeshDescription)
 	{
 		TArray<FVector3f> VertexArray;
 		ModelMesh->GetNodeCoordinates(VertexArray);
@@ -92,9 +92,9 @@ namespace CADLibrary
 		}
 	}
 
-	bool FillMesh(FMeshConversionContext& Context, const TSharedRef<CADKernel::FModelMesh>& ModelMesh, FMeshDescription& MeshDescription)
+	bool FillMesh(FMeshConversionContext& Context, const TSharedRef<UE::CADKernel::FModelMesh>& ModelMesh, FMeshDescription& MeshDescription)
 	{
-		using namespace CADKernel;
+		using namespace UE::CADKernel;
 
 		const int32 UVChannel = 0;
 		const int32 TriangleCount = 3;
@@ -257,7 +257,7 @@ namespace CADLibrary
 		return true;
 	}
 
-	static bool ConvertModelMeshToMeshDescription(FMeshConversionContext& Context, const TSharedRef<CADKernel::FModelMesh>& InModelMesh, FMeshDescription& MeshDescription)
+	static bool ConvertModelMeshToMeshDescription(FMeshConversionContext& Context, const TSharedRef<UE::CADKernel::FModelMesh>& InModelMesh, FMeshDescription& MeshDescription)
 	{
 		int32 VertexCount = InModelMesh->GetVertexCount();
 		int32 TriangleCount = InModelMesh->GetTriangleCount();
@@ -278,9 +278,9 @@ namespace CADLibrary
 		return MeshDescription.Polygons().Num() > 0;
 	}
 
-	bool FCADKernelTools::Tessellate(CADKernel::FTopologicalShapeEntity& CADTopologicalEntity, const FImportParameters& ImportParameters, const FMeshParameters& MeshParameters, FMeshDescription& OutMeshDescription)
+	bool FCADKernelTools::Tessellate(UE::CADKernel::FTopologicalShapeEntity& CADTopologicalEntity, const FImportParameters& ImportParameters, const FMeshParameters& MeshParameters, FMeshDescription& OutMeshDescription)
 	{
-		using namespace CADKernel;
+		using namespace UE::CADKernel;
 
 		// Tessellate the model
 		TSharedRef<FModelMesh> CADKernelModelMesh = FEntity::MakeShared<FModelMesh>();
@@ -295,7 +295,7 @@ namespace CADLibrary
 	}
 
 
-	uint32 FCADKernelTools::GetFaceTessellation(CADKernel::FFaceMesh& FaceMesh, FBodyMesh& OutBodyMesh, FObjectDisplayDataId FaceMaterial)
+	uint32 FCADKernelTools::GetFaceTessellation(UE::CADKernel::FFaceMesh& FaceMesh, FBodyMesh& OutBodyMesh, FObjectDisplayDataId FaceMaterial)
 	{
 		// Something wrong happened, either an error or no data to collect
 		if (FaceMesh.TrianglesVerticesIndex.Num() == 0)
@@ -305,7 +305,7 @@ namespace CADLibrary
 
 		FTessellationData& Tessellation = OutBodyMesh.Faces.Emplace_GetRef();
 
-		const CADKernel::FTopologicalFace& TopologicalFace = (const CADKernel::FTopologicalFace&) FaceMesh.GetGeometricEntity();
+		const UE::CADKernel::FTopologicalFace& TopologicalFace = (const UE::CADKernel::FTopologicalFace&) FaceMesh.GetGeometricEntity();
 		Tessellation.PatchId = TopologicalFace.GetPatchId();
 
 		Tessellation.PositionIndices = MoveTemp(FaceMesh.VerticesGlobalIndex);
@@ -329,7 +329,7 @@ namespace CADLibrary
 		return Tessellation.VertexIndices.Num() / 3;
 	}
 
-	void GetDisplayDataIds(const CADKernel::FTopologicalShapeEntity& ShapeEntity, FObjectDisplayDataId& DisplayDataId)
+	void GetDisplayDataIds(const UE::CADKernel::FTopologicalShapeEntity& ShapeEntity, FObjectDisplayDataId& DisplayDataId)
 	{
 		if (ShapeEntity.GetColorId() != 0)
 		{
@@ -341,7 +341,7 @@ namespace CADLibrary
 		}
 	}
 
-	void FCADKernelTools::GetBodyTessellation(const CADKernel::FModelMesh& ModelMesh, const CADKernel::FBody& Body, FBodyMesh& OutBodyMesh)
+	void FCADKernelTools::GetBodyTessellation(const UE::CADKernel::FModelMesh& ModelMesh, const UE::CADKernel::FBody& Body, FBodyMesh& OutBodyMesh)
 	{
 		ModelMesh.GetNodeCoordinates(OutBodyMesh.VertexArray);
 
@@ -363,7 +363,7 @@ namespace CADLibrary
 		BodyMaterial.DefaultMaterialUId = BodyMaterial.MaterialUId != 0 ? BodyMaterial.MaterialUId : BodyMaterial.ColorUId;
 
 		// Loop through the face of bodies and collect all tessellation data
-		for (const TSharedPtr<CADKernel::FShell>& Shell : Body.GetShells())
+		for (const TSharedPtr<UE::CADKernel::FShell>& Shell : Body.GetShells())
 		{
 			if (!Shell.IsValid())
 			{
@@ -373,7 +373,7 @@ namespace CADLibrary
 			FObjectDisplayDataId ShellMaterial = BodyMaterial;
 			GetDisplayDataIds(*Shell, ShellMaterial);
 
-			for (const CADKernel::FOrientedFace& Face : Shell->GetFaces())
+			for (const UE::CADKernel::FOrientedFace& Face : Shell->GetFaces())
 			{
 				if (!Face.Entity.IsValid())
 				{
@@ -394,31 +394,31 @@ namespace CADLibrary
 		}
 	}
 
-	void FCADKernelTools::DefineMeshCriteria(CADKernel::FModelMesh& MeshModel, const FImportParameters& ImportParameters, double GeometricTolerance)
+	void FCADKernelTools::DefineMeshCriteria(UE::CADKernel::FModelMesh& MeshModel, const FImportParameters& ImportParameters, double GeometricTolerance)
 	{
 		{
-			TSharedPtr<CADKernel::FCriterion> CurvatureCriterion = CADKernel::FCriterion::CreateCriterion(CADKernel::ECriterion::CADCurvature);
+			TSharedPtr<UE::CADKernel::FCriterion> CurvatureCriterion = UE::CADKernel::FCriterion::CreateCriterion(UE::CADKernel::ECriterion::CADCurvature);
 			MeshModel.AddCriterion(CurvatureCriterion);
 
-			TSharedPtr<CADKernel::FCriterion> MinSizeCriterion = CADKernel::FCriterion::CreateCriterion(CADKernel::ECriterion::MinSize, 2. * GeometricTolerance);
+			TSharedPtr<UE::CADKernel::FCriterion> MinSizeCriterion = UE::CADKernel::FCriterion::CreateCriterion(UE::CADKernel::ECriterion::MinSize, 2. * GeometricTolerance);
 			MeshModel.AddCriterion(MinSizeCriterion);
 		}
 
 		if (ImportParameters.GetMaxEdgeLength() > SMALL_NUMBER)
 		{
-			TSharedPtr<CADKernel::FCriterion> MaxSizeCriterion = CADKernel::FCriterion::CreateCriterion(CADKernel::ECriterion::MaxSize, ImportParameters.GetMaxEdgeLength());
+			TSharedPtr<UE::CADKernel::FCriterion> MaxSizeCriterion = UE::CADKernel::FCriterion::CreateCriterion(UE::CADKernel::ECriterion::MaxSize, ImportParameters.GetMaxEdgeLength());
 			MeshModel.AddCriterion(MaxSizeCriterion);
 		}
 
 		if (ImportParameters.GetChordTolerance() > SMALL_NUMBER)
 		{
-			TSharedPtr<CADKernel::FCriterion> ChordCriterion = CADKernel::FCriterion::CreateCriterion(CADKernel::ECriterion::Sag, ImportParameters.GetChordTolerance());
+			TSharedPtr<UE::CADKernel::FCriterion> ChordCriterion = UE::CADKernel::FCriterion::CreateCriterion(UE::CADKernel::ECriterion::Sag, ImportParameters.GetChordTolerance());
 			MeshModel.AddCriterion(ChordCriterion);
 		}
 
 		if (ImportParameters.GetMaxNormalAngle() > SMALL_NUMBER)
 		{
-			TSharedPtr<CADKernel::FCriterion> MaxNormalAngleCriterion = CADKernel::FCriterion::CreateCriterion(CADKernel::ECriterion::Angle, ImportParameters.GetMaxNormalAngle());
+			TSharedPtr<UE::CADKernel::FCriterion> MaxNormalAngleCriterion = UE::CADKernel::FCriterion::CreateCriterion(UE::CADKernel::ECriterion::Angle, ImportParameters.GetMaxNormalAngle());
 			MeshModel.AddCriterion(MaxNormalAngleCriterion);
 		}
 	}
diff --git a/Engine/Plugins/Enterprise/DatasmithCADImporter/Source/CADLibrary/Public/CADKernelTools.h b/Engine/Plugins/Enterprise/DatasmithCADImporter/Source/CADLibrary/Public/CADKernelTools.h
--- a/Engine/Plugins/Enterprise/DatasmithCADImporter/Source/CADLibrary/Public/CADKernelTools.h
+++ b/Engine/Plugins/Enterprise/DatasmithCADImporter/Source/CADLibrary/Public/CADKernelTools.h
@@ -6,7 +6,7 @@
 #include "CADData.h"
 #include "CADOptions.h"
 
-namespace CADKernel
+namespace UE::CADKernel
 {
 	class FBody;
 	class FBodyMesh;
@@ -25,15 +25,15 @@ namespace CADLibrary
 	class CADLIBRARY_API FCADKernelTools
 	{
 	public:
-		static void DefineMeshCriteria(CADKernel::FModelMesh& MeshModel, const FImportParameters& ImportParameters, double GeometricTolerance);
+		static void DefineMeshCriteria(UE::CADKernel::FModelMesh& MeshModel, const FImportParameters& ImportParameters, double GeometricTolerance);
 
-		static void GetBodyTessellation(const CADKernel::FModelMesh& ModelMesh, const CADKernel::FBody& Body, FBodyMesh& OutBodyMesh);
+		static void GetBodyTessellation(const UE::CADKernel::FModelMesh& ModelMesh, const UE::CADKernel::FBody& Body, FBodyMesh& OutBodyMesh);
 
 		/**
 		 * Tessellate a CADKernel entity and update the MeshDescription
 		 */
-		static bool Tessellate(CADKernel::FTopologicalShapeEntity& InCADKernelEntity, const FImportParameters& ImportParameters, const FMeshParameters& MeshParameters, FMeshDescription& OutMeshDescription);
+		static bool Tessellate(UE::CADKernel::FTopologicalShapeEntity& InCADKernelEntity, const FImportParameters& ImportParameters, const FMeshParameters& MeshParameters, FMeshDescription& OutMeshDescription);
 
-		static uint32 GetFaceTessellation(CADKernel::FFaceMesh& FaceMesh, FBodyMesh& OutBodyMesh, FObjectDisplayDataId FaceMaterial);
+		static uint32 GetFaceTessellation(UE::CADKernel::FFaceMesh& FaceMesh, FBodyMesh& OutBodyMesh, FObjectDisplayDataId FaceMaterial);
 	};
 }
diff --git a/Engine/Plugins/Enterprise/DatasmithCADImporter/Source/DatasmithOpenNurbsTranslator/Private/OpenNurbsBRepToCADKernelConverter.cpp b/Engine/Plugins/Enterprise/DatasmithCADImporter/Source/DatasmithOpenNurbsTranslator/Private/OpenNurbsBRepToCADKernelConverter.cpp
--- a/Engine/Plugins/Enterprise/DatasmithCADImporter/Source/DatasmithOpenNurbsTranslator/Private/OpenNurbsBRepToCADKernelConverter.cpp
+++ b/Engine/Plugins/Enterprise/DatasmithCADImporter/Source/DatasmithOpenNurbsTranslator/Private/OpenNurbsBRepToCADKernelConverter.cpp
@@ -31,11 +31,11 @@
 #include "CADKernel/Topo/TopologicalLoop.h"
 
 
-namespace
+namespace OpenNurbsBRepConverterTool
 {
 	enum EAxis { U, V };
 
-	void FillPerAxisInfo(EAxis Axis, ON_NurbsSurface& OpenNurbsSurface, CADKernel::FNurbsSurfaceHomogeneousData& OutNurbsInfo)
+	void FillPerAxisInfo(EAxis Axis, ON_NurbsSurface& OpenNurbsSurface, UE::CADKernel::FNurbsSurfaceHomogeneousData& OutNurbsInfo)
 	{
 		int32& Degree = Axis == EAxis::U ? OutNurbsInfo.UDegree : OutNurbsInfo.VDegree;
 		Degree = OpenNurbsSurface.Order(Axis) - 1;
@@ -65,11 +65,11 @@ namespace
 	}
 }
 
-TSharedPtr<CADKernel::FSurface> FOpenNurbsBRepToCADKernelConverter::AddSurface(ON_NurbsSurface& OpenNurbsSurface)
+TSharedPtr<UE::CADKernel::FSurface> FOpenNurbsBRepToCADKernelConverter::AddSurface(ON_NurbsSurface& OpenNurbsSurface)
 {
-	CADKernel::FNurbsSurfaceHomogeneousData NurbsData;
-	FillPerAxisInfo(U, OpenNurbsSurface, NurbsData);
-	FillPerAxisInfo(V, OpenNurbsSurface, NurbsData);
+	UE::CADKernel::FNurbsSurfaceHomogeneousData NurbsData;
+	OpenNurbsBRepConverterTool::FillPerAxisInfo(OpenNurbsBRepConverterTool::U, OpenNurbsSurface, NurbsData);
+	OpenNurbsBRepConverterTool::FillPerAxisInfo(OpenNurbsBRepConverterTool::V, OpenNurbsSurface, NurbsData);
 
 	{
 		int32 ControlVertexDimension = OpenNurbsSurface.CVSize();
@@ -98,12 +98,12 @@ TSharedPtr<CADKernel::FSurface> FOpenNurbsBRepToCADKernelConverter::AddSurface(O
 		}
 	}
 
-	return CADKernel::FSurface::MakeNurbsSurface(GeometricTolerance, NurbsData);
+	return UE::CADKernel::FSurface::MakeNurbsSurface(GeometricTolerance, NurbsData);
 }
 
-TSharedPtr<CADKernel::FTopologicalLoop> FOpenNurbsBRepToCADKernelConverter::AddLoop(const ON_BrepLoop& OpenNurbsLoop, TSharedPtr<CADKernel::FSurface> & CarrierSurface, const bool bIsExternal)
+TSharedPtr<UE::CADKernel::FTopologicalLoop> FOpenNurbsBRepToCADKernelConverter::AddLoop(const ON_BrepLoop& OpenNurbsLoop, TSharedPtr<UE::CADKernel::FSurface> & CarrierSurface, const bool bIsExternal)
 {
-	using namespace CADKernel;
+	using namespace UE::CADKernel;
 
 	if (!OpenNurbsLoop.IsValid())
 	{
@@ -115,7 +115,7 @@ TSharedPtr<CADKernel::FTopologicalLoop> FOpenNurbsBRepToCADKernelConverter::AddL
 
 	int32 EdgeCount = OpenNurbsLoop.TrimCount();
 	TArray<TSharedPtr<FTopologicalEdge>> Edges;
-	TArray<CADKernel::EOrientation> Directions;
+	TArray<UE::CADKernel::EOrientation> Directions;
 
 	Edges.Reserve(EdgeCount);
 	Directions.Reserve(EdgeCount);
@@ -128,7 +128,7 @@ TSharedPtr<CADKernel::FTopologicalLoop> FOpenNurbsBRepToCADKernelConverter::AddL
 		if (Edge.IsValid())
 		{
 			Edges.Add(Edge);
-			Directions.Emplace(CADKernel::EOrientation::Front);
+			Directions.Emplace(UE::CADKernel::EOrientation::Front);
 		}
 	}
 
@@ -140,7 +140,7 @@ TSharedPtr<CADKernel::FTopologicalLoop> FOpenNurbsBRepToCADKernelConverter::AddL
 	return FTopologicalLoop::Make(Edges, Directions, bIsExternal, GeometricTolerance);
 }
 
-void FOpenNurbsBRepToCADKernelConverter::LinkEdgesLoop(const ON_BrepLoop& OpenNurbsLoop, CADKernel::FTopologicalLoop& Loop)
+void FOpenNurbsBRepToCADKernelConverter::LinkEdgesLoop(const ON_BrepLoop& OpenNurbsLoop, UE::CADKernel::FTopologicalLoop& Loop)
 {
 	int32 EdgeCount = OpenNurbsLoop.TrimCount();
 	for (int32 Index = 0; Index < EdgeCount; ++Index)
@@ -152,7 +152,7 @@ void FOpenNurbsBRepToCADKernelConverter::LinkEdgesLoop(const ON_BrepLoop& OpenNu
 			continue;
 		}
 
-		TSharedPtr<CADKernel::FTopologicalEdge>* Edge = OpenNurbsTrimId2CADKernelEdge.Find(OpenNurbsTrim.m_trim_index);
+		TSharedPtr<UE::CADKernel::FTopologicalEdge>* Edge = OpenNurbsTrimId2CADKernelEdge.Find(OpenNurbsTrim.m_trim_index);
 		if (!Edge || !Edge->IsValid() || (*Edge)->IsDeleted() || (*Edge)->IsDegenerated())
 		{
 			continue;
@@ -166,7 +166,7 @@ void FOpenNurbsBRepToCADKernelConverter::LinkEdgesLoop(const ON_BrepLoop& OpenNu
 				continue;
 			}
 
-			TSharedPtr<CADKernel::FTopologicalEdge>* TwinEdge = OpenNurbsTrimId2CADKernelEdge.Find(LinkedEdgeId);
+			TSharedPtr<UE::CADKernel::FTopologicalEdge>* TwinEdge = OpenNurbsTrimId2CADKernelEdge.Find(LinkedEdgeId);
 			if (TwinEdge != nullptr && TwinEdge->IsValid() && !(*TwinEdge)->IsDeleted() && !(*TwinEdge)->IsDegenerated())
 			{
 				(*Edge)->Link(**TwinEdge, SquareTolerance);
@@ -176,9 +176,9 @@ void FOpenNurbsBRepToCADKernelConverter::LinkEdgesLoop(const ON_BrepLoop& OpenNu
 	}
 }
 
-TSharedPtr<CADKernel::FTopologicalEdge> FOpenNurbsBRepToCADKernelConverter::AddEdge(const ON_BrepTrim& OpenNurbsTrim, TSharedPtr<CADKernel::FSurface>& CarrierSurface)
+TSharedPtr<UE::CADKernel::FTopologicalEdge> FOpenNurbsBRepToCADKernelConverter::AddEdge(const ON_BrepTrim& OpenNurbsTrim, TSharedPtr<UE::CADKernel::FSurface>& CarrierSurface)
 {
-	using namespace CADKernel;
+	using namespace UE::CADKernel;
 
 	ON_BrepEdge* OpenNurbsEdge = OpenNurbsTrim.Edge();
 	if (OpenNurbsEdge == nullptr)
@@ -255,9 +255,9 @@ TSharedPtr<CADKernel::FTopologicalEdge> FOpenNurbsBRepToCADKernelConverter::AddE
 	return Edge;
 }
 
-TSharedPtr<CADKernel::FTopologicalFace> FOpenNurbsBRepToCADKernelConverter::AddFace(const ON_BrepFace& OpenNurbsFace)
+TSharedPtr<UE::CADKernel::FTopologicalFace> FOpenNurbsBRepToCADKernelConverter::AddFace(const ON_BrepFace& OpenNurbsFace)
 {
-	using namespace CADKernel;
+	using namespace UE::CADKernel;
 
 	ON_NurbsSurface OpenNurbsSurface;
 	OpenNurbsFace.NurbsSurface(&OpenNurbsSurface);
@@ -265,7 +265,7 @@ TSharedPtr<CADKernel::FTopologicalFace> FOpenNurbsBRepToCADKernelConverter::AddF
 	TSharedPtr<FSurface> Surface = AddSurface(OpenNurbsSurface);
 	if (!Surface.IsValid())
 	{
-		return TSharedPtr<CADKernel::FTopologicalFace>();
+		return TSharedPtr<FTopologicalFace>();
 	}
 
 	TSharedRef<FTopologicalFace> Face = FEntity::MakeShared<FTopologicalFace>(Surface);
@@ -296,7 +296,7 @@ TSharedPtr<CADKernel::FTopologicalFace> FOpenNurbsBRepToCADKernelConverter::AddF
 
 bool FOpenNurbsBRepToCADKernelConverter::AddBRep(ON_Brep& BRep, const ON_3dVector& Offset)
 {
-	using namespace CADKernel;
+	using namespace UE::CADKernel;
 
 	OpenNurbsTrimId2CADKernelEdge.Empty();
 
@@ -316,7 +316,7 @@ bool FOpenNurbsBRepToCADKernelConverter::AddBRep(ON_Brep& BRep, const ON_3dVecto
 		TSharedPtr<FTopologicalFace> Face = AddFace(OpenNurbsFace);
 		if (Face.IsValid())
 		{
-			Shell->Add(Face.ToSharedRef(), CADKernel::EOrientation::Front);
+			Shell->Add(Face.ToSharedRef(), UE::CADKernel::EOrientation::Front);
 		}
 	}
 
