Rename CADKernel namespace into UE::CADKernel

diff --git a/Engine/Plugins/Enterprise/DatasmithCADImporter/Source/DatasmithOpenNurbsTranslator/Private/OpenNurbsBRepToCADKernelConverter.h b/Engine/Plugins/Enterprise/DatasmithCADImporter/Source/DatasmithOpenNurbsTranslator/Private/OpenNurbsBRepToCADKernelConverter.h
--- a/Engine/Plugins/Enterprise/DatasmithCADImporter/Source/DatasmithOpenNurbsTranslator/Private/OpenNurbsBRepToCADKernelConverter.h
+++ b/Engine/Plugins/Enterprise/DatasmithCADImporter/Source/DatasmithOpenNurbsTranslator/Private/OpenNurbsBRepToCADKernelConverter.h
@@ -19,7 +19,7 @@ class ON_BrepTrim;
 class ON_NurbsSurface;
 class ON_3dVector;
 
-namespace CADKernel
+namespace UE::CADKernel
 {
 class FShell;
 class FSurface;
@@ -50,22 +50,22 @@ public:
 	static TSharedPtr<FOpenNurbsBRepToCADKernelConverter> GetSharedSession();
 
 private:
-	TSharedPtr<CADKernel::FTopologicalFace> AddFace(const ON_BrepFace& OpenNurbsFace);
-	TSharedPtr<CADKernel::FSurface> AddSurface(ON_NurbsSurface& Surface);
+	TSharedPtr<UE::CADKernel::FTopologicalFace> AddFace(const ON_BrepFace& OpenNurbsFace);
+	TSharedPtr<UE::CADKernel::FSurface> AddSurface(ON_NurbsSurface& Surface);
 
-	TSharedPtr<CADKernel::FTopologicalLoop> AddLoop(const ON_BrepLoop& OpenNurbsLoop, TSharedPtr<CADKernel::FSurface>& CarrierSurface, const bool bIsExternal);
+	TSharedPtr<UE::CADKernel::FTopologicalLoop> AddLoop(const ON_BrepLoop& OpenNurbsLoop, TSharedPtr<UE::CADKernel::FSurface>& CarrierSurface, const bool bIsExternal);
 
 	/**
 	 * Build face's links with its neighbor have to be done after the loop is finalize.
 	 * This is to avoid to link an edge with another and then to delete it...
 	 */
-	void LinkEdgesLoop(const ON_BrepLoop& OpenNurbsLoop, CADKernel::FTopologicalLoop& Loop);
+	void LinkEdgesLoop(const ON_BrepLoop& OpenNurbsLoop, UE::CADKernel::FTopologicalLoop& Loop);
 
-	TSharedPtr<CADKernel::FTopologicalEdge> AddEdge(const ON_BrepTrim& OpenNurbsTrim, TSharedPtr<CADKernel::FSurface>& CarrierSurface);
+	TSharedPtr<UE::CADKernel::FTopologicalEdge> AddEdge(const ON_BrepTrim& OpenNurbsTrim, TSharedPtr<UE::CADKernel::FSurface>& CarrierSurface);
 
 protected:
 
-	TMap<int32, TSharedPtr<CADKernel::FTopologicalEdge>>  OpenNurbsTrimId2CADKernelEdge;
+	TMap<int32, TSharedPtr<UE::CADKernel::FTopologicalEdge>>  OpenNurbsTrimId2CADKernelEdge;
 
 };
 
diff --git a/Engine/Plugins/Enterprise/DatasmithCADImporter/Source/DatasmithWireTranslator/Private/AliasModelToCADKernelConverter.cpp b/Engine/Plugins/Enterprise/DatasmithCADImporter/Source/DatasmithWireTranslator/Private/AliasModelToCADKernelConverter.cpp
--- a/Engine/Plugins/Enterprise/DatasmithCADImporter/Source/DatasmithWireTranslator/Private/AliasModelToCADKernelConverter.cpp
+++ b/Engine/Plugins/Enterprise/DatasmithCADImporter/Source/DatasmithWireTranslator/Private/AliasModelToCADKernelConverter.cpp
@@ -47,9 +47,9 @@ namespace AliasToCADKernelUtils
 {
 
 	template<typename Surface_T>
-	TSharedPtr<CADKernel::FSurface> AddNURBSSurface(double GeometricTolerance, Surface_T& AliasSurface, EAliasObjectReference InObjectReference, const AlMatrix4x4& InAlMatrix)
+	TSharedPtr<UE::CADKernel::FSurface> AddNURBSSurface(double GeometricTolerance, Surface_T& AliasSurface, EAliasObjectReference InObjectReference, const AlMatrix4x4& InAlMatrix)
 	{
-		CADKernel::FNurbsSurfaceHomogeneousData NURBSData;
+		UE::CADKernel::FNurbsSurfaceHomogeneousData NURBSData;
 		NURBSData.bSwapUV= true;
 		NURBSData.bIsRational = true;
 
@@ -99,13 +99,13 @@ namespace AliasToCADKernelUtils
 			Poles[Index + 2] *= 10.;
 		}
 
-		return CADKernel::FSurface::MakeNurbsSurface(GeometricTolerance, NURBSData);
+		return UE::CADKernel::FSurface::MakeNurbsSurface(GeometricTolerance, NURBSData);
 	}
 }
 
-TSharedPtr<CADKernel::FTopologicalEdge> FAliasModelToCADKernelConverter::AddEdge(const AlTrimCurve& AliasTrimCurve, TSharedPtr<CADKernel::FSurface>& CarrierSurface)
+TSharedPtr<UE::CADKernel::FTopologicalEdge> FAliasModelToCADKernelConverter::AddEdge(const AlTrimCurve& AliasTrimCurve, TSharedPtr<UE::CADKernel::FSurface>& CarrierSurface)
 {
-	CADKernel::FNurbsCurveData NurbsCurveData;
+	UE::CADKernel::FNurbsCurveData NurbsCurveData;
 
 	NurbsCurveData.Degree = AliasTrimCurve.degree();
 	int32 ControlPointCount = AliasTrimCurve.numberOfCVs();
@@ -134,13 +134,13 @@ TSharedPtr<CADKernel::FTopologicalEdge> FAliasModelToCADKernelConverter::AddEdge
 		NurbsCurveData.Poles[Index].Z = 0;
 	}
 
-	TSharedPtr<CADKernel::FCurve> Nurbs = CADKernel::FCurve::MakeNurbsCurve(NurbsCurveData);
+	TSharedPtr<UE::CADKernel::FCurve> Nurbs = UE::CADKernel::FCurve::MakeNurbsCurve(NurbsCurveData);
 
-	TSharedRef<CADKernel::FRestrictionCurve> RestrictionCurve = CADKernel::FCurve::MakeShared<CADKernel::FRestrictionCurve>(CarrierSurface.ToSharedRef(), Nurbs.ToSharedRef());
-	TSharedPtr<CADKernel::FTopologicalEdge> Edge = CADKernel::FTopologicalEdge::Make(RestrictionCurve);
+	TSharedRef<UE::CADKernel::FRestrictionCurve> RestrictionCurve = UE::CADKernel::FCurve::MakeShared<UE::CADKernel::FRestrictionCurve>(CarrierSurface.ToSharedRef(), Nurbs.ToSharedRef());
+	TSharedPtr<UE::CADKernel::FTopologicalEdge> Edge = UE::CADKernel::FTopologicalEdge::Make(RestrictionCurve);
 	if (!Edge.IsValid())
 	{
-		return TSharedPtr<CADKernel::FTopologicalEdge>();
+		return TSharedPtr<UE::CADKernel::FTopologicalEdge>();
 	}
 
 	// Only TrimCurve with twin need to be in the map used in LinkEdgesLoop 
@@ -153,35 +153,35 @@ TSharedPtr<CADKernel::FTopologicalEdge> FAliasModelToCADKernelConverter::AddEdge
 	return Edge;
 }
 
-TSharedPtr<CADKernel::FTopologicalLoop> FAliasModelToCADKernelConverter::AddLoop(const AlTrimBoundary& TrimBoundary, TSharedPtr<CADKernel::FSurface>& CarrierSurface, const bool bIsExternal)
+TSharedPtr<UE::CADKernel::FTopologicalLoop> FAliasModelToCADKernelConverter::AddLoop(const AlTrimBoundary& TrimBoundary, TSharedPtr<UE::CADKernel::FSurface>& CarrierSurface, const bool bIsExternal)
 {
-	TArray<TSharedPtr<CADKernel::FTopologicalEdge>> Edges;
-	TArray<CADKernel::EOrientation> Directions;
+	TArray<TSharedPtr<UE::CADKernel::FTopologicalEdge>> Edges;
+	TArray<UE::CADKernel::EOrientation> Directions;
 
 	for (TUniquePtr<AlTrimCurve> TrimCurve(TrimBoundary.firstCurve()); TrimCurve.IsValid(); TrimCurve = TUniquePtr<AlTrimCurve>(TrimCurve->nextCurve()))
 	{
-		TSharedPtr<CADKernel::FTopologicalEdge> Edge = AddEdge(*TrimCurve, CarrierSurface);
+		TSharedPtr<UE::CADKernel::FTopologicalEdge> Edge = AddEdge(*TrimCurve, CarrierSurface);
 		if (Edge.IsValid())
 		{
 			Edges.Add(Edge);
-			Directions.Emplace(CADKernel::EOrientation::Front);
+			Directions.Emplace(UE::CADKernel::EOrientation::Front);
 		}
 	}
 
 	if (Edges.Num() == 0)
 	{
-		return TSharedPtr<CADKernel::FTopologicalLoop>();
+		return TSharedPtr<UE::CADKernel::FTopologicalLoop>();
 	}
 
-	TSharedPtr<CADKernel::FTopologicalLoop> Loop = CADKernel::FTopologicalLoop::Make(Edges, Directions, bIsExternal, GeometricTolerance);
+	TSharedPtr<UE::CADKernel::FTopologicalLoop> Loop = UE::CADKernel::FTopologicalLoop::Make(Edges, Directions, bIsExternal, GeometricTolerance);
 	return Loop;
 }
 
-void FAliasModelToCADKernelConverter::LinkEdgesLoop(const AlTrimBoundary& TrimBoundary, CADKernel::FTopologicalLoop& Loop)
+void FAliasModelToCADKernelConverter::LinkEdgesLoop(const AlTrimBoundary& TrimBoundary, UE::CADKernel::FTopologicalLoop& Loop)
 {
 	for (TUniquePtr<AlTrimCurve> TrimCurve(TrimBoundary.firstCurve()); TrimCurve.IsValid(); TrimCurve = TUniquePtr<AlTrimCurve>(TrimCurve->nextCurve()))
 	{
-		TSharedPtr<CADKernel::FTopologicalEdge>* Edge = AlEdge2CADKernelEdge.Find(TrimCurve->fSpline);
+		TSharedPtr<UE::CADKernel::FTopologicalEdge>* Edge = AlEdge2CADKernelEdge.Find(TrimCurve->fSpline);
 		if (!Edge || !Edge->IsValid() || (*Edge)->IsDeleted() || (*Edge)->IsDegenerated())
 		{
 			continue;
@@ -193,7 +193,7 @@ void FAliasModelToCADKernelConverter::LinkEdgesLoop(const AlTrimBoundary& TrimBo
 		TUniquePtr<AlTrimCurve> TwinCurve(TrimCurve->getTwinCurve());
 		if (TwinCurve.IsValid())
 		{
-			if (TSharedPtr<CADKernel::FTopologicalEdge>* TwinEdge = AlEdge2CADKernelEdge.Find(TwinCurve->fSpline))
+			if (TSharedPtr<UE::CADKernel::FTopologicalEdge>* TwinEdge = AlEdge2CADKernelEdge.Find(TwinCurve->fSpline))
 			{
 				if (TwinEdge->IsValid() && !(*TwinEdge)->IsDeleted() && !(*TwinEdge)->IsDegenerated())
 				{
@@ -204,19 +204,19 @@ void FAliasModelToCADKernelConverter::LinkEdgesLoop(const AlTrimBoundary& TrimBo
 	}
 }
 
-TSharedPtr<CADKernel::FTopologicalFace> FAliasModelToCADKernelConverter::AddTrimRegion(const AlTrimRegion& TrimRegion, EAliasObjectReference InObjectReference, const AlMatrix4x4& InAlMatrix, bool bInOrientation)
+TSharedPtr<UE::CADKernel::FTopologicalFace> FAliasModelToCADKernelConverter::AddTrimRegion(const AlTrimRegion& TrimRegion, EAliasObjectReference InObjectReference, const AlMatrix4x4& InAlMatrix, bool bInOrientation)
 {
-	TSharedPtr<CADKernel::FSurface> Surface = AliasToCADKernelUtils::AddNURBSSurface(GeometricTolerance, TrimRegion, InObjectReference, InAlMatrix);
+	TSharedPtr<UE::CADKernel::FSurface> Surface = AliasToCADKernelUtils::AddNURBSSurface(GeometricTolerance, TrimRegion, InObjectReference, InAlMatrix);
 	if (!Surface.IsValid())
 	{
-		return TSharedPtr<CADKernel::FTopologicalFace>();
+		return TSharedPtr<UE::CADKernel::FTopologicalFace>();
 	}
 
 	bool bIsExternal = true;
-	TArray<TSharedPtr<CADKernel::FTopologicalLoop>> Loops;
+	TArray<TSharedPtr<UE::CADKernel::FTopologicalLoop>> Loops;
 	for (TUniquePtr<AlTrimBoundary> TrimBoundary(TrimRegion.firstBoundary()); TrimBoundary.IsValid(); TrimBoundary = TUniquePtr<AlTrimBoundary>(TrimBoundary->nextBoundary()))
 	{
-		TSharedPtr<CADKernel::FTopologicalLoop> Loop = AddLoop(*TrimBoundary, Surface, bIsExternal);
+		TSharedPtr<UE::CADKernel::FTopologicalLoop> Loop = AddLoop(*TrimBoundary, Surface, bIsExternal);
 		if (Loop.IsValid())
 		{
 			LinkEdgesLoop(*TrimBoundary, *Loop);
@@ -227,11 +227,11 @@ TSharedPtr<CADKernel::FTopologicalFace> FAliasModelToCADKernelConverter::AddTrim
 
 	if (Loops.Num() == 0)
 	{
-		CADKernel::FMessage::Printf(CADKernel::EVerboseLevel::Log, TEXT("The Face %s is degenerate, this face is ignored\n"), TrimRegion.name());
-		return TSharedPtr<CADKernel::FTopologicalFace>();
+		UE::CADKernel::FMessage::Printf(UE::CADKernel::EVerboseLevel::Log, TEXT("The Face %s is degenerate, this face is ignored\n"), TrimRegion.name());
+		return TSharedPtr<UE::CADKernel::FTopologicalFace>();
 	}
 
-	TSharedRef<CADKernel::FTopologicalFace> Face = CADKernel::FEntity::MakeShared<CADKernel::FTopologicalFace>(Surface);
+	TSharedRef<UE::CADKernel::FTopologicalFace> Face = UE::CADKernel::FEntity::MakeShared<UE::CADKernel::FTopologicalFace>(Surface);
 	Face->SetPatchId(LastFaceId++);
 
 	int32 DoubtfulLoopOrientationCount = 0;
@@ -239,39 +239,39 @@ TSharedPtr<CADKernel::FTopologicalFace> FAliasModelToCADKernelConverter::AddTrim
 	return Face;
 }
 
-void FAliasModelToCADKernelConverter::AddFace(const AlSurface& Surface, EAliasObjectReference InObjectReference, const AlMatrix4x4& InAlMatrix, bool bInOrientation, TSharedRef<CADKernel::FShell>& Shell)
+void FAliasModelToCADKernelConverter::AddFace(const AlSurface& Surface, EAliasObjectReference InObjectReference, const AlMatrix4x4& InAlMatrix, bool bInOrientation, TSharedRef<UE::CADKernel::FShell>& Shell)
 {
 	TUniquePtr<AlTrimRegion> TrimRegion(Surface.firstTrimRegion());
 	if (TrimRegion.IsValid())
 	{
 		for (; TrimRegion.IsValid(); TrimRegion = TUniquePtr<AlTrimRegion>(TrimRegion->nextRegion()))
 		{
-			TSharedPtr<CADKernel::FTopologicalFace> Face = AddTrimRegion(*TrimRegion, InObjectReference, InAlMatrix, bInOrientation);
+			TSharedPtr<UE::CADKernel::FTopologicalFace> Face = AddTrimRegion(*TrimRegion, InObjectReference, InAlMatrix, bInOrientation);
 			if (Face.IsValid())
 			{
-				Shell->Add(Face.ToSharedRef(), bInOrientation ? CADKernel::EOrientation::Front : CADKernel::EOrientation::Back);
+				Shell->Add(Face.ToSharedRef(), bInOrientation ? UE::CADKernel::EOrientation::Front : UE::CADKernel::EOrientation::Back);
 			}
 		}
 		return;
 	}
 
-	TSharedPtr<CADKernel::FSurface> CADKernelSurface = AliasToCADKernelUtils::AddNURBSSurface(GeometricTolerance, Surface, InObjectReference, InAlMatrix);
+	TSharedPtr<UE::CADKernel::FSurface> CADKernelSurface = AliasToCADKernelUtils::AddNURBSSurface(GeometricTolerance, Surface, InObjectReference, InAlMatrix);
 	if (CADKernelSurface.IsValid())
 	{
-		TSharedRef<CADKernel::FTopologicalFace> Face = CADKernel::FEntity::MakeShared<CADKernel::FTopologicalFace>(CADKernelSurface);
+		TSharedRef<UE::CADKernel::FTopologicalFace> Face = UE::CADKernel::FEntity::MakeShared<UE::CADKernel::FTopologicalFace>(CADKernelSurface);
 		Face->ApplyNaturalLoops();
-		Shell->Add(Face, bInOrientation ? CADKernel::EOrientation::Front : CADKernel::EOrientation::Back);
+		Shell->Add(Face, bInOrientation ? UE::CADKernel::EOrientation::Front : UE::CADKernel::EOrientation::Back);
 	}
 }
 
-void FAliasModelToCADKernelConverter::AddShell(const AlShell& InShell, EAliasObjectReference InObjectReference, const AlMatrix4x4& InAlMatrix, bool bInOrientation, TSharedRef<CADKernel::FShell>& CADKernelShell)
+void FAliasModelToCADKernelConverter::AddShell(const AlShell& InShell, EAliasObjectReference InObjectReference, const AlMatrix4x4& InAlMatrix, bool bInOrientation, TSharedRef<UE::CADKernel::FShell>& CADKernelShell)
 {
 	for(TUniquePtr<AlTrimRegion> TrimRegion(InShell.firstTrimRegion()); TrimRegion.IsValid(); TrimRegion = TUniquePtr<AlTrimRegion>(TrimRegion->nextRegion()))
 	{
-		TSharedPtr<CADKernel::FTopologicalFace> Face = AddTrimRegion(*TrimRegion, InObjectReference, InAlMatrix, bInOrientation);
+		TSharedPtr<UE::CADKernel::FTopologicalFace> Face = AddTrimRegion(*TrimRegion, InObjectReference, InAlMatrix, bInOrientation);
 		if (Face.IsValid())
 		{
-			CADKernelShell->Add(Face.ToSharedRef(), bInOrientation ? CADKernel::EOrientation::Front : CADKernel::EOrientation::Back);
+			CADKernelShell->Add(Face.ToSharedRef(), bInOrientation ? UE::CADKernel::EOrientation::Front : UE::CADKernel::EOrientation::Back);
 		}
 	}
 }
@@ -280,8 +280,8 @@ bool FAliasModelToCADKernelConverter::AddBRep(AlDagNode& DagNode, const FColor&
 {
 	AlEdge2CADKernelEdge.Empty();
 
-	TSharedRef<CADKernel::FBody> CADKernelBody = CADKernel::FEntity::MakeShared<CADKernel::FBody>();
-	TSharedRef<CADKernel::FShell> CADKernelShell = CADKernel::FEntity::MakeShared<CADKernel::FShell>();
+	TSharedRef<UE::CADKernel::FBody> CADKernelBody = UE::CADKernel::FEntity::MakeShared<UE::CADKernel::FBody>();
+	TSharedRef<UE::CADKernel::FShell> CADKernelShell = UE::CADKernel::FEntity::MakeShared<UE::CADKernel::FShell>();
 	uint32 ColorId = (uint32) CADLibrary::BuildColorUId(Color);
 	CADKernelShell->SetColorId(ColorId);
 	CADKernelBody->AddShell(CADKernelShell);
diff --git a/Engine/Plugins/Enterprise/DatasmithCADImporter/Source/DatasmithWireTranslator/Private/AliasModelToCADKernelConverter.h b/Engine/Plugins/Enterprise/DatasmithCADImporter/Source/DatasmithWireTranslator/Private/AliasModelToCADKernelConverter.h
--- a/Engine/Plugins/Enterprise/DatasmithCADImporter/Source/DatasmithWireTranslator/Private/AliasModelToCADKernelConverter.h
+++ b/Engine/Plugins/Enterprise/DatasmithCADImporter/Source/DatasmithWireTranslator/Private/AliasModelToCADKernelConverter.h
@@ -19,7 +19,7 @@ class AlTrimRegion;
 
 typedef double AlMatrix4x4[4][4];
 
-namespace CADKernel
+namespace UE::CADKernel
 {
 	class FShell;
 	class FSurface;
@@ -43,23 +43,23 @@ public:
 	virtual bool AddBRep(AlDagNode& DagNode, const FColor& Color, EAliasObjectReference ObjectReference) override;
 
 protected:
-	TSharedPtr<CADKernel::FTopologicalEdge> AddEdge(const AlTrimCurve& TrimCurve, TSharedPtr<CADKernel::FSurface>& CarrierSurface);
+	TSharedPtr<UE::CADKernel::FTopologicalEdge> AddEdge(const AlTrimCurve& TrimCurve, TSharedPtr<UE::CADKernel::FSurface>& CarrierSurface);
 
-	TSharedPtr<CADKernel::FTopologicalLoop> AddLoop(const AlTrimBoundary& TrimBoundary, TSharedPtr<CADKernel::FSurface>& CarrierSurface, const bool bIsExternal);
+	TSharedPtr<UE::CADKernel::FTopologicalLoop> AddLoop(const AlTrimBoundary& TrimBoundary, TSharedPtr<UE::CADKernel::FSurface>& CarrierSurface, const bool bIsExternal);
 	
 	/**
 	 * Build face's links with its neighbor have to be done after the loop is finalize.
 	 * This is to avoid to link an edge with another and then to delete it...
 	 */
-	void LinkEdgesLoop(const AlTrimBoundary& TrimBoundary, CADKernel::FTopologicalLoop& Loop);
+	void LinkEdgesLoop(const AlTrimBoundary& TrimBoundary, UE::CADKernel::FTopologicalLoop& Loop);
 
-	TSharedPtr<CADKernel::FTopologicalFace> AddTrimRegion(const AlTrimRegion& InTrimRegion, EAliasObjectReference InObjectReference, const AlMatrix4x4& InAlMatrix, bool bInOrientation);
-	void AddFace(const AlSurface& InSurface, EAliasObjectReference InObjectReference, const AlMatrix4x4& InAlMatrix, bool bInOrientation, TSharedRef<CADKernel::FShell>& Shell);
-	void AddShell(const AlShell& InShell, EAliasObjectReference InObjectReference, const AlMatrix4x4& InAlMatrix, bool bInOrientation, TSharedRef<CADKernel::FShell>& Shell);
+	TSharedPtr<UE::CADKernel::FTopologicalFace> AddTrimRegion(const AlTrimRegion& InTrimRegion, EAliasObjectReference InObjectReference, const AlMatrix4x4& InAlMatrix, bool bInOrientation);
+	void AddFace(const AlSurface& InSurface, EAliasObjectReference InObjectReference, const AlMatrix4x4& InAlMatrix, bool bInOrientation, TSharedRef<UE::CADKernel::FShell>& Shell);
+	void AddShell(const AlShell& InShell, EAliasObjectReference InObjectReference, const AlMatrix4x4& InAlMatrix, bool bInOrientation, TSharedRef<UE::CADKernel::FShell>& Shell);
 
 protected:
 	int32 LastFaceId = 1;
-	TMap<void*, TSharedPtr<CADKernel::FTopologicalEdge>>  AlEdge2CADKernelEdge;
+	TMap<void*, TSharedPtr<UE::CADKernel::FTopologicalEdge>>  AlEdge2CADKernelEdge;
 };
 
 }
\ No newline at end of file
diff --git a/Engine/Source/Runtime/Datasmith/CADKernel/Private/CADKernel/Core/Database.cpp b/Engine/Source/Runtime/Datasmith/CADKernel/Private/CADKernel/Core/Database.cpp
--- a/Engine/Source/Runtime/Datasmith/CADKernel/Private/CADKernel/Core/Database.cpp
+++ b/Engine/Source/Runtime/Datasmith/CADKernel/Private/CADKernel/Core/Database.cpp
@@ -18,7 +18,7 @@
 
 static const int32 DataBaseInitialSize = 10000;
 
-namespace CADKernel
+namespace UE::CADKernel
 {
 
 FDatabase::FDatabase()
diff --git a/Engine/Source/Runtime/Datasmith/CADKernel/Private/CADKernel/Core/Entity.cpp b/Engine/Source/Runtime/Datasmith/CADKernel/Private/CADKernel/Core/Entity.cpp
--- a/Engine/Source/Runtime/Datasmith/CADKernel/Private/CADKernel/Core/Entity.cpp
+++ b/Engine/Source/Runtime/Datasmith/CADKernel/Private/CADKernel/Core/Entity.cpp
@@ -20,7 +20,7 @@
 #include "CADKernel/Topo/TopologicalLoop.h"
 #include "CADKernel/Topo/TopologicalVertex.h"
 
-namespace CADKernel
+namespace UE::CADKernel
 {
 	const TCHAR* FEntity::TypesNames[] = 
 	{
@@ -377,4 +377,4 @@ namespace CADKernel
 		return Database.bForceSpawning || false;
 	}
 
-} // namespace CADKernel
+} // namespace UE::CADKernel
diff --git a/Engine/Source/Runtime/Datasmith/CADKernel/Private/CADKernel/Core/Factory.h b/Engine/Source/Runtime/Datasmith/CADKernel/Private/CADKernel/Core/Factory.h
--- a/Engine/Source/Runtime/Datasmith/CADKernel/Private/CADKernel/Core/Factory.h
+++ b/Engine/Source/Runtime/Datasmith/CADKernel/Private/CADKernel/Core/Factory.h
@@ -9,59 +9,59 @@
 
 #define DEBUG_FACTORY
 
-namespace CADKernel
+namespace UE::CADKernel
 {
-	const int32 MaxSize = 256;
+const int32 MaxSize = 256;
 
-	template<class ElementType>
-	class TFactory
-	{
-	protected:
-		TDoubleLinkedList<TArray<ElementType>> StackOfEntitySet;
-		TArray<ElementType*> FreeEntityStack;
+template<class ElementType>
+class TFactory
+{
+protected:
+	TDoubleLinkedList<TArray<ElementType>> StackOfEntitySet;
+	TArray<ElementType*> FreeEntityStack;
 #ifdef DEBUG_FACTORY
-		int32 AllocatedElementNum = 0;
-		int32 UsedElementNum = 0;
-		int32 AvailableElementNum = 0;
+	int32 AllocatedElementNum = 0;
+	int32 UsedElementNum = 0;
+	int32 AvailableElementNum = 0;
 #endif
 
-	private:
-		void FillFreeEntityStack()
+private:
+	void FillFreeEntityStack()
+	{
+		StackOfEntitySet.AddHead(TArray<ElementType>());
+		TArray<ElementType>& EntityPool = StackOfEntitySet.GetHead()->GetValue();
+		EntityPool.SetNum(MaxSize);
+		for (ElementType& Entity : EntityPool)
 		{
-			StackOfEntitySet.AddHead(TArray<ElementType>());
-			TArray<ElementType>& EntityPool = StackOfEntitySet.GetHead()->GetValue();
-			EntityPool.SetNum(MaxSize);
-			for (ElementType& Entity : EntityPool)
-			{
-				FreeEntityStack.Add(&Entity);
-			}
-			AllocatedElementNum += MaxSize;
+			FreeEntityStack.Add(&Entity);
 		}
+		AllocatedElementNum += MaxSize;
+	}
 
-	public:
-		TFactory()
-		{
-			FreeEntityStack.Reserve(MaxSize);
-		}
+public:
+	TFactory()
+	{
+		FreeEntityStack.Reserve(MaxSize);
+	}
 
-		void DeleteEntity(ElementType* Entity)
-		{
- 			Entity->Clean();
-			FreeEntityStack.Add(Entity);
-			UsedElementNum--;
-			ensureCADKernel(UsedElementNum + FreeEntityStack.Num() == AllocatedElementNum);
-		}
+	void DeleteEntity(ElementType* Entity)
+	{
+		Entity->Clean();
+		FreeEntityStack.Add(Entity);
+		UsedElementNum--;
+		ensureCADKernel(UsedElementNum + FreeEntityStack.Num() == AllocatedElementNum);
+	}
 
-		ElementType& New()
+	ElementType& New()
+	{
+		ensureCADKernel(UsedElementNum + FreeEntityStack.Num() == AllocatedElementNum);
+		if (FreeEntityStack.IsEmpty())
 		{
-			ensureCADKernel(UsedElementNum + FreeEntityStack.Num() == AllocatedElementNum);
-			if (FreeEntityStack.IsEmpty())
-			{
-				FillFreeEntityStack();
-			}
-			UsedElementNum++;
-			return *FreeEntityStack.Pop(false);
+			FillFreeEntityStack();
 		}
-	};
-} // namespace CADKernel
+		UsedElementNum++;
+		return *FreeEntityStack.Pop(false);
+	}
+};
+} // namespace UE::CADKernel
 
diff --git a/Engine/Source/Runtime/Datasmith/CADKernel/Private/CADKernel/Core/Group.cpp b/Engine/Source/Runtime/Datasmith/CADKernel/Private/CADKernel/Core/Group.cpp
--- a/Engine/Source/Runtime/Datasmith/CADKernel/Private/CADKernel/Core/Group.cpp
+++ b/Engine/Source/Runtime/Datasmith/CADKernel/Private/CADKernel/Core/Group.cpp
@@ -4,7 +4,7 @@
 
 #include "CADKernel/Core/System.h"
 
-namespace CADKernel
+namespace UE::CADKernel
 { 
 
 const TCHAR* GroupOriginNames[] = {
diff --git a/Engine/Source/Runtime/Datasmith/CADKernel/Private/CADKernel/Core/Group.h b/Engine/Source/Runtime/Datasmith/CADKernel/Private/CADKernel/Core/Group.h
--- a/Engine/Source/Runtime/Datasmith/CADKernel/Private/CADKernel/Core/Group.h
+++ b/Engine/Source/Runtime/Datasmith/CADKernel/Private/CADKernel/Core/Group.h
@@ -4,137 +4,137 @@
 #include "CADKernel/Core/Entity.h"
 #include "CADKernel/Core/CADKernelArchive.h"
 
-namespace CADKernel
+namespace UE::CADKernel
 {
-	enum class EGroupOrigin : uint8
-	{
-		Unknown,
-		CADGroup,
-		CADLayer,
-		CADColor
-	};
+enum class EGroupOrigin : uint8
+{
+	Unknown,
+	CADGroup,
+	CADLayer,
+	CADColor
+};
 
-	extern const TCHAR* GroupOriginNames[];
+extern const TCHAR* GroupOriginNames[];
 
-	class FGroup : public FEntity
-	{
-		friend FEntity;
+class FGroup : public FEntity
+{
+	friend FEntity;
 
-	protected:
-		EGroupOrigin Origin;
-		FString GroupName;
-		TArray<TSharedPtr<FEntity>> Entities;
+protected:
+	EGroupOrigin Origin;
+	FString GroupName;
+	TArray<TSharedPtr<FEntity>> Entities;
 
-		FGroup()
-			: Origin(EGroupOrigin::Unknown)
-		{
-		}
+	FGroup()
+		: Origin(EGroupOrigin::Unknown)
+	{
+	}
 
-		FGroup(TArray<TSharedPtr<FEntity>>& InEntities)
-			: Origin(EGroupOrigin::Unknown)
-		{
-			Entities.Append(InEntities);
-		}
+	FGroup(TArray<TSharedPtr<FEntity>>& InEntities)
+		: Origin(EGroupOrigin::Unknown)
+	{
+		Entities.Append(InEntities);
+	}
 
-	public:
+public:
 
-		virtual void Serialize(FCADKernelArchive& Ar) override
-		{
-			FEntity::Serialize(Ar);
-			Ar << Origin;
-			Ar << GroupName;
-			SerializeIdents(Ar, (TArray<TSharedPtr<FEntity>>&) Entities);
-		}
+	virtual void Serialize(FCADKernelArchive& Ar) override
+	{
+		FEntity::Serialize(Ar);
+		Ar << Origin;
+		Ar << GroupName;
+		SerializeIdents(Ar, (TArray<TSharedPtr<FEntity>>&) Entities);
+	}
 
-		virtual void SpawnIdent(FDatabase& Database) override
+	virtual void SpawnIdent(FDatabase& Database) override
+	{
+		if (!FEntity::SetId(Database))
 		{
-			if (!FEntity::SetId(Database))
-			{
-				return;
-			}
-
-			SpawnIdentOnEntities(Entities, Database);
+			return;
 		}
 
-		virtual void ResetMarkersRecursively() override
-		{
-			ResetMarkers();
-			ResetMarkersRecursivelyOnEntities(Entities);
-		}
+		SpawnIdentOnEntities(Entities, Database);
+	}
+
+	virtual void ResetMarkersRecursively() override
+	{
+		ResetMarkers();
+		ResetMarkersRecursivelyOnEntities(Entities);
+	}
 
 #ifdef CADKERNEL_DEV
-		virtual FInfoEntity& GetInfo(FInfoEntity&) const override;
+	virtual FInfoEntity& GetInfo(FInfoEntity&) const override;
 #endif
 
-		virtual EEntity GetEntityType() const override
-		{
-			return EEntity::Group;
-		}
+	virtual EEntity GetEntityType() const override
+	{
+		return EEntity::Group;
+	}
 
-		void SetName(const FString& Name);
+	void SetName(const FString& Name);
 
-		const FString& GetName() const
-		{
-			return GroupName;
-		}
+	const FString& GetName() const
+	{
+		return GroupName;
+	}
 
-		void AddEntity(TSharedPtr<FEntity> Entity)
-		{
-			Entities.AddUnique(Entity);
-		}
+	void AddEntity(TSharedPtr<FEntity> Entity)
+	{
+		Entities.AddUnique(Entity);
+	}
 
-		void Empty()
-		{
-			Entities.Empty();
-		}
+	void Empty()
+	{
+		Entities.Empty();
+	}
 
-		void RemoveEntity(TSharedPtr<FEntity> Entity)
-		{
-			Entities.Remove(Entity);
-		}
+	void RemoveEntity(TSharedPtr<FEntity> Entity)
+	{
+		Entities.Remove(Entity);
+	}
 
-		bool Contains(TSharedPtr<FEntity> Entity)
-		{
-			return Entities.Contains(Entity);
-		}
+	bool Contains(TSharedPtr<FEntity> Entity)
+	{
+		return Entities.Contains(Entity);
+	}
 
-		EGroupOrigin GetOrigin() const
-		{
-			return Origin;
-		}
+	EGroupOrigin GetOrigin() const
+	{
+		return Origin;
+	}
 
-		void SetOrigin(EGroupOrigin InOrigin)
-		{
-			Origin = InOrigin;
-		}
+	void SetOrigin(EGroupOrigin InOrigin)
+	{
+		Origin = InOrigin;
+	}
 
-		EEntity GetGroupType() const;
+	EEntity GetGroupType() const;
 
-		void GetValidEntities(TArray<TSharedPtr<FEntity>>& OutEntities) const
+	void GetValidEntities(TArray<TSharedPtr<FEntity>>& OutEntities) const
+	{
+		for (TSharedPtr<FEntity> Entity : Entities)
 		{
-			for (TSharedPtr<FEntity> Entity : Entities)
+			if (Entity.IsValid())
 			{
-				if (Entity.IsValid())
-				{
-					OutEntities.Add(Entity);
-				}
+				OutEntities.Add(Entity);
 			}
 		}
+	}
 
-		const TArray<TSharedPtr<FEntity>>& GetEntities() const
-		{
-			return Entities;
-		}
+	const TArray<TSharedPtr<FEntity>>& GetEntities() const
+	{
+		return Entities;
+	}
 
-		bool IsEmpty()
-		{
-			return Entities.IsEmpty();
-		}
+	bool IsEmpty()
+	{
+		return Entities.IsEmpty();
+	}
 
-		void ReplaceEntitiesWithMap(const TMap<TSharedPtr<FEntity>, TSharedPtr<FEntity>>& Map);
+	void ReplaceEntitiesWithMap(const TMap<TSharedPtr<FEntity>, TSharedPtr<FEntity>>& Map);
 
-		void RemoveNonTopologicalEntities();
-	};
+	void RemoveNonTopologicalEntities();
+};
 
-} // namespace CADKernel
+} // namespace UE::CADKernel
 
diff --git a/Engine/Source/Runtime/Datasmith/CADKernel/Private/CADKernel/Core/KernelParameters.h b/Engine/Source/Runtime/Datasmith/CADKernel/Private/CADKernel/Core/KernelParameters.h
--- a/Engine/Source/Runtime/Datasmith/CADKernel/Private/CADKernel/Core/KernelParameters.h
+++ b/Engine/Source/Runtime/Datasmith/CADKernel/Private/CADKernel/Core/KernelParameters.h
@@ -4,17 +4,17 @@
 #include "CADKernel/Core/Parameters.h"
 
 
-namespace CADKernel
+namespace UE::CADKernel
 {
-	class CADKERNEL_API FKernelParameters : public FParameters
-	{
-	public:
-		FParameter GeometricalTolerance;
+class CADKERNEL_API FKernelParameters : public FParameters
+{
+public:
+	FParameter GeometricalTolerance;
 
-		FKernelParameters()
-			: FParameters(3)
-			, GeometricalTolerance(TEXT("GeometricalTolerance"), 0.02, *this)
-		{}
-	};
-} // namespace CADKernel
+	FKernelParameters()
+		: FParameters(3)
+		, GeometricalTolerance(TEXT("GeometricalTolerance"), 0.02, *this)
+	{}
+};
+} // namespace UE::CADKernel
 
diff --git a/Engine/Source/Runtime/Datasmith/CADKernel/Private/CADKernel/Core/Session.cpp b/Engine/Source/Runtime/Datasmith/CADKernel/Private/CADKernel/Core/Session.cpp
--- a/Engine/Source/Runtime/Datasmith/CADKernel/Private/CADKernel/Core/Session.cpp
+++ b/Engine/Source/Runtime/Datasmith/CADKernel/Private/CADKernel/Core/Session.cpp
@@ -4,7 +4,7 @@
 #include "CADKernel/Topo/Model.h"
 #include "CADKernel/UI/Message.h"
 
-namespace CADKernel
+namespace UE::CADKernel
 {
 
 #if defined(CADKERNEL_DEV) || defined(CADKERNEL_STDA)
@@ -97,4 +97,4 @@ void FSession::SetGeometricTolerance(double NewTolerance)
 	ensure(Database.GetModel().EntityCount() == 0);
 	GeometricTolerance = NewTolerance;
 }
-} // namespace CADKernel
+} // namespace UE::CADKernel
diff --git a/Engine/Source/Runtime/Datasmith/CADKernel/Private/CADKernel/Geo/Curves/BezierCurve.h b/Engine/Source/Runtime/Datasmith/CADKernel/Private/CADKernel/Geo/Curves/BezierCurve.h
--- a/Engine/Source/Runtime/Datasmith/CADKernel/Private/CADKernel/Geo/Curves/BezierCurve.h
+++ b/Engine/Source/Runtime/Datasmith/CADKernel/Private/CADKernel/Geo/Curves/BezierCurve.h
@@ -3,56 +3,56 @@
 
 #include "CADKernel/Geo/Curves/Curve.h"
 
-namespace CADKernel
+namespace UE::CADKernel
 {
-	class CADKERNEL_API FBezierCurve : public FCurve
-	{
-		friend class FEntity;
+class CADKERNEL_API FBezierCurve : public FCurve
+{
+	friend class FEntity;
 
-	protected:
-		TArray<FPoint> Poles;
+protected:
+	TArray<FPoint> Poles;
 
-		FBezierCurve(const TArray<FPoint>& InPoles)
-			: Poles{ InPoles }
-		{
-		}
+	FBezierCurve(const TArray<FPoint>& InPoles)
+		: Poles{ InPoles }
+	{
+	}
 
-		FBezierCurve() = default;
+	FBezierCurve() = default;
 
-	public:
+public:
 
-		virtual void Serialize(FCADKernelArchive& Ar) override
-		{
-			FCurve::Serialize(Ar);
-			Ar.Serialize(Poles);
-		}
+	virtual void Serialize(FCADKernelArchive& Ar) override
+	{
+		FCurve::Serialize(Ar);
+		Ar.Serialize(Poles);
+	}
 
 #ifdef CADKERNEL_DEV
-		virtual FInfoEntity& GetInfo(FInfoEntity&) const override;
+	virtual FInfoEntity& GetInfo(FInfoEntity&) const override;
 #endif
 
-		virtual ECurve GetCurveType() const override
-		{
-			return ECurve::Bezier;
-		}
+	virtual ECurve GetCurveType() const override
+	{
+		return ECurve::Bezier;
+	}
 
-		int32 GetDegre() const
-		{
-			return Poles.Num() - 1;
-		}
+	int32 GetDegre() const
+	{
+		return Poles.Num() - 1;
+	}
 
-		const TArray<FPoint>& GetPoles() const
-		{
-			return Poles;
-		}
+	const TArray<FPoint>& GetPoles() const
+	{
+		return Poles;
+	}
 
-		virtual TSharedPtr<FEntityGeom> ApplyMatrix(const FMatrixH& InMatrix) const override;
+	virtual TSharedPtr<FEntityGeom> ApplyMatrix(const FMatrixH& InMatrix) const override;
 
-		virtual void EvaluatePoint(double Coordinate, FCurvePoint& OutPoint, int32 DerivativeOrder = 0) const override;
+	virtual void EvaluatePoint(double Coordinate, FCurvePoint& OutPoint, int32 DerivativeOrder = 0) const override;
 
-		virtual void ExtendTo(const FPoint& Point) override;
+	virtual void ExtendTo(const FPoint& Point) override;
 
-	};
+};
 
-} // namespace CADKernel
+} // namespace UE::CADKernel
 
diff --git a/Engine/Source/Runtime/Datasmith/CADKernel/Private/CADKernel/Geo/Curves/BoundedCurve.h b/Engine/Source/Runtime/Datasmith/CADKernel/Private/CADKernel/Geo/Curves/BoundedCurve.h
--- a/Engine/Source/Runtime/Datasmith/CADKernel/Private/CADKernel/Geo/Curves/BoundedCurve.h
+++ b/Engine/Source/Runtime/Datasmith/CADKernel/Private/CADKernel/Geo/Curves/BoundedCurve.h
@@ -3,72 +3,72 @@
 
 #include "CADKernel/Geo/Curves/Curve.h"
 
-namespace CADKernel
+namespace UE::CADKernel
 {
 
-	class CADKERNEL_API FBoundedCurve : public FCurve
-	{
-		friend class FEntity;
+class CADKERNEL_API FBoundedCurve : public FCurve
+{
+	friend class FEntity;
 
-	protected:
+protected:
 
-		TSharedPtr<FCurve> Curve;
+	TSharedPtr<FCurve> Curve;
 
-		FBoundedCurve(TSharedRef<FCurve> InCurve, const FLinearBoundary& InBoundary, int8 InDimension)
-			: FCurve(InBoundary, InDimension)
-			, Curve(InCurve)
-		{
-		}
+	FBoundedCurve(TSharedRef<FCurve> InCurve, const FLinearBoundary& InBoundary, int8 InDimension)
+		: FCurve(InBoundary, InDimension)
+		, Curve(InCurve)
+	{
+	}
 
-		FBoundedCurve() = default;
+	FBoundedCurve() = default;
 
-	public:
+public:
 
-		virtual void Serialize(FCADKernelArchive& Ar) override
-		{
-			FCurve::Serialize(Ar);
-			SerializeIdent(Ar, Curve);
-		}
+	virtual void Serialize(FCADKernelArchive& Ar) override
+	{
+		FCurve::Serialize(Ar);
+		SerializeIdent(Ar, Curve);
+	}
 
-		virtual void SpawnIdent(FDatabase& Database) override
+	virtual void SpawnIdent(FDatabase& Database) override
+	{
+		if (!FEntity::SetId(Database))
 		{
-			if (!FEntity::SetId(Database))
-			{
-				return;
-			}
-
-			Curve->SpawnIdent(Database);
+			return;
 		}
 
-		virtual void ResetMarkersRecursively() override
-		{
-			ResetMarkers();
-			Curve->ResetMarkersRecursively();
-		}
+		Curve->SpawnIdent(Database);
+	}
+
+	virtual void ResetMarkersRecursively() override
+	{
+		ResetMarkers();
+		Curve->ResetMarkersRecursively();
+	}
 
 #ifdef CADKERNEL_DEV
-		virtual FInfoEntity& GetInfo(FInfoEntity&) const override;
+	virtual FInfoEntity& GetInfo(FInfoEntity&) const override;
 #endif
 
-		virtual ECurve GetCurveType() const override
-		{
-			return ECurve::BoundedCurve;
-		}
+	virtual ECurve GetCurveType() const override
+	{
+		return ECurve::BoundedCurve;
+	}
 
-		TSharedRef<const FCurve> GetBaseCurve()
-		{
-			return Curve.ToSharedRef();
-		}
+	TSharedRef<const FCurve> GetBaseCurve()
+	{
+		return Curve.ToSharedRef();
+	}
 
-		virtual TSharedPtr<FEntityGeom> ApplyMatrix(const FMatrixH& InMatrix) const override;
+	virtual TSharedPtr<FEntityGeom> ApplyMatrix(const FMatrixH& InMatrix) const override;
 
-		virtual void EvaluatePoint(double Coordinate, FCurvePoint& OutPoint, int32 DerivativeOrder = 0) const override;
-		virtual void Evaluate2DPoint(double Coordinate, FCurvePoint2D& OutPoint, int32 DerivativeOrder = 0) const override;
+	virtual void EvaluatePoint(double Coordinate, FCurvePoint& OutPoint, int32 DerivativeOrder = 0) const override;
+	virtual void Evaluate2DPoint(double Coordinate, FCurvePoint2D& OutPoint, int32 DerivativeOrder = 0) const override;
 
-		virtual void FindNotDerivableCoordinates(const FLinearBoundary& InBoundary, int32 DerivativeOrder, TArray<double>& OutNotDerivableCoordinates) const override;
+	virtual void FindNotDerivableCoordinates(const FLinearBoundary& InBoundary, int32 DerivativeOrder, TArray<double>& OutNotDerivableCoordinates) const override;
 
-		virtual TSharedPtr<FCurve> MakeBoundedCurve(const FLinearBoundary& InBoundary) override;
-	};
+	virtual TSharedPtr<FCurve> MakeBoundedCurve(const FLinearBoundary& InBoundary) override;
+};
 
-} // namespace CADKernel
+} // namespace UE::CADKernel
 
diff --git a/Engine/Source/Runtime/Datasmith/CADKernel/Private/CADKernel/Geo/Curves/CompositeCurve.h b/Engine/Source/Runtime/Datasmith/CADKernel/Private/CADKernel/Geo/Curves/CompositeCurve.h
--- a/Engine/Source/Runtime/Datasmith/CADKernel/Private/CADKernel/Geo/Curves/CompositeCurve.h
+++ b/Engine/Source/Runtime/Datasmith/CADKernel/Private/CADKernel/Geo/Curves/CompositeCurve.h
@@ -4,90 +4,90 @@
 #include "CADKernel/Core/OrientedEntity.h"
 #include "CADKernel/Geo/Curves/Curve.h"
 
-namespace CADKernel
+namespace UE::CADKernel
 {
-	struct FCurvePoint;
+struct FCurvePoint;
 
-	class CADKERNEL_API FOrientedCurve : public TOrientedEntity<FCurve>
+class CADKERNEL_API FOrientedCurve : public TOrientedEntity<FCurve>
+{
+public:
+	FOrientedCurve(TSharedPtr<FCurve>& InEntity, EOrientation InDirection)
+		: TOrientedEntity(InEntity, InDirection)
 	{
-	public:
-		FOrientedCurve(TSharedPtr<FCurve>& InEntity, EOrientation InDirection)
-			: TOrientedEntity(InEntity, InDirection)
-		{
-		}
+	}
 
-		FOrientedCurve(const FOrientedCurve& OrientiredEntity)
-			: TOrientedEntity(OrientiredEntity)
-		{
-		}
-
-		FOrientedCurve()
-			: TOrientedEntity()
-		{
-		}
-	};
-
-	class CADKERNEL_API FCompositeCurve : public FCurve
+	FOrientedCurve(const FOrientedCurve& OrientiredEntity)
+		: TOrientedEntity(OrientiredEntity)
 	{
-		friend class FEntity;
+	}
 
-	protected:
-		TArray<FOrientedCurve> Curves;
-		TArray<double> Coordinates;
+	FOrientedCurve()
+		: TOrientedEntity()
+	{
+	}
+};
 
-		FCompositeCurve(const TArray<TSharedPtr<FCurve>>& Curves, bool bDoInversions = false);
+class CADKERNEL_API FCompositeCurve : public FCurve
+{
+	friend class FEntity;
 
-		FCompositeCurve() = default;
+protected:
+	TArray<FOrientedCurve> Curves;
+	TArray<double> Coordinates;
 
-	public:
+	FCompositeCurve(const TArray<TSharedPtr<FCurve>>& Curves, bool bDoInversions = false);
 
-		virtual void Serialize(FCADKernelArchive& Ar) override
-		{
-			FCurve::Serialize(Ar);
-			SerializeIdents(Ar, (TArray<TOrientedEntity<FEntity>>&) Curves);
-			Ar.Serialize(Coordinates);
-		}
+	FCompositeCurve() = default;
 
-		virtual void SpawnIdent(FDatabase& Database) override
-		{
-			if (!FEntity::SetId(Database))
-			{
-				return;
-			}
+public:
 
-			SpawnIdentOnEntities((TArray<TOrientedEntity<FEntity>>&) Curves, Database);
-		}
+	virtual void Serialize(FCADKernelArchive& Ar) override
+	{
+		FCurve::Serialize(Ar);
+		SerializeIdents(Ar, (TArray<TOrientedEntity<FEntity>>&) Curves);
+		Ar.Serialize(Coordinates);
+	}
 
-		virtual void ResetMarkersRecursively() override
+	virtual void SpawnIdent(FDatabase& Database) override
+	{
+		if (!FEntity::SetId(Database))
 		{
-			ResetMarkers();
-			ResetMarkersRecursivelyOnEntities((TArray<TOrientedEntity<FEntity>>&) Curves);
+			return;
 		}
 
+		SpawnIdentOnEntities((TArray<TOrientedEntity<FEntity>>&) Curves, Database);
+	}
+
+	virtual void ResetMarkersRecursively() override
+	{
+		ResetMarkers();
+		ResetMarkersRecursivelyOnEntities((TArray<TOrientedEntity<FEntity>>&) Curves);
+	}
+
 #ifdef CADKERNEL_DEV
-		virtual FInfoEntity& GetInfo(FInfoEntity&) const override;
+	virtual FInfoEntity& GetInfo(FInfoEntity&) const override;
 #endif
 
-		virtual ECurve GetCurveType() const override
-		{
-			return ECurve::Composite;
-		}
+	virtual ECurve GetCurveType() const override
+	{
+		return ECurve::Composite;
+	}
 
-		virtual TSharedPtr<FEntityGeom> ApplyMatrix(const FMatrixH& InMatrix) const override;
+	virtual TSharedPtr<FEntityGeom> ApplyMatrix(const FMatrixH& InMatrix) const override;
 
-		virtual TSharedPtr<FCurve> GetCurve(int32 Index) const 
-		{
-			return Curves[Index].Entity;
-		}
+	virtual TSharedPtr<FCurve> GetCurve(int32 Index) const
+	{
+		return Curves[Index].Entity;
+	}
 
-		virtual void EvaluatePoint(double Coordinate, FCurvePoint& OutPoint, int32 DerivativeOrder = 0) const override;
-		virtual void Evaluate2DPoint(double Coordinate, FCurvePoint2D& OutPoint, int32 DerivativeOrder = 0) const override;
+	virtual void EvaluatePoint(double Coordinate, FCurvePoint& OutPoint, int32 DerivativeOrder = 0) const override;
+	virtual void Evaluate2DPoint(double Coordinate, FCurvePoint2D& OutPoint, int32 DerivativeOrder = 0) const override;
 
-		virtual void FindNotDerivableCoordinates(const FLinearBoundary& InBoundary, int32 DerivativeOrder, TArray<double>& OutNotDerivableCoordinates) const override;
+	virtual void FindNotDerivableCoordinates(const FLinearBoundary& InBoundary, int32 DerivativeOrder, TArray<double>& OutNotDerivableCoordinates) const override;
 
-		double LocalToGlobalCoordinate(int32 CurveIndex, double LocalCoordinate) const;
-		double GlobalToLocalCoordinate(int32 CurveIndex, double GlobalCoordinate) const;
-	};
+	double LocalToGlobalCoordinate(int32 CurveIndex, double LocalCoordinate) const;
+	double GlobalToLocalCoordinate(int32 CurveIndex, double GlobalCoordinate) const;
+};
 
-} // namespace CADKernel
+} // namespace UE::CADKernel
 
diff --git a/Engine/Source/Runtime/Datasmith/CADKernel/Private/CADKernel/Geo/Curves/Curve.cpp b/Engine/Source/Runtime/Datasmith/CADKernel/Private/CADKernel/Geo/Curves/Curve.cpp
--- a/Engine/Source/Runtime/Datasmith/CADKernel/Private/CADKernel/Geo/Curves/Curve.cpp
+++ b/Engine/Source/Runtime/Datasmith/CADKernel/Private/CADKernel/Geo/Curves/Curve.cpp
@@ -11,154 +11,154 @@
 #include "CADKernel/Geo/Sampler/SamplerOnParam.h"
 #include "CADKernel/UI/Display.h"
 
-namespace CADKernel
+namespace UE::CADKernel
 {
-	double FCurve::GetLength(double Tolerance) const
+double FCurve::GetLength(double Tolerance) const
+{
+	if (!GlobalLength.IsValid())
 	{
-		if (!GlobalLength.IsValid())
+		switch (Dimension)
 		{
-			switch (Dimension)
-			{
-			case 3:
-				GlobalLength = ComputeLength(Boundary, Tolerance);
-				break;
-			case 2:
-				GlobalLength = ComputeLength2D(Boundary, Tolerance);
-				break;
-			default:
-				GlobalLength = 0;
-				break;
-			}
+		case 3:
+			GlobalLength = ComputeLength(Boundary, Tolerance);
+			break;
+		case 2:
+			GlobalLength = ComputeLength2D(Boundary, Tolerance);
+			break;
+		default:
+			GlobalLength = 0;
+			break;
 		}
-		return GlobalLength;
 	}
+	return GlobalLength;
+}
 
-	void FCurve::EvaluatePoints(const TArray<double>& Coordinates, TArray<FCurvePoint>& OutPoints, int32 DerivativeOrder) const
+void FCurve::EvaluatePoints(const TArray<double>& Coordinates, TArray<FCurvePoint>& OutPoints, int32 DerivativeOrder) const
+{
+	OutPoints.SetNum(Coordinates.Num());
+	for (int32 iPoint = 0; iPoint < Coordinates.Num(); iPoint++)
 	{
-		OutPoints.SetNum(Coordinates.Num());
-		for (int32 iPoint = 0; iPoint < Coordinates.Num(); iPoint++)
-		{
-			EvaluatePoint(Coordinates[iPoint], OutPoints[iPoint], DerivativeOrder);
-		}
+		EvaluatePoint(Coordinates[iPoint], OutPoints[iPoint], DerivativeOrder);
 	}
+}
 
-	void FCurve::EvaluatePoints(const TArray<double>& Coordinates, TArray<FPoint>& OutPoints) const
+void FCurve::EvaluatePoints(const TArray<double>& Coordinates, TArray<FPoint>& OutPoints) const
+{
+	OutPoints.Empty(Coordinates.Num());
+	for (double Coordinate : Coordinates)
 	{
-		OutPoints.Empty(Coordinates.Num());
-		for (double Coordinate : Coordinates)
-		{
-			FCurvePoint Point;
-			EvaluatePoint(Coordinate, Point, 0);
-			OutPoints.Emplace(Point.Point);
-		}
+		FCurvePoint Point;
+		EvaluatePoint(Coordinate, Point, 0);
+		OutPoints.Emplace(Point.Point);
 	}
+}
 
-	void FCurve::Evaluate2DPoints(const TArray<double>& Coordinates, TArray<FCurvePoint2D>& OutPoints, int32 DerivativeOrder) const
+void FCurve::Evaluate2DPoints(const TArray<double>& Coordinates, TArray<FCurvePoint2D>& OutPoints, int32 DerivativeOrder) const
+{
+	OutPoints.SetNum(Coordinates.Num());
+	for (int32 iPoint = 0; iPoint < Coordinates.Num(); iPoint++)
 	{
-		OutPoints.SetNum(Coordinates.Num());
-		for (int32 iPoint = 0; iPoint < Coordinates.Num(); iPoint++)
-		{
-			Evaluate2DPoint(Coordinates[iPoint], OutPoints[iPoint], DerivativeOrder);
-		}
+		Evaluate2DPoint(Coordinates[iPoint], OutPoints[iPoint], DerivativeOrder);
 	}
+}
 
-	/**
-	 * Evaluate exact 2D points of the curve at the input Coordinates
-	 * The function can only be used with 2D curve (Dimension == 2)
-	 */
-	void FCurve::Evaluate2DPoints(const TArray<double>& Coordinates, TArray<FPoint2D>& OutPoints) const
+/**
+ * Evaluate exact 2D points of the curve at the input Coordinates
+ * The function can only be used with 2D curve (Dimension == 2)
+ */
+void FCurve::Evaluate2DPoints(const TArray<double>& Coordinates, TArray<FPoint2D>& OutPoints) const
+{
+	OutPoints.SetNum(Coordinates.Num());
+	for (int32 iPoint = 0; iPoint < Coordinates.Num(); iPoint++)
 	{
-		OutPoints.SetNum(Coordinates.Num());
-		for (int32 iPoint = 0; iPoint < Coordinates.Num(); iPoint++)
-		{
-			Evaluate2DPoint(Coordinates[iPoint], OutPoints[iPoint]);
-		}
+		Evaluate2DPoint(Coordinates[iPoint], OutPoints[iPoint]);
 	}
+}
 
 #ifdef CADKERNEL_DEV
-	FInfoEntity& FCurve::GetInfo(FInfoEntity& Info) const
-	{
-		return FEntity::GetInfo(Info)
-			.Add(TEXT("Curve type"), CurvesTypesNames[(uint8)GetCurveType()])
-			.Add(TEXT("Dimension"), (int32) Dimension)
-			.Add(TEXT("Boundary"), Boundary)
-			.Add(TEXT("Length"), GetLength(0.01));
-	}
+FInfoEntity& FCurve::GetInfo(FInfoEntity& Info) const
+{
+	return FEntity::GetInfo(Info)
+		.Add(TEXT("Curve type"), CurvesTypesNames[(uint8)GetCurveType()])
+		.Add(TEXT("Dimension"), (int32)Dimension)
+		.Add(TEXT("Boundary"), Boundary)
+		.Add(TEXT("Length"), GetLength(0.01));
+}
 #endif
 
-	void FCurve::FindNotDerivableCoordinates(int32 DerivativeOrder, TArray<double>& OutNotDerivableCoordinates) const
-	{
-		FindNotDerivableCoordinates(Boundary, DerivativeOrder, OutNotDerivableCoordinates);
-	}
+void FCurve::FindNotDerivableCoordinates(int32 DerivativeOrder, TArray<double>& OutNotDerivableCoordinates) const
+{
+	FindNotDerivableCoordinates(Boundary, DerivativeOrder, OutNotDerivableCoordinates);
+}
 
-	void FCurve::FindNotDerivableCoordinates(const FLinearBoundary& InBoundary, int32 DerivativeOrder, TArray<double>& OutNotDerivableCoordinates) const
+void FCurve::FindNotDerivableCoordinates(const FLinearBoundary& InBoundary, int32 DerivativeOrder, TArray<double>& OutNotDerivableCoordinates) const
+{
+}
+
+TSharedPtr<FCurve> FCurve::Rebound(const FLinearBoundary& InBoundary)
+{
+	if (FMath::IsNearlyEqual(InBoundary.Min, GetUMin()) && FMath::IsNearlyEqual(InBoundary.Max, GetUMax()))
 	{
+		FMessage::Printf(Debug, TEXT("Invalid rebound (UMin and UMax are nearly equal) on curve %d\n"), GetId());
+		TSharedPtr<FEntity> Entity = AsShared();
+		return StaticCastSharedPtr<FCurve>(Entity);
 	}
 
-	TSharedPtr<FCurve> FCurve::Rebound(const FLinearBoundary& InBoundary)
+	ensureCADKernel(false);
+	return TSharedPtr<FCurve>();
+}
+
+TSharedPtr<FCurve> FCurve::MakeBoundedCurve(const FLinearBoundary& InBoundary)
+{
+	FLinearBoundary NewBoundary = InBoundary;
+	if (NewBoundary.Min < GetUMin())
 	{
-		if (FMath::IsNearlyEqual(InBoundary.Min, GetUMin()) && FMath::IsNearlyEqual(InBoundary.Max, GetUMax()))
-		{
-			FMessage::Printf(Debug, TEXT("Invalid rebound (UMin and UMax are nearly equal) on curve %d\n"), GetId());
-			TSharedPtr<FEntity> Entity = AsShared();
-			return StaticCastSharedPtr<FCurve> (Entity);
-		}
-	
-		ensureCADKernel(false);
-		return TSharedPtr<FCurve>();
+		NewBoundary.Min = GetUMin();
 	}
 
-	TSharedPtr<FCurve> FCurve::MakeBoundedCurve(const FLinearBoundary& InBoundary) 
+	if (NewBoundary.Max > GetUMax())
 	{
-		FLinearBoundary NewBoundary = InBoundary;
-		if (NewBoundary.Min < GetUMin())
-		{
-			NewBoundary.Min = GetUMin();
-		}
-
-		if (NewBoundary.Max > GetUMax())
-		{
-			NewBoundary.Max = GetUMax();
-		}
-
-		if (NewBoundary.IsDegenerated())
-		{
-			FMessage::Printf(Log, TEXT("Invalid bounds (u1=%f u2=%f) on curve %d\n"), NewBoundary.Min, NewBoundary.Max, GetId());
-			return TSharedPtr<FCurve>();
-		}
-
-		if (FMath::IsNearlyEqual(NewBoundary.Min, GetUMin()) && FMath::IsNearlyEqual(NewBoundary.Max, GetUMax()))
-		{
-			FMessage::Printf(Debug, TEXT("Invalid rebound (UMin and UMax are nearly equal) on curve %d\n"), GetId());
-			return TSharedPtr<FCurve>();
-		}
-
-		return FEntity::MakeShared<FBoundedCurve>(StaticCastSharedRef<FCurve>(AsShared()), NewBoundary, Dimension);
+		NewBoundary.Max = GetUMax();
 	}
 
-	double FCurve::ComputeLength(const FLinearBoundary& InBoundary, double Tolerance) const
+	if (NewBoundary.IsDegenerated())
 	{
-		FPolyline3D Polyline;
-		FCurveSamplerOnChord Sampler(*this, Boundary, Tolerance, Polyline);
-		Sampler.Sample();
-		return Polyline.GetLength(Boundary);
+		FMessage::Printf(Log, TEXT("Invalid bounds (u1=%f u2=%f) on curve %d\n"), NewBoundary.Min, NewBoundary.Max, GetId());
+		return TSharedPtr<FCurve>();
 	}
 
-	double FCurve::ComputeLength2D(const FLinearBoundary& InBoundary, double Tolerance) const
+	if (FMath::IsNearlyEqual(NewBoundary.Min, GetUMin()) && FMath::IsNearlyEqual(NewBoundary.Max, GetUMax()))
 	{
-		FPolyline2D Polyline;
-		FCurve2DSamplerOnChord Sampler(*this, Boundary, Tolerance, Polyline);
-		Sampler.Sample();
-		return Polyline.GetLength(Boundary);
+		FMessage::Printf(Debug, TEXT("Invalid rebound (UMin and UMax are nearly equal) on curve %d\n"), GetId());
+		return TSharedPtr<FCurve>();
 	}
 
-	void FCurve::Presample(const FLinearBoundary& InBoundary, double Tolerance, TArray<double>& OutSampling) const
-	{
-		FPolyline3D Presampling;
-		FCurveSamplerOnParam Sampler(*this, Boundary, Tolerance * 10., Tolerance, Presampling);
-		Sampler.Sample();
+	return FEntity::MakeShared<FBoundedCurve>(StaticCastSharedRef<FCurve>(AsShared()), NewBoundary, Dimension);
+}
 
-		Presampling.SwapCoordinates(OutSampling);
-	}
+double FCurve::ComputeLength(const FLinearBoundary& InBoundary, double Tolerance) const
+{
+	FPolyline3D Polyline;
+	FCurveSamplerOnChord Sampler(*this, Boundary, Tolerance, Polyline);
+	Sampler.Sample();
+	return Polyline.GetLength(Boundary);
+}
+
+double FCurve::ComputeLength2D(const FLinearBoundary& InBoundary, double Tolerance) const
+{
+	FPolyline2D Polyline;
+	FCurve2DSamplerOnChord Sampler(*this, Boundary, Tolerance, Polyline);
+	Sampler.Sample();
+	return Polyline.GetLength(Boundary);
+}
+
+void FCurve::Presample(const FLinearBoundary& InBoundary, double Tolerance, TArray<double>& OutSampling) const
+{
+	FPolyline3D Presampling;
+	FCurveSamplerOnParam Sampler(*this, Boundary, Tolerance * 10., Tolerance, Presampling);
+	Sampler.Sample();
+
+	Presampling.SwapCoordinates(OutSampling);
+}
 
-} // namespace CADKernel
+} // namespace UE::CADKernel
diff --git a/Engine/Source/Runtime/Datasmith/CADKernel/Private/CADKernel/Geo/Curves/CurveMaker.cpp b/Engine/Source/Runtime/Datasmith/CADKernel/Private/CADKernel/Geo/Curves/CurveMaker.cpp
--- a/Engine/Source/Runtime/Datasmith/CADKernel/Private/CADKernel/Geo/Curves/CurveMaker.cpp
+++ b/Engine/Source/Runtime/Datasmith/CADKernel/Private/CADKernel/Geo/Curves/CurveMaker.cpp
@@ -13,12 +13,12 @@
 #include "CADKernel/Geo/Curves/RestrictionCurve.h"
 #include "CADKernel/Geo/Curves/SurfacicCurve.h"
 
-namespace CADKernel
+namespace UE::CADKernel
 {
 
 TSharedPtr<FCurve> FCurve::MakeNurbsCurve(FNurbsCurveData& InNurbsData)
 {
-	return FEntity::MakeShared<CADKernel::FNURBSCurve>(InNurbsData);
+	return FEntity::MakeShared<UE::CADKernel::FNURBSCurve>(InNurbsData);
 }
 
 }
\ No newline at end of file
diff --git a/Engine/Source/Runtime/Datasmith/CADKernel/Private/CADKernel/Geo/Curves/EllipseCurve.h b/Engine/Source/Runtime/Datasmith/CADKernel/Private/CADKernel/Geo/Curves/EllipseCurve.h
--- a/Engine/Source/Runtime/Datasmith/CADKernel/Private/CADKernel/Geo/Curves/EllipseCurve.h
+++ b/Engine/Source/Runtime/Datasmith/CADKernel/Private/CADKernel/Geo/Curves/EllipseCurve.h
@@ -5,111 +5,111 @@
 #include "CADKernel/Geo/GeoPoint.h"
 #include "CADKernel/Math/MathConst.h"
 
-namespace CADKernel
+namespace UE::CADKernel
 {
 
-	class CADKERNEL_API FEllipseCurve : public FCurve
+class CADKERNEL_API FEllipseCurve : public FCurve
+{
+	friend class FEntity;
+
+protected:
+	FMatrixH Matrix;
+	double   RadiusU;
+	double   RadiusV;
+	FMatrixH InverseMatrix;
+
+	FEllipseCurve(const FMatrixH& InMatrix, double InRadiusU, double InRadiusV, int8 InDimension = 3)
+		: FCurve(FLinearBoundary(0, PI * 2.), InDimension)
+		, Matrix(InMatrix)
+		, RadiusU(InRadiusU)
+		, RadiusV(InRadiusV)
+		, InverseMatrix(InMatrix.GetInverse())
 	{
-		friend class FEntity;
-
-	protected:
-		FMatrixH Matrix;
-		double   RadiusU;
-		double   RadiusV;
-		FMatrixH InverseMatrix;
-
-		FEllipseCurve(const FMatrixH& InMatrix, double InRadiusU, double InRadiusV, int8 InDimension = 3)
-			: FCurve(FLinearBoundary(0, PI * 2.), InDimension)
-			, Matrix(InMatrix)
-			, RadiusU(InRadiusU)
-			, RadiusV(InRadiusV)
-			, InverseMatrix(InMatrix.GetInverse())
-		{
-		}
-
-		FEllipseCurve(const FMatrixH& InMatrix, double InRadiusU, double InRadiusV, const FLinearBoundary& InBounds, int8 InDimension = 3)
-			: FCurve(InBounds, InDimension)
-			, Matrix(InMatrix)
-			, RadiusU(InRadiusU)
-			, RadiusV(InRadiusV)
-			, InverseMatrix(InMatrix.GetInverse())
-		{
-		}
+	}
+
+	FEllipseCurve(const FMatrixH& InMatrix, double InRadiusU, double InRadiusV, const FLinearBoundary& InBounds, int8 InDimension = 3)
+		: FCurve(InBounds, InDimension)
+		, Matrix(InMatrix)
+		, RadiusU(InRadiusU)
+		, RadiusV(InRadiusV)
+		, InverseMatrix(InMatrix.GetInverse())
+	{
+	}
 
-		FEllipseCurve() = default;
+	FEllipseCurve() = default;
 
-	public:
+public:
 
-		virtual void Serialize(FCADKernelArchive& Ar) override
-		{
-			FCurve::Serialize(Ar);
-			Ar << Matrix;
-			Ar << RadiusU;
-			Ar << RadiusV;
-			Ar << InverseMatrix;
-		}
+	virtual void Serialize(FCADKernelArchive& Ar) override
+	{
+		FCurve::Serialize(Ar);
+		Ar << Matrix;
+		Ar << RadiusU;
+		Ar << RadiusV;
+		Ar << InverseMatrix;
+	}
 
 #ifdef CADKERNEL_DEV
-		virtual FInfoEntity& GetInfo(FInfoEntity&) const override;
+	virtual FInfoEntity& GetInfo(FInfoEntity&) const override;
 #endif
 
-		virtual ECurve GetCurveType() const override
-		{
-			return ECurve::Ellipse;
-		}
+	virtual ECurve GetCurveType() const override
+	{
+		return ECurve::Ellipse;
+	}
 
-		const FMatrixH& GetMatrix() const 
-		{
-			return Matrix;
-		}
+	const FMatrixH& GetMatrix() const
+	{
+		return Matrix;
+	}
 
-		bool IsCircular() const
-		{
-			return FMath::IsNearlyZero(RadiusU - RadiusV);
-		}
+	bool IsCircular() const
+	{
+		return FMath::IsNearlyZero(RadiusU - RadiusV);
+	}
+
+	virtual TSharedPtr<FEntityGeom> ApplyMatrix(const FMatrixH& InMatrix) const override;
+
+	virtual void EvaluatePoint(double Coordinate, FCurvePoint& OutPoint, int32 DerivativeOrder = 0) const override
+	{
+		ensure(Dimension == 3);
+		Evaluate<FCurvePoint, FPoint>(Coordinate, OutPoint, DerivativeOrder);
+	}
+
+	virtual void Evaluate2DPoint(double Coordinate, FCurvePoint2D& OutPoint, int32 DerivativeOrder = 0) const override
+	{
+		ensure(Dimension == 2);
+		Evaluate<FCurvePoint2D, FPoint2D>(Coordinate, OutPoint, DerivativeOrder);
+	}
 
-		virtual TSharedPtr<FEntityGeom> ApplyMatrix(const FMatrixH& InMatrix) const override;
+private:
+
+	template <typename CurvePointType, typename PointType>
+	void Evaluate(double Coordinate, CurvePointType& OutPoint, int32 DerivativeOrder) const
+	{
+		OutPoint.DerivativeOrder = DerivativeOrder;
+
+		double CosU = cos(Coordinate);
+		double SinU = sin(Coordinate);
 
-		virtual void EvaluatePoint(double Coordinate, FCurvePoint& OutPoint, int32 DerivativeOrder = 0) const override
 		{
-			ensure(Dimension == 3);
-			Evaluate<FCurvePoint, FPoint>(Coordinate, OutPoint, DerivativeOrder);
+			PointType Result(RadiusU * CosU, RadiusV * SinU);
+			OutPoint.Point = Matrix.Multiply(Result);
 		}
 
-		virtual void Evaluate2DPoint(double Coordinate, FCurvePoint2D& OutPoint, int32 DerivativeOrder = 0) const override
+		if (DerivativeOrder > 0)
 		{
-			ensure(Dimension == 2);
-			Evaluate<FCurvePoint2D, FPoint2D>(Coordinate, OutPoint, DerivativeOrder);
+			PointType Result(RadiusU * (-SinU), RadiusV * CosU);
+			OutPoint.Gradient = Matrix.MultiplyVector(Result);
 		}
 
-	private:
-
-		template <typename CurvePointType, typename PointType>
-		void Evaluate(double Coordinate, CurvePointType& OutPoint, int32 DerivativeOrder) const
+		if (DerivativeOrder > 1)
 		{
-			OutPoint.DerivativeOrder = DerivativeOrder;
-
-			double CosU = cos(Coordinate);
-			double SinU = sin(Coordinate);
-
-			{
-				PointType Result(RadiusU * CosU, RadiusV * SinU);
-				OutPoint.Point = Matrix.Multiply(Result);
-			}
-
-			if (DerivativeOrder > 0)
-			{
-				PointType Result(RadiusU * (-SinU), RadiusV * CosU);
-				OutPoint.Gradient = Matrix.MultiplyVector(Result);
-			}
-
-			if (DerivativeOrder > 1)
-			{
-				PointType Result(RadiusU * (-CosU), RadiusV * (-SinU));
-				OutPoint.Laplacian = Matrix.MultiplyVector(Result);
-			}
+			PointType Result(RadiusU * (-CosU), RadiusV * (-SinU));
+			OutPoint.Laplacian = Matrix.MultiplyVector(Result);
 		}
-	};
+	}
+};
 
-} // namespace CADKernel
+} // namespace UE::CADKernel
 
diff --git a/Engine/Source/Runtime/Datasmith/CADKernel/Private/CADKernel/Geo/Curves/HyperbolaCurve.h b/Engine/Source/Runtime/Datasmith/CADKernel/Private/CADKernel/Geo/Curves/HyperbolaCurve.h
--- a/Engine/Source/Runtime/Datasmith/CADKernel/Private/CADKernel/Geo/Curves/HyperbolaCurve.h
+++ b/Engine/Source/Runtime/Datasmith/CADKernel/Private/CADKernel/Geo/Curves/HyperbolaCurve.h
@@ -3,96 +3,96 @@
 
 #include "CADKernel/Geo/Curves/Curve.h"
 
-namespace CADKernel
+namespace UE::CADKernel
 {
 
-	class CADKERNEL_API FHyperbolaCurve : public FCurve
+class CADKERNEL_API FHyperbolaCurve : public FCurve
+{
+	friend class FIGESEntity104;
+	friend class FEntity;
+	friend class FEntity;
+
+protected:
+	FMatrixH Matrix;
+	double SemiMajorAxis;
+	double SemiImaginaryAxis;
+
+	FHyperbolaCurve(const FMatrixH& InMatrix, double InSemiAxis, double InSemiImagAxis, const FLinearBoundary& InBounds, int8 InDimension = 3)
+		: FCurve(InBounds, InDimension)
+		, Matrix(InMatrix)
+		, SemiMajorAxis(InSemiAxis)
+		, SemiImaginaryAxis(InSemiImagAxis)
 	{
-		friend class FIGESEntity104;
-		friend class FEntity;
-		friend class FEntity;
-
-	protected:
-		FMatrixH Matrix;
-		double SemiMajorAxis;
-		double SemiImaginaryAxis;
-
-		FHyperbolaCurve(const FMatrixH& InMatrix, double InSemiAxis, double InSemiImagAxis, const FLinearBoundary& InBounds, int8 InDimension = 3)
-			: FCurve(InBounds, InDimension)
-			, Matrix(InMatrix)
-			, SemiMajorAxis(InSemiAxis)
-			, SemiImaginaryAxis(InSemiImagAxis)
-		{
-		}
+	}
 
-		FHyperbolaCurve() = default;
+	FHyperbolaCurve() = default;
 
-	public:
+public:
 
-		virtual void Serialize(FCADKernelArchive& Ar) override
-		{
-			FCurve::Serialize(Ar);
-			Ar << Matrix;
-			Ar << SemiMajorAxis;
-			Ar << SemiImaginaryAxis;
-		}
+	virtual void Serialize(FCADKernelArchive& Ar) override
+	{
+		FCurve::Serialize(Ar);
+		Ar << Matrix;
+		Ar << SemiMajorAxis;
+		Ar << SemiImaginaryAxis;
+	}
 
 #ifdef CADKERNEL_DEV
-		virtual FInfoEntity& GetInfo(FInfoEntity&) const override;
+	virtual FInfoEntity& GetInfo(FInfoEntity&) const override;
 #endif
 
-		virtual ECurve GetCurveType() const override
-		{
-			return ECurve::Hyperbola;
-		}
+	virtual ECurve GetCurveType() const override
+	{
+		return ECurve::Hyperbola;
+	}
 
-		FMatrixH& GetMatrix() 
-		{
-			return Matrix;
-		}
+	FMatrixH& GetMatrix()
+	{
+		return Matrix;
+	}
+
+	virtual TSharedPtr<FEntityGeom> ApplyMatrix(const FMatrixH& InMatrix) const override;
+
+	virtual void EvaluatePoint(double Coordinate, FCurvePoint& OutPoint, int32 DerivativeOrder = 0) const override
+	{
+		ensure(Dimension == 3);
+		Evaluate<FCurvePoint, FPoint>(Coordinate, OutPoint, DerivativeOrder);
+	}
+
+	virtual void Evaluate2DPoint(double Coordinate, FCurvePoint2D& OutPoint, int32 DerivativeOrder = 0) const override
+	{
+		ensure(Dimension == 2);
+		Evaluate<FCurvePoint2D, FPoint2D>(Coordinate, OutPoint, DerivativeOrder);
+	}
+
+private:
 
-		virtual TSharedPtr<FEntityGeom> ApplyMatrix(const FMatrixH& InMatrix) const override;
+	template <typename CurvePointType, typename PointType>
+	void Evaluate(double Coordinate, CurvePointType& OutPoint, int32 DerivativeOrder) const
+	{
+		OutPoint.DerivativeOrder = DerivativeOrder;
+
+		const double CosHUCoord = cosh(Coordinate);
+		const double SinHUCoord = sinh(Coordinate);
 
-		virtual void EvaluatePoint(double Coordinate, FCurvePoint& OutPoint, int32 DerivativeOrder = 0) const override
 		{
-			ensure(Dimension == 3);
-			Evaluate<FCurvePoint, FPoint>(Coordinate, OutPoint, DerivativeOrder);
+			PointType Result(SemiMajorAxis * CosHUCoord, SemiImaginaryAxis * SinHUCoord);
+			OutPoint.Point = Matrix.Multiply(Result);
 		}
 
-		virtual void Evaluate2DPoint(double Coordinate, FCurvePoint2D& OutPoint, int32 DerivativeOrder = 0) const override
+		if (DerivativeOrder > 0)
 		{
-			ensure(Dimension == 2);
-			Evaluate<FCurvePoint2D, FPoint2D>(Coordinate, OutPoint, DerivativeOrder);
+			FPoint Result(SemiMajorAxis * SinHUCoord, SemiImaginaryAxis * CosHUCoord);
+			OutPoint.Gradient = Matrix.MultiplyVector(Result);
 		}
 
-	private:
-
-		template <typename CurvePointType, typename PointType>
-		void Evaluate(double Coordinate, CurvePointType& OutPoint, int32 DerivativeOrder) const
+		if (DerivativeOrder > 1)
 		{
-			OutPoint.DerivativeOrder = DerivativeOrder;
-
-			const double CosHUCoord = cosh(Coordinate);
-			const double SinHUCoord = sinh(Coordinate);
-
-			{
-				PointType Result(SemiMajorAxis * CosHUCoord, SemiImaginaryAxis * SinHUCoord);
-				OutPoint.Point = Matrix.Multiply(Result);
-			}
-
-			if (DerivativeOrder > 0)
-			{
-				FPoint Result(SemiMajorAxis * SinHUCoord, SemiImaginaryAxis * CosHUCoord);
-				OutPoint.Gradient = Matrix.MultiplyVector(Result);
-			}
-
-			if (DerivativeOrder > 1)
-			{
-				FPoint Result(SemiMajorAxis * CosHUCoord,  SemiImaginaryAxis * SinHUCoord);
-				OutPoint.Laplacian = Matrix.MultiplyVector(Result);
-			}
+			FPoint Result(SemiMajorAxis * CosHUCoord, SemiImaginaryAxis * SinHUCoord);
+			OutPoint.Laplacian = Matrix.MultiplyVector(Result);
 		}
-	};
+	}
+};
 
-} // namespace CADKernel
+} // namespace UE::CADKernel
 
diff --git a/Engine/Source/Runtime/Datasmith/CADKernel/Private/CADKernel/Geo/Curves/NURBSCurve.h b/Engine/Source/Runtime/Datasmith/CADKernel/Private/CADKernel/Geo/Curves/NURBSCurve.h
--- a/Engine/Source/Runtime/Datasmith/CADKernel/Private/CADKernel/Geo/Curves/NURBSCurve.h
+++ b/Engine/Source/Runtime/Datasmith/CADKernel/Private/CADKernel/Geo/Curves/NURBSCurve.h
@@ -4,138 +4,138 @@
 #include "CADKernel/Geo/Curves/Curve.h"
 #include "CADKernel/Math/BSpline.h"
 
-namespace CADKernel
+namespace UE::CADKernel
 {
-	struct FNurbsCurveData
-	{
-		bool bIsRational;
-		int32 Dimension = 0;
+struct FNurbsCurveData
+{
+	bool bIsRational;
+	int32 Dimension = 0;
 
-		int32 Degree = 0;
-		TArray<double> NodalVector;
+	int32 Degree = 0;
+	TArray<double> NodalVector;
 
-		TArray<double> Weights;
-		TArray<FPoint> Poles;
-	};
+	TArray<double> Weights;
+	TArray<FPoint> Poles;
+};
 
-	class CADKERNEL_API FNURBSCurve : public FCurve
-	{
-		friend class FEntity;
+class CADKERNEL_API FNURBSCurve : public FCurve
+{
+	friend class FEntity;
 
-	protected:
+protected:
 
-		int32 Degree;
+	int32 Degree;
 
-		TArray<double> NodalVector;
-		TArray<double> Weights;
+	TArray<double> NodalVector;
+	TArray<double> Weights;
 
-		TArray<FPoint> Poles;
+	TArray<FPoint> Poles;
 
-		bool bIsRational;
+	bool bIsRational;
 
-		/**
-		 * Data generated at initialization which are not serialized
-		 */
-		TArray<double> HomogeneousPoles;
-		int32 PoleDimension;
+	/**
+	 * Data generated at initialization which are not serialized
+	 */
+	TArray<double> HomogeneousPoles;
+	int32 PoleDimension;
 
-		FNURBSCurve(int32 InDegree, const TArray<double>& InNodalVector, const TArray<FPoint>& InPoles, const TArray<double>& InWeights, int8 InDimension = 3);
-		FNURBSCurve(FNurbsCurveData& NurbsCurveData);
-		FNURBSCurve(const FNURBSCurve& Nurbs);
+	FNURBSCurve(int32 InDegree, const TArray<double>& InNodalVector, const TArray<FPoint>& InPoles, const TArray<double>& InWeights, int8 InDimension = 3);
+	FNURBSCurve(FNurbsCurveData& NurbsCurveData);
+	FNURBSCurve(const FNURBSCurve& Nurbs);
 
-		FNURBSCurve() = default;
+	FNURBSCurve() = default;
 
-	public:
+public:
 
-		virtual void Serialize(FCADKernelArchive& Ar) override
+	virtual void Serialize(FCADKernelArchive& Ar) override
+	{
+		FCurve::Serialize(Ar);
+		Ar << Degree;
+		Ar.Serialize(NodalVector);
+		Ar.Serialize(Weights);
+		Ar.Serialize(Poles);
+		Ar << bIsRational;
+
+		if (Ar.IsLoading())
 		{
-			FCurve::Serialize(Ar);
-			Ar << Degree;
-			Ar.Serialize(NodalVector);
-			Ar.Serialize(Weights);
-			Ar.Serialize(Poles);
-			Ar << bIsRational;
-
-			if (Ar.IsLoading())
-			{
-				Finalize();
-			}
+			Finalize();
 		}
+	}
 
 #ifdef CADKERNEL_DEV
-		virtual FInfoEntity& GetInfo(FInfoEntity&) const override;
+	virtual FInfoEntity& GetInfo(FInfoEntity&) const override;
 #endif
 
-		virtual ECurve GetCurveType() const override
-		{
-			return ECurve::Nurbs;
-		}
+	virtual ECurve GetCurveType() const override
+	{
+		return ECurve::Nurbs;
+	}
 
-		int32 GetDegree() const
-		{
-			return Degree;
-		}
+	int32 GetDegree() const
+	{
+		return Degree;
+	}
 
-		const int32 GetPoleCount() const
-		{
-			return Poles.Num();
-		}
+	const int32 GetPoleCount() const
+	{
+		return Poles.Num();
+	}
 
-		const TArray<FPoint>& GetPoles() const
-		{
-			return Poles;
-		}
+	const TArray<FPoint>& GetPoles() const
+	{
+		return Poles;
+	}
 
-		const TArray<double>& GetWeights() const
-		{
-			return Weights;
-		}
+	const TArray<double>& GetWeights() const
+	{
+		return Weights;
+	}
 
-		TArray<double> GetHPoles() const
-		{
-			return HomogeneousPoles;
-		}
-		
-		const TArray<double>& GetNodalVector() const
-		{
-			return NodalVector;
-		}
+	TArray<double> GetHPoles() const
+	{
+		return HomogeneousPoles;
+	}
 
-		bool IsRational() const
-		{
-			return bIsRational;
-		}
+	const TArray<double>& GetNodalVector() const
+	{
+		return NodalVector;
+	}
 
-		virtual TSharedPtr<FEntityGeom> ApplyMatrix(const FMatrixH& InMatrix) const override;
+	bool IsRational() const
+	{
+		return bIsRational;
+	}
 
-		virtual void EvaluatePoint(double Coordinate, FCurvePoint& OutPoint, int32 DerivativeOrder = 0) const override
-		{
-			BSpline::EvaluatePoint(*this, Coordinate, OutPoint, DerivativeOrder);
-		}
+	virtual TSharedPtr<FEntityGeom> ApplyMatrix(const FMatrixH& InMatrix) const override;
 
-		virtual void Evaluate2DPoint(double Coordinate, FCurvePoint2D& OutPoint, int32 DerivativeOrder = 0) const override
-		{
-			BSpline::Evaluate2DPoint(*this, Coordinate, OutPoint, DerivativeOrder);
-		}
+	virtual void EvaluatePoint(double Coordinate, FCurvePoint& OutPoint, int32 DerivativeOrder = 0) const override
+	{
+		BSpline::EvaluatePoint(*this, Coordinate, OutPoint, DerivativeOrder);
+	}
 
-		virtual void FindNotDerivableCoordinates(const FLinearBoundary& InBoundary, int32 DerivativeOrder, TArray<double>& OutNotDerivableCoordinates) const override
-		{
-			BSpline::FindNotDerivableParameters(*this, DerivativeOrder, InBoundary, OutNotDerivableCoordinates);
-		}
+	virtual void Evaluate2DPoint(double Coordinate, FCurvePoint2D& OutPoint, int32 DerivativeOrder = 0) const override
+	{
+		BSpline::Evaluate2DPoint(*this, Coordinate, OutPoint, DerivativeOrder);
+	}
+
+	virtual void FindNotDerivableCoordinates(const FLinearBoundary& InBoundary, int32 DerivativeOrder, TArray<double>& OutNotDerivableCoordinates) const override
+	{
+		BSpline::FindNotDerivableParameters(*this, DerivativeOrder, InBoundary, OutNotDerivableCoordinates);
+	}
 
-		virtual void ExtendTo(const FPoint& DesiredPosition) override;
+	virtual void ExtendTo(const FPoint& DesiredPosition) override;
 
-		void Invert();
-		void SetStartNodalCoordinate(double NewStartBoundary);
+	void Invert();
+	void SetStartNodalCoordinate(double NewStartBoundary);
 
-	private:
+private:
 
-		/**
-		 * Fill homogeneous points and set bounds
-		 */
-		void Finalize();
+	/**
+	 * Fill homogeneous points and set bounds
+	 */
+	void Finalize();
 
-	};
+};
 
-} // namespace CADKernel
+} // namespace UE::CADKernel
 
diff --git a/Engine/Source/Runtime/Datasmith/CADKernel/Private/CADKernel/Geo/Curves/ParabolaCurve.h b/Engine/Source/Runtime/Datasmith/CADKernel/Private/CADKernel/Geo/Curves/ParabolaCurve.h
--- a/Engine/Source/Runtime/Datasmith/CADKernel/Private/CADKernel/Geo/Curves/ParabolaCurve.h
+++ b/Engine/Source/Runtime/Datasmith/CADKernel/Private/CADKernel/Geo/Curves/ParabolaCurve.h
@@ -3,89 +3,89 @@
 
 #include "CADKernel/Geo/Curves/Curve.h"
 
-namespace CADKernel
+namespace UE::CADKernel
 {
 
-	class CADKERNEL_API FParabolaCurve : public FCurve
-	{
-		friend class FIGESEntity104;
-		friend class FEntity;
+class CADKERNEL_API FParabolaCurve : public FCurve
+{
+	friend class FIGESEntity104;
+	friend class FEntity;
 
-	protected:
-		FMatrixH Matrix;
-		double FocalDistance;
+protected:
+	FMatrixH Matrix;
+	double FocalDistance;
 
-		FParabolaCurve(const FMatrixH& InMatrix, double InFocalDistance, const FLinearBoundary& InBounds, int8 InDimension = 3)
-			: FCurve(InBounds, InDimension)
-			, Matrix(InMatrix)
-			, FocalDistance(InFocalDistance)
-		{
-		}
+	FParabolaCurve(const FMatrixH& InMatrix, double InFocalDistance, const FLinearBoundary& InBounds, int8 InDimension = 3)
+		: FCurve(InBounds, InDimension)
+		, Matrix(InMatrix)
+		, FocalDistance(InFocalDistance)
+	{
+	}
 
-		FParabolaCurve() = default;
+	FParabolaCurve() = default;
 
-	public:
+public:
 
-		virtual void Serialize(FCADKernelArchive& Ar) override
-		{
-			FCurve::Serialize(Ar);
-			Ar << Matrix;
-			Ar << FocalDistance;
-		}
+	virtual void Serialize(FCADKernelArchive& Ar) override
+	{
+		FCurve::Serialize(Ar);
+		Ar << Matrix;
+		Ar << FocalDistance;
+	}
 
 #ifdef CADKERNEL_DEV
-		virtual FInfoEntity& GetInfo(FInfoEntity&) const override;
+	virtual FInfoEntity& GetInfo(FInfoEntity&) const override;
 #endif
 
-		virtual ECurve GetCurveType() const override
-		{
-			return ECurve::Parabola;
-		}
+	virtual ECurve GetCurveType() const override
+	{
+		return ECurve::Parabola;
+	}
 
-		virtual TSharedPtr<FEntityGeom> ApplyMatrix(const FMatrixH& InMatrix) const override;
+	virtual TSharedPtr<FEntityGeom> ApplyMatrix(const FMatrixH& InMatrix) const override;
 
-		virtual void ExtendTo(const FPoint& Point) override
-		{
-			ensureCADKernel(false);
-		}
+	virtual void ExtendTo(const FPoint& Point) override
+	{
+		ensureCADKernel(false);
+	}
+
+	virtual void EvaluatePoint(double Coordinate, FCurvePoint& OutPoint, int32 DerivativeOrder = 0) const override
+	{
+		ensure(Dimension == 3);
+		Evaluate<FCurvePoint, FPoint>(Coordinate, OutPoint, DerivativeOrder);
+	}
+
+	virtual void Evaluate2DPoint(double Coordinate, FCurvePoint2D& OutPoint, int32 DerivativeOrder = 0) const override
+	{
+		ensure(Dimension == 2);
+		Evaluate<FCurvePoint2D, FPoint2D>(Coordinate, OutPoint, DerivativeOrder);
+	}
+
+private:
+
+	template <typename CurvePointType, typename PointType>
+	void Evaluate(double Coordinate, CurvePointType& OutPoint, int32 DerivativeOrder) const
+	{
+		OutPoint.DerivativeOrder = DerivativeOrder;
 
-		virtual void EvaluatePoint(double Coordinate, FCurvePoint& OutPoint, int32 DerivativeOrder = 0) const override
 		{
-			ensure(Dimension == 3);
-			Evaluate<FCurvePoint, FPoint>(Coordinate, OutPoint, DerivativeOrder);
+			PointType Result(Coordinate * Coordinate * FocalDistance, Coordinate);
+			OutPoint.Point = Matrix.Multiply(Result);
 		}
 
-		virtual void Evaluate2DPoint(double Coordinate, FCurvePoint2D& OutPoint, int32 DerivativeOrder = 0) const override
+		if (DerivativeOrder > 0)
 		{
-			ensure(Dimension == 2);
-			Evaluate<FCurvePoint2D, FPoint2D>(Coordinate, OutPoint, DerivativeOrder);
+			FPoint Result(2.0 * Coordinate * FocalDistance, 1);
+			OutPoint.Gradient = Matrix.MultiplyVector(Result);
 		}
 
-	private:
-
-		template <typename CurvePointType, typename PointType>
-		void Evaluate(double Coordinate, CurvePointType& OutPoint, int32 DerivativeOrder) const
+		if (DerivativeOrder > 1)
 		{
-			OutPoint.DerivativeOrder = DerivativeOrder;
-
-			{
-				PointType Result(Coordinate * Coordinate * FocalDistance, Coordinate);
-				OutPoint.Point = Matrix.Multiply(Result);
-			}
-
-			if (DerivativeOrder > 0)
-			{
-				FPoint Result(2.0 * Coordinate * FocalDistance, 1);
-				OutPoint.Gradient = Matrix.MultiplyVector(Result);
-			}
-
-			if (DerivativeOrder > 1)
-			{
-				FPoint Result(2.0 * FocalDistance, 0.);
-				OutPoint.Laplacian = Matrix.MultiplyVector(Result);
-			}
+			FPoint Result(2.0 * FocalDistance, 0.);
+			OutPoint.Laplacian = Matrix.MultiplyVector(Result);
 		}
-	};
+	}
+};
 
-} // namespace CADKernel
+} // namespace UE::CADKernel
 
diff --git a/Engine/Source/Runtime/Datasmith/CADKernel/Private/CADKernel/Geo/Curves/PolylineCurve.h b/Engine/Source/Runtime/Datasmith/CADKernel/Private/CADKernel/Geo/Curves/PolylineCurve.h
--- a/Engine/Source/Runtime/Datasmith/CADKernel/Private/CADKernel/Geo/Curves/PolylineCurve.h
+++ b/Engine/Source/Runtime/Datasmith/CADKernel/Private/CADKernel/Geo/Curves/PolylineCurve.h
@@ -5,213 +5,213 @@
 #include "CADKernel/Geo/Sampling/Polyline.h"
 #include "CADKernel/Geo/Sampling/PolylineTools.h"
 
-namespace CADKernel
+namespace UE::CADKernel
 {
-	class FInfoEntity;
-	class FSurface;
+class FInfoEntity;
+class FSurface;
 
-	template<class PointType, class PointCurveType>
-	class CADKERNEL_API TPolylineCurve : public FCurve, public TPolyline<PointType>
-	{
-		friend class FEntity;
-		friend class FPolylineTools;
-
-	protected:
-		
-		TPolylineApproximator<PointType> Approximator;
+template<class PointType, class PointCurveType>
+class CADKERNEL_API TPolylineCurve : public FCurve, public TPolyline<PointType>
+{
+	friend class FEntity;
+	friend class FPolylineTools;
 
-		TPolylineCurve(const TArray<PointType>& InPoints, const TArray<double>& InCoordinates, int8 InDimension)
-			: FCurve(InDimension)
-			, Approximator(this->Coordinates, this->Points)
-		{
-			this->Coordinates = InCoordinates;
-			this->Points = InPoints;
-			ensureCADKernel(this->Coordinates[0] < this->Coordinates.Last());
-			Boundary.Set(this->Coordinates.HeapTop(), this->Coordinates.Last());
-		}
+protected:
 
-		TPolylineCurve(const TArray<PointType>& InPoints, int8 InDimension)
-			: FCurve(InDimension)
-			, TPolyline<PointType>(InPoints)
-			, Approximator(this->Coordinates, this->Points)
-		{
-			this->Coordinates.Reserve(InPoints.Num());
-			this->Coordinates.Add(0.);
+	TPolylineApproximator<PointType> Approximator;
 
-			double CurvilineLength = 0;
-			for (int32 iPoint = 1; iPoint < this->Points.Num(); iPoint++)
-			{
-				CurvilineLength += this->Points[iPoint].Distance(this->Points[iPoint - 1]);
-				this->Coordinates.Add(CurvilineLength);
-			}
-
-			Boundary.Set(0, CurvilineLength);
-		}
+	TPolylineCurve(const TArray<PointType>& InPoints, const TArray<double>& InCoordinates, int8 InDimension)
+		: FCurve(InDimension)
+		, Approximator(this->Coordinates, this->Points)
+	{
+		this->Coordinates = InCoordinates;
+		this->Points = InPoints;
+		ensureCADKernel(this->Coordinates[0] < this->Coordinates.Last());
+		Boundary.Set(this->Coordinates.HeapTop(), this->Coordinates.Last());
+	}
+
+	TPolylineCurve(const TArray<PointType>& InPoints, int8 InDimension)
+		: FCurve(InDimension)
+		, TPolyline<PointType>(InPoints)
+		, Approximator(this->Coordinates, this->Points)
+	{
+		this->Coordinates.Reserve(InPoints.Num());
+		this->Coordinates.Add(0.);
 
-		TPolylineCurve()
-			: Approximator(this->Coordinates, this->Points)
+		double CurvilineLength = 0;
+		for (int32 iPoint = 1; iPoint < this->Points.Num(); iPoint++)
 		{
+			CurvilineLength += this->Points[iPoint].Distance(this->Points[iPoint - 1]);
+			this->Coordinates.Add(CurvilineLength);
 		}
 
-	public:
+		Boundary.Set(0, CurvilineLength);
+	}
 
-		virtual void Serialize(FCADKernelArchive& Ar) override
-		{
-			FCurve::Serialize(Ar);
-			TPolyline<PointType>::Serialize(Ar);
-		}
+	TPolylineCurve()
+		: Approximator(this->Coordinates, this->Points)
+	{
+	}
 
-		virtual void EvaluateCurvesPoint(double InCoordinate, PointCurveType& OutPoint, int32 InDerivativeOrder = 0) const
-		{
-			Approximator.ApproximatePoint(InCoordinate, OutPoint, InDerivativeOrder);
-		}
+public:
 
-		PointType EvaluatePointAt(double InCoordinate) const
-		{
-			return Approximator.ApproximatePoint(InCoordinate);
-		}
+	virtual void Serialize(FCADKernelArchive& Ar) override
+	{
+		FCurve::Serialize(Ar);
+		TPolyline<PointType>::Serialize(Ar);
+	}
 
-		virtual void EvaluateCurvesPoints(const TArray<double>& InCoordinates, TArray<PointCurveType>& OutPoints, int32 InDerivativeOrder = 0) const
-		{
-			Approximator.ApproximatePoints(InCoordinates, OutPoints, InDerivativeOrder);
-		}
+	virtual void EvaluateCurvesPoint(double InCoordinate, PointCurveType& OutPoint, int32 InDerivativeOrder = 0) const
+	{
+		Approximator.ApproximatePoint(InCoordinate, OutPoint, InDerivativeOrder);
+	}
 
-		virtual double ComputeSubLength(const FLinearBoundary& InBoundary) const
-		{
-			return Approximator.ComputeLengthOfSubPolyline(InBoundary);
-		}
+	PointType EvaluatePointAt(double InCoordinate) const
+	{
+		return Approximator.ApproximatePoint(InCoordinate);
+	}
 
-		const TArray<PointType>& GetPolylinePoints() const
-		{
-			return this->Points;
-		}
+	virtual void EvaluateCurvesPoints(const TArray<double>& InCoordinates, TArray<PointCurveType>& OutPoints, int32 InDerivativeOrder = 0) const
+	{
+		Approximator.ApproximatePoints(InCoordinates, OutPoints, InDerivativeOrder);
+	}
 
-		const TArray<double>& GetPolylineParameters() const
-		{
-			return this->Coordinates;
-		}
+	virtual double ComputeSubLength(const FLinearBoundary& InBoundary) const
+	{
+		return Approximator.ComputeLengthOfSubPolyline(InBoundary);
+	}
 
-		virtual void FindNotDerivableCoordinates(const FLinearBoundary& InBoundary, int32 DerivativeOrder, TArray<double>& OutNotDerivableCoordinates) const override
-		{
-			//TO DO
-			ensureCADKernel(false);
-		}
+	const TArray<PointType>& GetPolylinePoints() const
+	{
+		return this->Points;
+	}
 
-		void SetPoints(const TArray<PointType>& InPoints)
-		{
-			this->Points = InPoints;
-			GlobalLength.Empty();
-		}
+	const TArray<double>& GetPolylineParameters() const
+	{
+		return this->Coordinates;
+	}
 
-		template<class PolylineType>
-		TSharedPtr<FEntityGeom> ApplyMatrixImpl(const FMatrixH& InMatrix) const
-		{
-			TArray<PointType> NewPoints;
-			NewPoints.Reserve(this->Points.Num());
+	virtual void FindNotDerivableCoordinates(const FLinearBoundary& InBoundary, int32 DerivativeOrder, TArray<double>& OutNotDerivableCoordinates) const override
+	{
+		//TO DO
+		ensureCADKernel(false);
+	}
 
-			for (PointType Point : this->Points)
-			{
-				NewPoints.Emplace(InMatrix.Multiply(Point));
-			}
+	void SetPoints(const TArray<PointType>& InPoints)
+	{
+		this->Points = InPoints;
+		GlobalLength.Empty();
+	}
 
-			return FEntity::MakeShared<PolylineType>(NewPoints, this->Coordinates);
-		}
+	template<class PolylineType>
+	TSharedPtr<FEntityGeom> ApplyMatrixImpl(const FMatrixH& InMatrix) const
+	{
+		TArray<PointType> NewPoints;
+		NewPoints.Reserve(this->Points.Num());
 
-		virtual void ExtendTo(const FPoint& DesiredPoint) override
+		for (PointType Point : this->Points)
 		{
-			PolylineTools::ExtendTo(this->Points, (PointType) DesiredPoint);
+			NewPoints.Emplace(InMatrix.Multiply(Point));
 		}
-	};
 
-	class CADKERNEL_API FPolylineCurve : public TPolylineCurve<FPoint, FCurvePoint>
+		return FEntity::MakeShared<PolylineType>(NewPoints, this->Coordinates);
+	}
+
+	virtual void ExtendTo(const FPoint& DesiredPoint) override
 	{
-		friend class FEntity;
+		PolylineTools::ExtendTo(this->Points, (PointType)DesiredPoint);
+	}
+};
 
-	protected:
-		FPolylineCurve(const TArray<FPoint>& InPoints, const TArray<double>& InCoordinates)
-			: TPolylineCurve<FPoint, FCurvePoint>(InPoints, InCoordinates, 3)
-		{
-		}
+class CADKERNEL_API FPolylineCurve : public TPolylineCurve<FPoint, FCurvePoint>
+{
+	friend class FEntity;
 
-		FPolylineCurve(const TArray<FPoint>& InPoints)
-			: TPolylineCurve<FPoint, FCurvePoint>(InPoints, 3)
-		{
-		}
+protected:
+	FPolylineCurve(const TArray<FPoint>& InPoints, const TArray<double>& InCoordinates)
+		: TPolylineCurve<FPoint, FCurvePoint>(InPoints, InCoordinates, 3)
+	{
+	}
 
-		FPolylineCurve() = default;
+	FPolylineCurve(const TArray<FPoint>& InPoints)
+		: TPolylineCurve<FPoint, FCurvePoint>(InPoints, 3)
+	{
+	}
 
-	public:
+	FPolylineCurve() = default;
 
-		virtual ECurve GetCurveType() const override
-		{
-			return ECurve::Polyline3D;
-		}
+public:
 
-		virtual TSharedPtr<FEntityGeom> ApplyMatrix(const FMatrixH& InMatrix) const override
-		{
-			return ApplyMatrixImpl<FPolylineCurve>(InMatrix);
-		}
+	virtual ECurve GetCurveType() const override
+	{
+		return ECurve::Polyline3D;
+	}
 
-		virtual FPoint EvaluatePoint(double InCoordinate) const override
-		{
-			return EvaluatePointAt(InCoordinate);
-		} 
+	virtual TSharedPtr<FEntityGeom> ApplyMatrix(const FMatrixH& InMatrix) const override
+	{
+		return ApplyMatrixImpl<FPolylineCurve>(InMatrix);
+	}
+
+	virtual FPoint EvaluatePoint(double InCoordinate) const override
+	{
+		return EvaluatePointAt(InCoordinate);
+	}
 
 #ifdef CADKERNEL_DEV
-		virtual FInfoEntity& GetInfo(FInfoEntity& Info) const override;
+	virtual FInfoEntity& GetInfo(FInfoEntity& Info) const override;
 #endif
 
-	};
+};
 
-	class CADKERNEL_API FPolyline2DCurve : public TPolylineCurve<FPoint2D, FCurvePoint2D>
-	{
-		friend class FEntity;
+class CADKERNEL_API FPolyline2DCurve : public TPolylineCurve<FPoint2D, FCurvePoint2D>
+{
+	friend class FEntity;
 
-	protected:
-		FPolyline2DCurve(const TArray<FPoint2D>& InPoints, const TArray<double>& InCoordinates)
-			: TPolylineCurve<FPoint2D, FCurvePoint2D>(InPoints, InCoordinates, 3)
-		{
-		}
+protected:
+	FPolyline2DCurve(const TArray<FPoint2D>& InPoints, const TArray<double>& InCoordinates)
+		: TPolylineCurve<FPoint2D, FCurvePoint2D>(InPoints, InCoordinates, 3)
+	{
+	}
 
-		FPolyline2DCurve(const TArray<FPoint2D>& InPoints)
-			: TPolylineCurve<FPoint2D, FCurvePoint2D>(InPoints, 3)
-		{
-		}
+	FPolyline2DCurve(const TArray<FPoint2D>& InPoints)
+		: TPolylineCurve<FPoint2D, FCurvePoint2D>(InPoints, 3)
+	{
+	}
 
-		FPolyline2DCurve() = default;
+	FPolyline2DCurve() = default;
 
-	public:
+public:
 
-		virtual ECurve GetCurveType() const override
-		{
-			return ECurve::Polyline3D;
-		}
+	virtual ECurve GetCurveType() const override
+	{
+		return ECurve::Polyline3D;
+	}
 
-		virtual TSharedPtr<FEntityGeom> ApplyMatrix(const FMatrixH& InMatrix) const override
-		{
-			return ApplyMatrixImpl<FPolyline2DCurve>(InMatrix);
-		}
+	virtual TSharedPtr<FEntityGeom> ApplyMatrix(const FMatrixH& InMatrix) const override
+	{
+		return ApplyMatrixImpl<FPolyline2DCurve>(InMatrix);
+	}
 
-		virtual void Evaluate2DPoint(double InCoordinate, FCurvePoint2D& OutPoint, int32 InDerivativeOrder = 0) const override
-		{
-			EvaluateCurvesPoint(InCoordinate, OutPoint, InDerivativeOrder);
-		}
+	virtual void Evaluate2DPoint(double InCoordinate, FCurvePoint2D& OutPoint, int32 InDerivativeOrder = 0) const override
+	{
+		EvaluateCurvesPoint(InCoordinate, OutPoint, InDerivativeOrder);
+	}
 
-		virtual FPoint2D Evaluate2DPoint(double InCoordinate) const override
-		{
-			return EvaluatePointAt(InCoordinate);
-		}
+	virtual FPoint2D Evaluate2DPoint(double InCoordinate) const override
+	{
+		return EvaluatePointAt(InCoordinate);
+	}
 
-		virtual void Evaluate2DPoints(const TArray<double>& InCoordinates, TArray<FCurvePoint2D>& OutPoints, int32 InDerivativeOrder = 0) const override
-		{
-			EvaluateCurvesPoints(InCoordinates, OutPoints, InDerivativeOrder);
-		}
+	virtual void Evaluate2DPoints(const TArray<double>& InCoordinates, TArray<FCurvePoint2D>& OutPoints, int32 InDerivativeOrder = 0) const override
+	{
+		EvaluateCurvesPoints(InCoordinates, OutPoints, InDerivativeOrder);
+	}
 
 #ifdef CADKERNEL_DEV
-		virtual FInfoEntity& GetInfo(FInfoEntity& Info) const override;
+	virtual FInfoEntity& GetInfo(FInfoEntity& Info) const override;
 #endif
 
-	}; 
-	
-} // namespace CADKernel
+};
+
+} // namespace UE::CADKernel
 
diff --git a/Engine/Source/Runtime/Datasmith/CADKernel/Private/CADKernel/Geo/Curves/SegmentCurve.h b/Engine/Source/Runtime/Datasmith/CADKernel/Private/CADKernel/Geo/Curves/SegmentCurve.h
--- a/Engine/Source/Runtime/Datasmith/CADKernel/Private/CADKernel/Geo/Curves/SegmentCurve.h
+++ b/Engine/Source/Runtime/Datasmith/CADKernel/Private/CADKernel/Geo/Curves/SegmentCurve.h
@@ -4,98 +4,98 @@
 #include "CADKernel/Geo/Curves/Curve.h"
 #include "CADKernel/Geo/GeoEnum.h"
 
-namespace CADKernel
+namespace UE::CADKernel
 {
 
-	class CADKERNEL_API FSegmentCurve : public FCurve
-	{
-		friend class FEntity;
+class CADKERNEL_API FSegmentCurve : public FCurve
+{
+	friend class FEntity;
 
-	protected:
-		FPoint StartPoint;
-		FPoint EndPoint;
+protected:
+	FPoint StartPoint;
+	FPoint EndPoint;
 
-		FSegmentCurve(const FPoint& InStartPoint, const FPoint& InEndPoint, int8 InDimension = 3)
-			: FCurve(InDimension)
-			, StartPoint(InStartPoint)
-			, EndPoint(InEndPoint)
-		{
-		}
+	FSegmentCurve(const FPoint& InStartPoint, const FPoint& InEndPoint, int8 InDimension = 3)
+		: FCurve(InDimension)
+		, StartPoint(InStartPoint)
+		, EndPoint(InEndPoint)
+	{
+	}
 
-		FSegmentCurve() = default;
+	FSegmentCurve() = default;
 
-	public:
+public:
 
-		virtual void Serialize(FCADKernelArchive& Ar) override
-		{
-			FCurve::Serialize(Ar);
-			Ar << StartPoint;
-			Ar << EndPoint;
-		}
+	virtual void Serialize(FCADKernelArchive& Ar) override
+	{
+		FCurve::Serialize(Ar);
+		Ar << StartPoint;
+		Ar << EndPoint;
+	}
 
 #ifdef CADKERNEL_DEV
-		virtual FInfoEntity& GetInfo(FInfoEntity&) const override;
+	virtual FInfoEntity& GetInfo(FInfoEntity&) const override;
 #endif
 
-		virtual ECurve GetCurveType() const override
-		{
-			return ECurve::Segment;
-		}
+	virtual ECurve GetCurveType() const override
+	{
+		return ECurve::Segment;
+	}
 
-		virtual TSharedPtr<FEntityGeom> ApplyMatrix(const FMatrixH& InMatrix) const override;
+	virtual TSharedPtr<FEntityGeom> ApplyMatrix(const FMatrixH& InMatrix) const override;
 
-		virtual void EvaluatePoint(double Coordinate, FCurvePoint& OutPoint, int32 DerivativeOrder = 0) const override
-		{
-			ensure(Dimension == 3);
-			Evaluate<FCurvePoint, FPoint>(Coordinate, OutPoint, DerivativeOrder);
-		}
+	virtual void EvaluatePoint(double Coordinate, FCurvePoint& OutPoint, int32 DerivativeOrder = 0) const override
+	{
+		ensure(Dimension == 3);
+		Evaluate<FCurvePoint, FPoint>(Coordinate, OutPoint, DerivativeOrder);
+	}
 
-		virtual void Evaluate2DPoint(double Coordinate, FCurvePoint2D& OutPoint, int32 DerivativeOrder = 0) const override
-		{
-			ensure(Dimension == 2);
-			Evaluate<FCurvePoint2D, FPoint2D>(Coordinate, OutPoint, DerivativeOrder);
-		}
+	virtual void Evaluate2DPoint(double Coordinate, FCurvePoint2D& OutPoint, int32 DerivativeOrder = 0) const override
+	{
+		ensure(Dimension == 2);
+		Evaluate<FCurvePoint2D, FPoint2D>(Coordinate, OutPoint, DerivativeOrder);
+	}
 
-		const FPoint& GetStartPoint() const
-		{
-			return StartPoint;
-		}
+	const FPoint& GetStartPoint() const
+	{
+		return StartPoint;
+	}
 
-		const FPoint& GetEndPoint() const
+	const FPoint& GetEndPoint() const
+	{
+		return EndPoint;
+	}
+
+	virtual void ExtendTo(const FPoint& DesiredPosition) override
+	{
+		double DistanceToStartPoint = DesiredPosition.SquareDistance(StartPoint);
+		double DistanceToEndPoint = DesiredPosition.SquareDistance(EndPoint);
+		if (DistanceToEndPoint < DistanceToStartPoint)
 		{
-			return EndPoint;
+			EndPoint = DesiredPosition;
 		}
-
-		virtual void ExtendTo(const FPoint& DesiredPosition) override
+		else
 		{
-			double DistanceToStartPoint = DesiredPosition.SquareDistance(StartPoint);
-			double DistanceToEndPoint = DesiredPosition.SquareDistance(EndPoint);
-			if (DistanceToEndPoint < DistanceToStartPoint)
-			{
-				EndPoint = DesiredPosition;
-			}
-			else
-			{
-				StartPoint = DesiredPosition;
-			}
+			StartPoint = DesiredPosition;
 		}
+	}
 
-	private:
-		template <typename CurvePointType, typename PointType>
-		void Evaluate(double Coordinate, CurvePointType& OutPoint, int32 DerivativeOrder) const
-		{
-			OutPoint.DerivativeOrder = DerivativeOrder;
+private:
+	template <typename CurvePointType, typename PointType>
+	void Evaluate(double Coordinate, CurvePointType& OutPoint, int32 DerivativeOrder) const
+	{
+		OutPoint.DerivativeOrder = DerivativeOrder;
 
-			PointType Tangent = (EndPoint - StartPoint);
-			OutPoint.Point = Tangent * Coordinate + StartPoint;
+		PointType Tangent = (EndPoint - StartPoint);
+		OutPoint.Point = Tangent * Coordinate + StartPoint;
 
-			if (DerivativeOrder > 0)
-			{
-				OutPoint.Gradient = Tangent;
-			}
+		if (DerivativeOrder > 0)
+		{
+			OutPoint.Gradient = Tangent;
 		}
+	}
 
-	};
+};
 
-} // namespace CADKernel
+} // namespace UE::CADKernel
 
diff --git a/Engine/Source/Runtime/Datasmith/CADKernel/Private/CADKernel/Geo/Sampler/SamplerAbstract.h b/Engine/Source/Runtime/Datasmith/CADKernel/Private/CADKernel/Geo/Sampler/SamplerAbstract.h
--- a/Engine/Source/Runtime/Datasmith/CADKernel/Private/CADKernel/Geo/Sampler/SamplerAbstract.h
+++ b/Engine/Source/Runtime/Datasmith/CADKernel/Private/CADKernel/Geo/Sampler/SamplerAbstract.h
@@ -11,7 +11,7 @@
 //#define DEBUG_CURVE_SAMPLING 
 //#define CHECK_RESULT_MAKE_ISO_POLYLINE
 
-namespace CADKernel
+namespace UE::CADKernel
 {
 
 template<typename PolylineType, typename PointType>
@@ -419,7 +419,7 @@ public:
 				F3DDebugSession G(FString::Printf(TEXT("Curve %d"), Step));
 				for (const FPoint& Point : Sampling.GetPoints())
 				{
-					CADKernel::DisplayPoint(Point, EVisuProperty::BluePoint, Index++);
+					UE::CADKernel::DisplayPoint(Point, EVisuProperty::BluePoint, Index++);
 				}
 
 				for (int32 Index = 0; Index < Sampling.GetPoints().Num() - 1; ++Index)
@@ -439,7 +439,7 @@ public:
 				F3DDebugSession G(FString::Printf(TEXT("Next Point Tmp %d"), Step));
 				for (const FPoint& Point : CandidatePoints.GetPoints())
 				{
-					CADKernel::DisplayPoint(Point, EVisuProperty::YellowPoint, Index++);
+					UE::CADKernel::DisplayPoint(Point, EVisuProperty::YellowPoint, Index++);
 				}
 			}
 			Wait();
@@ -500,4 +500,4 @@ public:
 #endif
 };
 
-} // ns CADKernel
\ No newline at end of file
+} // ns UE::CADKernel
\ No newline at end of file
diff --git a/Engine/Source/Runtime/Datasmith/CADKernel/Private/CADKernel/Geo/Sampler/SamplerOnChord.h b/Engine/Source/Runtime/Datasmith/CADKernel/Private/CADKernel/Geo/Sampler/SamplerOnChord.h
--- a/Engine/Source/Runtime/Datasmith/CADKernel/Private/CADKernel/Geo/Sampler/SamplerOnChord.h
+++ b/Engine/Source/Runtime/Datasmith/CADKernel/Private/CADKernel/Geo/Sampler/SamplerOnChord.h
@@ -14,7 +14,7 @@
 
 //#define DEBUG_CHECKCHORDERROR
 
-namespace CADKernel
+namespace UE::CADKernel
 {
 
 template<class PointType>
@@ -77,13 +77,13 @@ protected:
 				int32 Index = 0;
 				F3DDebugSession G(FString::Printf(TEXT("CheckChord %f %d"), MaxChord, FirstIndex));
 
-				CADKernel::Display(APoint, EVisuProperty::BluePoint);
-				CADKernel::Display(BPoint, EVisuProperty::BluePoint);
+				UE::CADKernel::Display(APoint, EVisuProperty::BluePoint);
+				UE::CADKernel::Display(BPoint, EVisuProperty::BluePoint);
 
 				for (int32 Index = FirstIndex; Index < EndIndex; ++Index)
 				{
 					const PointType& Middle = CandidatePoints.GetPoints()[Index];
-					CADKernel::Display(Middle, EVisuProperty::YellowPoint);
+					UE::CADKernel::Display(Middle, EVisuProperty::YellowPoint);
 				}
 
 				DisplaySegment(APoint, BPoint, 0, EVisuProperty::BlueCurve);
@@ -252,4 +252,4 @@ protected:
 
 };
 
-} // ns CADKernel
+} // ns UE::CADKernel
diff --git a/Engine/Source/Runtime/Datasmith/CADKernel/Private/CADKernel/Geo/Sampler/SamplerOnParam.h b/Engine/Source/Runtime/Datasmith/CADKernel/Private/CADKernel/Geo/Sampler/SamplerOnParam.h
--- a/Engine/Source/Runtime/Datasmith/CADKernel/Private/CADKernel/Geo/Sampler/SamplerOnParam.h
+++ b/Engine/Source/Runtime/Datasmith/CADKernel/Private/CADKernel/Geo/Sampler/SamplerOnParam.h
@@ -14,7 +14,7 @@
 #include "CADKernel/Math/Point.h"
 #include "CADKernel/Utils/ArrayUtils.h"
 
-namespace CADKernel
+namespace UE::CADKernel
 {
 
 template<typename PolylineType, typename PointType>
@@ -100,13 +100,13 @@ protected:
 				int32 Index = 0;
 				F3DDebugSession G(FString::Printf(TEXT("CheckChord %f %d"), MaxChord, FirstIndex));
 
-				CADKernel::Display(APoint, EVisuProperty::BluePoint);
-				CADKernel::Display(BPoint, EVisuProperty::BluePoint);
+				UE::CADKernel::Display(APoint, EVisuProperty::BluePoint);
+				UE::CADKernel::Display(BPoint, EVisuProperty::BluePoint);
 
 				for (int32 Index = FirstIndex; Index < EndIndex; ++Index)
 				{
 					const FPoint& Middle = CandidatePoints.GetPoints()[Index];
-					CADKernel::Display(Middle, EVisuProperty::YellowPoint);
+					UE::CADKernel::Display(Middle, EVisuProperty::YellowPoint);
 				}
 
 				DisplaySegment(APoint, BPoint, 0, EVisuProperty::BlueCurve);
@@ -420,4 +420,4 @@ protected:
 
 };
 
-} // ns CADKernel
\ No newline at end of file
+} // ns UE::CADKernel
\ No newline at end of file
diff --git a/Engine/Source/Runtime/Datasmith/CADKernel/Private/CADKernel/Geo/Sampling/SurfacicPolyline.cpp b/Engine/Source/Runtime/Datasmith/CADKernel/Private/CADKernel/Geo/Sampling/SurfacicPolyline.cpp
--- a/Engine/Source/Runtime/Datasmith/CADKernel/Private/CADKernel/Geo/Sampling/SurfacicPolyline.cpp
+++ b/Engine/Source/Runtime/Datasmith/CADKernel/Private/CADKernel/Geo/Sampling/SurfacicPolyline.cpp
@@ -10,7 +10,7 @@
 #include "CADKernel/Geo/Sampling/PolylineTools.h"
 #include "CADKernel/Math/SlopeUtils.h"
 
-namespace CADKernel
+namespace UE::CADKernel
 {
 
 FSurfacicPolyline::FSurfacicPolyline(TSharedRef<FSurface> InCarrierSurface, TSharedRef<FCurve> Curve2D, const double InTolerance)
@@ -211,4 +211,4 @@ void FSurfacicPolyline::ComputeIntersectionsWithIsos(const FLinearBoundary& InBo
 	}
 }
 
-} //namespace CADKernel
+} //namespace UE::CADKernel
diff --git a/Engine/Source/Runtime/Datasmith/CADKernel/Private/CADKernel/Geo/Sampling/SurfacicSampling.h b/Engine/Source/Runtime/Datasmith/CADKernel/Private/CADKernel/Geo/Sampling/SurfacicSampling.h
--- a/Engine/Source/Runtime/Datasmith/CADKernel/Private/CADKernel/Geo/Sampling/SurfacicSampling.h
+++ b/Engine/Source/Runtime/Datasmith/CADKernel/Private/CADKernel/Geo/Sampling/SurfacicSampling.h
@@ -5,7 +5,7 @@
 #include "CADKernel/Geo/GeoEnum.h"
 #include "CADKernel/Math/Point.h"
 
-namespace CADKernel
+namespace UE::CADKernel
 {
 
 struct FSurfacicSampling
@@ -61,4 +61,4 @@ struct FSurfacicSampling
 	}
 };
 
-} // ns CADKernel
\ No newline at end of file
+} // ns UE::CADKernel
\ No newline at end of file
diff --git a/Engine/Source/Runtime/Datasmith/CADKernel/Private/CADKernel/Geo/Surfaces/BezierSurface.h b/Engine/Source/Runtime/Datasmith/CADKernel/Private/CADKernel/Geo/Surfaces/BezierSurface.h
--- a/Engine/Source/Runtime/Datasmith/CADKernel/Private/CADKernel/Geo/Surfaces/BezierSurface.h
+++ b/Engine/Source/Runtime/Datasmith/CADKernel/Private/CADKernel/Geo/Surfaces/BezierSurface.h
@@ -4,79 +4,79 @@
 #include "CADKernel/Core/CADKernelArchive.h"
 #include "CADKernel/Geo/Surfaces/Surface.h"
 
-namespace CADKernel
+namespace UE::CADKernel
 {
 
-	class CADKERNEL_API FBezierSurface : public FSurface
-	{
-		friend FEntity;
-
-	protected:
+class CADKERNEL_API FBezierSurface : public FSurface
+{
+	friend FEntity;
 
-		int32 UDegre;
-		int32 VDegre;
+protected:
 
-		int32 UPoleNum;
-		int32 VPoleNum;
+	int32 UDegre;
+	int32 VDegre;
 
-		TArray<FPoint> Poles;
+	int32 UPoleNum;
+	int32 VPoleNum;
 
-		FBezierSurface(const double InToleranceGeometric, int32 InUDegre, int32 InVDegre, const TArray<FPoint>& InPoles)
-			: FSurface(InToleranceGeometric)
-			, UDegre(InUDegre)
-			, VDegre(InVDegre)
-			, UPoleNum(InUDegre + 1)
-			, VPoleNum(InVDegre + 1)
-			, Poles(InPoles)
-		{
-			ensureCADKernel((UDegre + 1) * (VDegre + 1) == Poles.Num());
-			ComputeDefaultMinToleranceIso();
-		}
+	TArray<FPoint> Poles;
 
-		FBezierSurface() = default;
+	FBezierSurface(const double InToleranceGeometric, int32 InUDegre, int32 InVDegre, const TArray<FPoint>& InPoles)
+		: FSurface(InToleranceGeometric)
+		, UDegre(InUDegre)
+		, VDegre(InVDegre)
+		, UPoleNum(InUDegre + 1)
+		, VPoleNum(InVDegre + 1)
+		, Poles(InPoles)
+	{
+		ensureCADKernel((UDegre + 1) * (VDegre + 1) == Poles.Num());
+		ComputeDefaultMinToleranceIso();
+	}
 
-	public:
+	FBezierSurface() = default;
 
-		virtual void Serialize(FCADKernelArchive& Ar) override
-		{
-			FSurface::Serialize(Ar);
-			Ar << UDegre;
-			Ar << VDegre;
-			Ar << UPoleNum;
-			Ar << VPoleNum;
-			Ar.Serialize(Poles);
-		}
+public:
 
-		ESurface GetSurfaceType() const
-		{
-			return ESurface::Bezier;
-		}
+	virtual void Serialize(FCADKernelArchive& Ar) override
+	{
+		FSurface::Serialize(Ar);
+		Ar << UDegre;
+		Ar << VDegre;
+		Ar << UPoleNum;
+		Ar << VPoleNum;
+		Ar.Serialize(Poles);
+	}
+
+	ESurface GetSurfaceType() const
+	{
+		return ESurface::Bezier;
+	}
 
-		int32 GetUDegre() const
-		{
-			return UDegre;
-		}
+	int32 GetUDegre() const
+	{
+		return UDegre;
+	}
 
-		int32 GetVDegre() const
-		{
-			return VDegre;
-		}
+	int32 GetVDegre() const
+	{
+		return VDegre;
+	}
 
-		virtual const TArray<FPoint>& GetPoles() const
-		{
-			return Poles;
-		}
+	virtual const TArray<FPoint>& GetPoles() const
+	{
+		return Poles;
+	}
 
-		virtual TSharedPtr<FEntityGeom> ApplyMatrix(const FMatrixH& InMatrix) const override;
+	virtual TSharedPtr<FEntityGeom> ApplyMatrix(const FMatrixH& InMatrix) const override;
 
-		virtual void EvaluatePoint(const FPoint2D& InSurfacicCoordinate, FSurfacicPoint& OutPoint3D, int32 InDerivativeOrder = 0) const override;
+	virtual void EvaluatePoint(const FPoint2D& InSurfacicCoordinate, FSurfacicPoint& OutPoint3D, int32 InDerivativeOrder = 0) const override;
 
-		virtual void Presample(const FSurfacicBoundary& InBoundaries, FCoordinateGrid& OutCoordinates) override;
+	virtual void Presample(const FSurfacicBoundary& InBoundaries, FCoordinateGrid& OutCoordinates) override;
 
 #ifdef CADKERNEL_DEV
-		virtual FInfoEntity& GetInfo(FInfoEntity&) const override;
+	virtual FInfoEntity& GetInfo(FInfoEntity&) const override;
 #endif
-	};
+};
 
-} // namespace CADKernel
+} // namespace UE::CADKernel
 
