[NO NON-PII COMMIT MESSAGE]

diff --git a/Engine/Config/ConsoleVariables.ini b/Engine/Config/ConsoleVariables.ini
--- a/Engine/Config/ConsoleVariables.ini
+++ b/Engine/Config/ConsoleVariables.ini
@@ -126,7 +126,10 @@ fx.SkipVectorVMBackendOptimizations=1
 ;ds.CADTranslator.EnableTimeControl=1
 
 ; Uncomment to enable/disable CADKernel tessellation. If disable, CAD SDK tessellator is used.
-;ds.CADTranslator.DisableCADKernelTessellation=false
+ds.CADTranslator.DisableCADKernelTessellation=false
+
+; Perform a welding of the mesh to try to stitch mesh cracks.
+;ds.CADTranslator.SewMeshIfNeeded=false
 
 ; Uncomment to force to import embedded tessellation instead of meshing BRep model. Default is false
 ;ds.CADTranslator.PreferJtFileEmbeddedTessellation=true
diff --git a/Engine/Plugins/Enterprise/DatasmithCADImporter/Source/CADInterfaces/Private/TechSoftFileParserCADKernelTessellator.cpp b/Engine/Plugins/Enterprise/DatasmithCADImporter/Source/CADInterfaces/Private/TechSoftFileParserCADKernelTessellator.cpp
--- a/Engine/Plugins/Enterprise/DatasmithCADImporter/Source/CADInterfaces/Private/TechSoftFileParserCADKernelTessellator.cpp
+++ b/Engine/Plugins/Enterprise/DatasmithCADImporter/Source/CADInterfaces/Private/TechSoftFileParserCADKernelTessellator.cpp
@@ -197,6 +197,10 @@ void FTechSoftFileParserCADKernelTessellator::GenerateBodyMesh(A3DRiRepresentati
 	FTechSoftBridge TechSoftBridge(*this, CADKernelSession);
 
 	FBody* CADKernelBody = TechSoftBridge.AddBody(Representation, ArchiveBody.MetaData, ArchiveBody.Unit);
+	if (CADKernelBody == nullptr)
+	{
+		return;
+	}
 
 	if (CADFileData.GetImportParameters().GetStitchingTechnique() == StitchingHeal)
 	{
diff --git a/Engine/Plugins/Enterprise/DatasmithCADImporter/Source/CADTools/Private/CADOption.cpp b/Engine/Plugins/Enterprise/DatasmithCADImporter/Source/CADTools/Private/CADOption.cpp
--- a/Engine/Plugins/Enterprise/DatasmithCADImporter/Source/CADTools/Private/CADOption.cpp
+++ b/Engine/Plugins/Enterprise/DatasmithCADImporter/Source/CADTools/Private/CADOption.cpp
@@ -74,7 +74,7 @@ bool FImportParameters::bGSewMeshIfNeeded = true;
 FAutoConsoleVariableRef GCADTranslatorSewMeshIfNeeded(
 	TEXT("ds.CADTranslator.SewMeshIfNeeded"),
 	FImportParameters::bGSewMeshIfNeeded,
-	TEXT("Prefrom a welding of the mesh to try to stitch mesh cracks\n\
+	TEXT("Perform a welding of the mesh to try to stitch mesh cracks\n\
 This welding is performed respecting the ds.CADTranslator.StitchingTolerance\n\
 Default value is true\n"),
 ECVF_Default);
