[NO NON-PII COMMIT MESSAGE]

diff --git a/Engine/Plugins/Enterprise/DatasmithCADImporter/Source/CADInterfaces/Private/TechSoftFileParser.cpp b/Engine/Plugins/Enterprise/DatasmithCADImporter/Source/CADInterfaces/Private/TechSoftFileParser.cpp
--- a/Engine/Plugins/Enterprise/DatasmithCADImporter/Source/CADInterfaces/Private/TechSoftFileParser.cpp
+++ b/Engine/Plugins/Enterprise/DatasmithCADImporter/Source/CADInterfaces/Private/TechSoftFileParser.cpp
@@ -1245,6 +1245,8 @@ void FTechSoftFileParser::TraverseBRepModel(A3DRiBrepModel* BRepModelPtr, FArchi
 void FTechSoftFileParser::TraversePolyBRepModel(A3DRiPolyBrepModel* PolygonalPtr, FArchiveReference& Parent)
 {
 	FArchiveBody& BRep = SceneGraph.AddBody(Parent);
+	BRep.bIsFromCad = false;
+
 	ExtractMetaData(PolygonalPtr, BRep);
 	BRep.DefineGraphicsPropertiesFromNoOverwrite(Parent);
 
