Fix null normal from TechSoft (in the wait of TechSoft fix)
Do not import degenerated triangles

diff --git a/Engine/Plugins/Enterprise/DatasmithCADImporter/Source/CADLibrary/Private/MeshDescriptionHelper.cpp b/Engine/Plugins/Enterprise/DatasmithCADImporter/Source/CADLibrary/Private/MeshDescriptionHelper.cpp
--- a/Engine/Plugins/Enterprise/DatasmithCADImporter/Source/CADLibrary/Private/MeshDescriptionHelper.cpp
+++ b/Engine/Plugins/Enterprise/DatasmithCADImporter/Source/CADLibrary/Private/MeshDescriptionHelper.cpp
@@ -366,7 +366,7 @@ bool FillMesh(const FMeshParameters& MeshParameters, const FImportParameters& Im
 					MeshVertexInstanceIDs.Add(VertexInstanceID);
 					VertexInstanceNormals[VertexInstanceID] = (FVector3f)Tessellation.NormalArray[NormalIndex];
 					if (!Tessellation.TexCoordArray.IsEmpty())
-					{
+			{
 						VertexInstanceUVs.Set(VertexInstanceID, UVChannel, FVector2f(Tessellation.TexCoordArray[UVIndex]));
 					}
 				}
@@ -375,7 +375,7 @@ bool FillMesh(const FMeshParameters& MeshParameters, const FImportParameters& Im
 
 			TMap<FVertexID, FVertexInstanceID> VertexIDToInstanceIDForCad;
 			TFunction<FVertexInstanceID(FVertexID, int32)> FindOrAddVertexInstanceIDForCad = [&](FVertexID VertexID, int32 VertexIndex) ->FVertexInstanceID
-			{
+				{
 				FVertexInstanceID& VertexInstanceID = VertexIDToInstanceIDForCad.FindOrAdd(VertexID);
 				if (VertexInstanceID == -1)
 				{
@@ -387,7 +387,7 @@ bool FillMesh(const FMeshParameters& MeshParameters, const FImportParameters& Im
 					MeshVertexInstanceIDs.Add(VertexInstanceID);
 					VertexInstanceNormals[VertexInstanceID] = (FVector3f)Tessellation.NormalArray[VertexIndex];
 					if (!Tessellation.TexCoordArray.IsEmpty())
-					{
+				{
 						VertexInstanceUVs.Set(VertexInstanceID, UVChannel, FVector2f(Tessellation.TexCoordArray[VertexIndex]));
 					}
 				}
@@ -406,7 +406,7 @@ bool FillMesh(const FMeshParameters& MeshParameters, const FImportParameters& Im
 				VertexIDToInstanceIDForCad.Reserve(Tessellation.VertexIndices.Num());
 
 				if (Tessellation.TexCoordArray.IsEmpty())
-				{
+			{
 					UVIndices.Init(0, Tessellation.VertexIndices.Num());
 				}
 				else
@@ -415,20 +415,20 @@ bool FillMesh(const FMeshParameters& MeshParameters, const FImportParameters& Im
 					for (int32 Index = 0; Index < Tessellation.VertexIndices.Num(); ++Index)
 					{
 						UVIndices[Index] = Index;
-					}
+				}
 					const FVector2f OneVector(1.f, 1.f);
 					MergeCoincidents(Tessellation.TexCoordArray, OneVector, KINDA_SMALL_NUMBER, UVIndices);
-				}
+			}
 
 				if (Tessellation.NormalArray.Num() == 1)
-				{
+			{
 					NormalIndices.Init(0, Tessellation.VertexIndices.Num());
-				}
+			}
 				else
-				{
+			{
 					NormalIndices.SetNum(Tessellation.VertexIndices.Num());
 					for (int32 Index = 0; Index < Tessellation.VertexIndices.Num(); ++Index)
-					{
+				{
 						NormalIndices[Index] = Index;
 					}
 					MergeCoincidents(Tessellation.NormalArray, FVector3f::OneVector, KINDA_SMALL_NUMBER, NormalIndices);
@@ -469,19 +469,19 @@ bool FillMesh(const FMeshParameters& MeshParameters, const FImportParameters& Im
 				if (FaceVertexIDs[0] == FaceVertexIDs[1] || FaceVertexIDs[0] == FaceVertexIDs[2] || FaceVertexIDs[1] == FaceVertexIDs[2])
 				{
 					continue;
-				}
+			}
 
 				TriangleVertexInstanceIDs[0] = FindOrAddVertexInstanceID(FaceVertexIDs[0], FaceVertexIndices[0]);
 				TriangleVertexInstanceIDs[1] = FindOrAddVertexInstanceID(FaceVertexIDs[1], FaceVertexIndices[1]);
 				TriangleVertexInstanceIDs[2] = FindOrAddVertexInstanceID(FaceVertexIDs[2], FaceVertexIndices[2]);
 
 				if(FImportParameters::bGRemoveDuplicatedTriangle)
-				{
+			{
 					Sort(FaceVertexIDs, 3);
 					bool bIsAlreadyInSet;
 					VertexIdsToTriangle.Emplace(TTuple<FVertexInstanceID, FVertexInstanceID, FVertexInstanceID>(FaceVertexIDs[0], FaceVertexIDs[1], FaceVertexIDs[2]), &bIsAlreadyInSet);
 					if (bIsAlreadyInSet)
-					{
+				{
 						continue;
 					}
 				}
@@ -531,6 +531,9 @@ bool ConvertBodyMeshToMeshDescription(const FImportParameters& ImportParams, con
 		return false;
 	}
 
+	// Workaround SDHE-19725: Compute any null normals.
+	MeshOperator::RecomputeNullNormal(MeshDescription);
+
 	// Orient mesh
 	MeshOperator::OrientMesh(MeshDescription);
 
diff --git a/Engine/Plugins/Enterprise/DatasmithCADImporter/Source/CADLibrary/Private/MeshOperator.cpp b/Engine/Plugins/Enterprise/DatasmithCADImporter/Source/CADLibrary/Private/MeshOperator.cpp
--- a/Engine/Plugins/Enterprise/DatasmithCADImporter/Source/CADLibrary/Private/MeshOperator.cpp
+++ b/Engine/Plugins/Enterprise/DatasmithCADImporter/Source/CADLibrary/Private/MeshOperator.cpp
@@ -16,6 +16,86 @@
 #include "Operations/MeshResolveTJunctions.h"
 #include "StaticMeshAttributes.h"
 
+void MeshOperator::RecomputeNullNormal(FMeshDescription& MeshDescription)
+{
+	//using namespace UE::MeshCategory;
+	constexpr float SquareNormalThreshold = KINDA_SMALL_NUMBER * KINDA_SMALL_NUMBER;
+
+	FStaticMeshAttributes StaticMeshAttributes = FStaticMeshAttributes(MeshDescription);
+	TVertexInstanceAttributesRef<FVector3f> Normals = StaticMeshAttributes.GetVertexInstanceNormals();
+	TVertexInstanceAttributesRef<FVertexID> Indices = StaticMeshAttributes.GetVertexInstanceVertexIndices();
+	TVertexAttributesRef<FVector3f> VertexPositions = StaticMeshAttributes.GetVertexPositions();
+
+	for (FTriangleID Triangle : MeshDescription.Triangles().GetElementIDs())
+	{
+		TArrayView<const FVertexInstanceID> Vertices = MeshDescription.GetTriangleVertexInstances(Triangle);
+		for (int32 VertexIndex = 0; VertexIndex < 3; VertexIndex++)
+		{
+			FVector3f& Normal = Normals[Vertices[VertexIndex]];
+
+			if (Normal.IsNearlyZero(KINDA_SMALL_NUMBER))
+			{
+				Normal = FVector3f::ZeroVector;
+
+				const FVertexInstanceID VertexInstanceID = Vertices[VertexIndex];
+				FVertexID VertexID = MeshDescription.GetVertexInstanceVertex(VertexInstanceID);
+				const TArrayView<const FTriangleID> VertexConnectedTriangles = MeshDescription.GetVertexInstanceConnectedTriangleIDs(VertexInstanceID);
+
+				int32 TriangleCount = 0;
+
+				const int32 TriangleIndex[3][3] = { {0,1,2},{1,2,0},{2,0,1} };
+
+				// compute the weighted sum of normals of the "partition star" according to the corner angle
+				for(const FTriangleID& TriangleID : VertexConnectedTriangles)
+				{
+					// compute face normal at the vertex
+					TArrayView<const FVertexID> TriangleVertices = MeshDescription.GetTriangleVertices(TriangleID);
+
+					int32 ApexIndex = 0;
+					for (; ApexIndex < 3; ++ApexIndex)
+					{
+						if (TriangleVertices[ApexIndex] == VertexID)
+						{
+							break;
+						}
+					}
+
+					const FVector3f Position0 = VertexPositions[TriangleVertices[TriangleIndex[ApexIndex][0]]];
+					FVector3f DPosition1 = VertexPositions[TriangleVertices[TriangleIndex[ApexIndex][1]]] - Position0;
+					FVector3f DPosition2 = VertexPositions[TriangleVertices[TriangleIndex[ApexIndex][2]]] - Position0;
+
+					// to avoid numerical issue due to small vector
+					DPosition1.Normalize();
+					DPosition2.Normalize();
+
+					// We have a left-handed coordinate system, but a counter-clockwise winding order
+					// Hence normal calculation has to take the triangle vectors cross product in reverse.
+					FVector3f TriangleNormal = FVector3f::CrossProduct(DPosition2, DPosition1);
+					double SinOfApexAngle = TriangleNormal.Length();
+					double ApexAngle = FMath::Asin(SinOfApexAngle);
+
+					if (TriangleNormal.Normalize(SquareNormalThreshold))
+					{
+						Normal += (TriangleNormal * ApexAngle);
+						TriangleCount++;
+					}
+				}
+
+				if (TriangleCount)
+				{
+					if (Normal.Normalize(SquareNormalThreshold))
+					{
+						continue;
+					}
+				}
+
+				// the vertex is a vertex of degenerated triangles, the vertex normal doesn't matter as the triangle is flat, but a non null normal is needed.
+				Normal = FVector3f::UpVector;
+			}
+		}
+	}
+}
+
 bool MeshOperator::OrientMesh(FMeshDescription& MeshDescription)
 {
 	using namespace MeshCategory;
@@ -252,7 +332,7 @@ bool MeshOperator::OrientMesh(FMeshDescription& MeshDescription)
 
 void MeshOperator::ResolveTJunctions(FMeshDescription& MeshDescription, double Tolerance)
 {
-	UE::Geometry::FDynamicMesh3 DynamicMesh(UE::Geometry::EMeshComponents::FaceGroups);
+	UE::Geometry::FDynamicMesh3 DynamicMesh(UE::Geometry::EMeshComponents::All);
 	DynamicMesh.EnableAttributes();
 
 	{
@@ -294,17 +374,16 @@ void MeshOperator::ResolveTJunctions(FMeshDescription& MeshDescription, double T
 		ConversionOptions.bSetPolyGroups = true;
 		ConversionOptions.bUpdatePositions = true;
 		ConversionOptions.bUpdateNormals = true;
-		ConversionOptions.bUpdateTangents = true;
+		ConversionOptions.bUpdateTangents = false;
 		ConversionOptions.bUpdateUVs = true;
-		ConversionOptions.bUpdateVtxColors = true;
+		ConversionOptions.bUpdateVtxColors = false;
 		ConversionOptions.bTransformVtxColorsSRGBToLinear = false;
 
 		FDynamicMeshToMeshDescription ConverterToMeshDescription(ConversionOptions);
 		ConverterToMeshDescription.Convert(&DynamicMesh, NewMeshDescription);
 	}
 
-
-	// TODO: Should not be needed as soon as ConverterToMeshDescription.Convert won't lose data anymore
+	// Transfer material slot names and patch groups as these datum are not managed by FDynamicMesh3
 	CADLibrary::CopyMaterialSlotNames(MeshDescription, NewMeshDescription);
 	CADLibrary::CopyPatchGroups(MeshDescription, NewMeshDescription);
 
diff --git a/Engine/Plugins/Enterprise/DatasmithCADImporter/Source/CADLibrary/Private/MeshOperator.h b/Engine/Plugins/Enterprise/DatasmithCADImporter/Source/CADLibrary/Private/MeshOperator.h
--- a/Engine/Plugins/Enterprise/DatasmithCADImporter/Source/CADLibrary/Private/MeshOperator.h
+++ b/Engine/Plugins/Enterprise/DatasmithCADImporter/Source/CADLibrary/Private/MeshOperator.h
@@ -11,6 +11,8 @@ namespace MeshOperator
 
 bool OrientMesh(FMeshDescription& MeshDescription);
 
+void RecomputeNullNormal(FMeshDescription& MeshDescription);
+
 void ResolveTJunctions(FMeshDescription& MeshDescription, double Tolerance);
 
 }
\ No newline at end of file
