Fix that import of the CATProduct hangs. The main problem is due to the size of the message between the CADWorker and the dispatcher in this case (the CATProduct has more than 4000 external references, that generated a message bigger than the limit.)
=>
fix FNetworkNode::ReceiveMessage(): remove the limite of 1 Mb for message size, manage the reception of large messages
fix FDatasmithWorkerHandler::RunInternal(): manage the case of EWorkerErrorState::ConnectionLost

+ some improvements:
- Launch only one worker for monolithic file
- Replace FThread by UE::Tasks::FTask

diff --git a/Engine/Plugins/Enterprise/DatasmithCADImporter/Source/CADTools/Private/CADData.cpp b/Engine/Plugins/Enterprise/DatasmithCADImporter/Source/CADTools/Private/CADData.cpp
--- a/Engine/Plugins/Enterprise/DatasmithCADImporter/Source/CADTools/Private/CADData.cpp
+++ b/Engine/Plugins/Enterprise/DatasmithCADImporter/Source/CADTools/Private/CADData.cpp
@@ -202,86 +202,136 @@ uint32 GetTypeHash(const FFileDescriptor& FileDescriptor)
 	return DescriptorHash;
 }
 
-ECADFormat FileFormat(const FString& Extension)
+void FFileDescriptor::SetFileFormat(const FString& Extension)
 {
-	if (Extension == TEXT("catpart") || Extension == TEXT("catproduct"))
+	if (Extension == TEXT("catpart") )
 	{
-		return ECADFormat::CATIA;
+		Format = ECADFormat::CATIA;
+		bCanReferenceOtherFiles = false;
+	}
+	else if (Extension == TEXT("catproduct"))
+	{
+		Format = ECADFormat::CATIA;
+		bCanReferenceOtherFiles = true;
 	}
 	else if (Extension == TEXT("cgr"))
 	{
-		return ECADFormat::CATIA_CGR;
+		Format = ECADFormat::CATIA_CGR;
+		bCanReferenceOtherFiles = false;
 	}
 	else if (Extension == TEXT("iges") || Extension == TEXT("igs"))
 	{
-		return ECADFormat::IGES;
+		Format = ECADFormat::IGES;
+		bCanReferenceOtherFiles = false;
 	}
 	else if (Extension == TEXT("step") || Extension == TEXT("stp"))
 	{
-		return ECADFormat::STEP;
+		Format = ECADFormat::STEP;
+		bCanReferenceOtherFiles = true;
+	}
+	else if (Extension == TEXT("ipt"))
+	{
+		Format = ECADFormat::INVENTOR;
+		bCanReferenceOtherFiles = false;
 	}
-	else if (Extension == TEXT("ipt") || Extension == TEXT("iam"))
+	else if (Extension == TEXT("iam"))
 	{
-		return ECADFormat::INVENTOR;
+		Format = ECADFormat::INVENTOR;
+		bCanReferenceOtherFiles = true;
 	}
 	else if (Extension == TEXT("jt"))
 	{
-		return ECADFormat::JT;
+		Format = ECADFormat::JT;
+		bCanReferenceOtherFiles = true;
+	}
+	else if (Extension == TEXT("model"))
+	{
+		Format = ECADFormat::CATIAV4;
+		bCanReferenceOtherFiles = false;
 	}
-	else if (Extension == TEXT("model") || Extension == TEXT("exp") || Extension == TEXT("session") || Extension == TEXT("dlv"))
+	else if (Extension == TEXT("exp") || Extension == TEXT("session") || Extension == TEXT("dlv"))
 	{
-		return ECADFormat::CATIAV4;
+		Format = ECADFormat::CATIAV4;
+		bCanReferenceOtherFiles = true;
 	}
-	else if (Extension == TEXT("prt.*") || Extension == TEXT("asm.*") 
+	else if (Extension == TEXT("prt.*")
+		|| Extension == TEXT("xpr"))
+	{
+		Format = ECADFormat::CREO;
+		bCanReferenceOtherFiles = false;
+	}
+	else if (Extension == TEXT("asm.*")
 		|| Extension == TEXT("creo") || Extension == TEXT("creo.*")
 		|| Extension == TEXT("neu") || Extension == TEXT("neu.*")
-		|| Extension == TEXT("xas") || Extension == TEXT("xpr"))
+		|| Extension == TEXT("xpr"))
 	{
-		return ECADFormat::CREO;
+		Format = ECADFormat::CREO;
+		bCanReferenceOtherFiles = true;
 	}
-	else if (Extension == TEXT("prt") || Extension == TEXT("asm"))
+	else if (Extension == TEXT("prt"))
 	{
-		return ECADFormat::N_X;
+		Format = ECADFormat::N_X;
+		bCanReferenceOtherFiles = true;
 	}
 	else if (Extension == TEXT("sat") || Extension == TEXT("sab"))
 	{
-		return ECADFormat::ACIS;
+		Format = ECADFormat::ACIS;
+		bCanReferenceOtherFiles = true;
+	}
+	else if (Extension == TEXT("sldprt") )
+	{
+		Format = ECADFormat::SOLIDWORKS;
+		bCanReferenceOtherFiles = false;
 	}
-	else if (Extension == TEXT("sldprt") || Extension == TEXT("sldasm"))
+	else if (Extension == TEXT("sldasm"))
 	{
-		return ECADFormat::SOLIDWORKS;
+		Format = ECADFormat::SOLIDWORKS;
+		bCanReferenceOtherFiles = true;
 	}
 	else if (Extension == TEXT("x_t") || Extension == TEXT("x_b"))
 	{
-		return ECADFormat::PARASOLID;
+		Format = ECADFormat::PARASOLID;
+		bCanReferenceOtherFiles = true;
 	}
 	else if (Extension == TEXT("3dxml") || Extension == TEXT("3drep"))
 	{
-		return ECADFormat::CATIA_3DXML;
+		Format = ECADFormat::CATIA_3DXML;
+		bCanReferenceOtherFiles = false;
+	}
+	else if (Extension == TEXT("par") )
+	{
+		Format = ECADFormat::SOLID_EDGE;
+		bCanReferenceOtherFiles = false;
 	}
-	else if (Extension == TEXT("par") || Extension == TEXT("psm"))
+	else if (Extension == TEXT("psm"))
 	{
-		return ECADFormat::SOLID_EDGE;
+		Format = ECADFormat::SOLID_EDGE;
+		bCanReferenceOtherFiles = true;
 	}
 	else if (Extension == TEXT("dwg"))
 	{
-		return ECADFormat::AUTOCAD;
+		Format = ECADFormat::AUTOCAD;
+		bCanReferenceOtherFiles = true;
 	}
 	else if (Extension == TEXT("ifc"))
 	{
-		return ECADFormat::IFC;
+		Format = ECADFormat::IFC;
+		bCanReferenceOtherFiles = false;
 	}
 	else if (Extension == TEXT("dgn"))
 	{
-		return ECADFormat::MICROSTATION;
+		Format = ECADFormat::MICROSTATION;
+		bCanReferenceOtherFiles = true;
 	}
 	else if (Extension == TEXT("hsf") || Extension == TEXT("prc"))
 	{
-		return ECADFormat::TECHSOFT;
+		Format = ECADFormat::TECHSOFT;
+		bCanReferenceOtherFiles = false;
 	}
 	else
 	{
-		return ECADFormat::OTHER;
+		Format = ECADFormat::OTHER;
+		bCanReferenceOtherFiles = true;
 	}
 }
 
diff --git a/Engine/Plugins/Enterprise/DatasmithCADImporter/Source/CADTools/Public/CADData.h b/Engine/Plugins/Enterprise/DatasmithCADImporter/Source/CADTools/Public/CADData.h
--- a/Engine/Plugins/Enterprise/DatasmithCADImporter/Source/CADTools/Public/CADData.h
+++ b/Engine/Plugins/Enterprise/DatasmithCADImporter/Source/CADTools/Public/CADData.h
@@ -41,8 +41,6 @@ enum class ECADFormat
 	OTHER
 };
 
-CADTOOLS_API ECADFormat FileFormat(const FString& Extension);
-
 enum class ECADParsingResult : uint8
 {
 	Unknown,
@@ -68,11 +66,10 @@ enum class ECADGraphicPropertyInheritance : uint8
 	ChildHerit,
 };
 
-// TODO: Remove from hear and replace by DatasmithUtils::GetCleanFilenameAndExtension... But need to remove DatasmithCore dependancies 
+// TODO: Remove from here and replace by DatasmithUtils::GetCleanFilenameAndExtension... But need to remove DatasmithCore dependancies 
 CADTOOLS_API void GetCleanFilenameAndExtension(const FString& InFilePath, FString& OutFilename, FString& OutExtension);
 CADTOOLS_API FString GetExtension(const FString& InFilePath);
 
-
 class CADTOOLS_API FCADMaterial
 {
 public:
@@ -170,7 +167,7 @@ public:
 		, Configuration(InConfiguration)
 	{
 		Name = FPaths::GetCleanFilename(InFilePath);
-		Format = FileFormat(GetExtension(InFilePath));
+		SetFileFormat(GetExtension(InFilePath));
 		RootFolder = InRootFolder ? InRootFolder : FPaths::GetPath(InFilePath);
 	}
 
@@ -240,6 +237,11 @@ public:
 		return Format;
 	}
 
+	bool CanReferenceOtherFiles() const
+	{
+		return bCanReferenceOtherFiles;
+	}
+
 	const FString& GetPathOfFileToLoad() const
 	{
 		if (CacheFilePath.IsEmpty())
@@ -274,10 +276,13 @@ private:
 	FString CacheFilePath; // if the file has already been loaded 
 	FString Name; // content.jt
 	ECADFormat Format; // ECADFormat::JT
+	bool bCanReferenceOtherFiles;
 	FString Configuration; // dedicated to JT or SW file to read the good configuration (SW) or only a sub-file (JT)
 	FString RootFolder; // alternative folder where the file could be if its path is not valid.
 
 	mutable uint32 DescriptorHash = 0;
+
+	void SetFileFormat(const FString& Extension);
 };
 
 /**
diff --git a/Engine/Plugins/Enterprise/DatasmithCADImporter/Source/DatasmithCADTranslator/Private/DatasmithCADTranslator.cpp b/Engine/Plugins/Enterprise/DatasmithCADImporter/Source/DatasmithCADTranslator/Private/DatasmithCADTranslator.cpp
--- a/Engine/Plugins/Enterprise/DatasmithCADImporter/Source/DatasmithCADTranslator/Private/DatasmithCADTranslator.cpp
+++ b/Engine/Plugins/Enterprise/DatasmithCADImporter/Source/DatasmithCADTranslator/Private/DatasmithCADTranslator.cpp
@@ -203,7 +203,14 @@ bool FDatasmithCADTranslator::LoadScene(TSharedRef<IDatasmithScene> DatasmithSce
 			int32 NumCores = FPlatformMisc::NumberOfCores();
 			if (CADLibrary::GMaxImportThreads > 1)
 			{
-				NumCores = FMath::Min(CADLibrary::GMaxImportThreads, NumCores);
+				if (FileDescriptor.CanReferenceOtherFiles())
+				{
+					NumCores = FMath::Min(CADLibrary::GMaxImportThreads, NumCores);
+				}
+				else 
+				{
+					NumCores = 1;
+				}
 			}
 			DatasmithDispatcher::FDatasmithDispatcher Dispatcher(ImportParameters, CachePath, NumCores, CADFileToUEFileMap, CADFileToUEGeomMap);
 			Dispatcher.AddTask(FileDescriptor);
diff --git a/Engine/Plugins/Enterprise/DatasmithCADImporter/Source/DatasmithDispatcher/Private/DatasmithDispatcher.cpp b/Engine/Plugins/Enterprise/DatasmithCADImporter/Source/DatasmithDispatcher/Private/DatasmithDispatcher.cpp
--- a/Engine/Plugins/Enterprise/DatasmithCADImporter/Source/DatasmithDispatcher/Private/DatasmithDispatcher.cpp
+++ b/Engine/Plugins/Enterprise/DatasmithCADImporter/Source/DatasmithDispatcher/Private/DatasmithDispatcher.cpp
@@ -14,7 +14,6 @@
 #include "Misc/ScopeLock.h"
 #include "HAL/IConsoleManager.h"
 
-
 namespace DatasmithDispatcher
 {
 
@@ -130,10 +129,6 @@ void FDatasmithDispatcher::Process(bool bWithProcessor)
 		bWithProcessor &= FPaths::FileExists(WorkerPath);
 	}
 
-#ifdef CAD_DISPATCHER_DEBUG
-	bWithProcessor = false;
-#endif //CAD_TRANSLATOR_DEBUG
-
 	if (bWithProcessor)
 	{
 		SpawnHandlers();
@@ -168,8 +163,7 @@ void FDatasmithDispatcher::Process(bool bWithProcessor)
 			{
 				break;
 			}
-
-			FPlatformProcess::Sleep(0.1);
+			FPlatformProcess::Sleep(0.03);
 		}
 
 		CloseHandlers();
diff --git a/Engine/Plugins/Enterprise/DatasmithCADImporter/Source/DatasmithDispatcher/Private/DatasmithDispatcherNetworking.cpp b/Engine/Plugins/Enterprise/DatasmithCADImporter/Source/DatasmithDispatcher/Private/DatasmithDispatcherNetworking.cpp
--- a/Engine/Plugins/Enterprise/DatasmithCADImporter/Source/DatasmithDispatcher/Private/DatasmithDispatcherNetworking.cpp
+++ b/Engine/Plugins/Enterprise/DatasmithCADImporter/Source/DatasmithDispatcher/Private/DatasmithDispatcherNetworking.cpp
@@ -65,7 +65,7 @@ bool FNetworkNode::SendMessage(const TArray<uint8>& Message, double Timeout_s)
 		{
 			int32 BytesSent = 0;
 			bSendSucceed &= ConnectedSocket->Send(Data + TotalByteSent, BufferSize - TotalByteSent, BytesSent);
-			TotalByteSent += (uint32) BytesSent;
+			TotalByteSent += (uint32)BytesSent;
 		}
 		return bSendSucceed && TotalByteSent == BufferSize;
 	};
@@ -112,7 +112,7 @@ bool FNetworkNode::ReceiveMessage(TArray<uint8>& OutMessage, double Timeout_s)
 
 	ConnectedSocket->SetNonBlocking(true);
 
-	bool bCanRead = ConnectedSocket && ConnectedSocket->Wait(ESocketWaitConditions::WaitForRead, TimespanFromSeconds(Timeout_s));
+	bool bCanRead = ConnectedSocket->Wait(ESocketWaitConditions::WaitForRead, TimespanFromSeconds(Timeout_s));
 	if (!bCanRead)
 	{
 		return false;
@@ -150,7 +150,7 @@ bool FNetworkNode::ReceiveMessage(TArray<uint8>& OutMessage, double Timeout_s)
 		FMemoryReader ArReader(HeaderBuffer);
 		ArReader << IncommingMessage.Header;
 
-		if (IncommingMessage.Header.ByteSize < 0 || IncommingMessage.Header.ByteSize > 1<<20)
+		if (IncommingMessage.Header.ByteSize < 0)
 		{
 			UE_LOG(LogDatasmithDispatcher, Error, TEXT("Parsed header failed: bad message size %d"), IncommingMessage.Header.ByteSize);
 			bReadError = true;
@@ -162,7 +162,7 @@ bool FNetworkNode::ReceiveMessage(TArray<uint8>& OutMessage, double Timeout_s)
 
 	// fill the message with available data
 	uint32 PendingByteCount;
-	if (IncommingMessage.Header.ByteSize >= 0 && ConnectedSocket->HasPendingData(PendingByteCount))
+	while (IncommingMessage.Header.ByteSize >= 0 && ConnectedSocket->HasPendingData(PendingByteCount))
 	{
 		int32 MissingByteInCurrentMessage = IncommingMessage.Header.ByteSize - IncommingMessage.Content.Num();
 		int32 ReadTarget = int32(PendingByteCount) < MissingByteInCurrentMessage ? PendingByteCount : MissingByteInCurrentMessage;
diff --git a/Engine/Plugins/Enterprise/DatasmithCADImporter/Source/DatasmithDispatcher/Private/DatasmithWorkerHandler.cpp b/Engine/Plugins/Enterprise/DatasmithCADImporter/Source/DatasmithDispatcher/Private/DatasmithWorkerHandler.cpp
--- a/Engine/Plugins/Enterprise/DatasmithCADImporter/Source/DatasmithDispatcher/Private/DatasmithWorkerHandler.cpp
+++ b/Engine/Plugins/Enterprise/DatasmithCADImporter/Source/DatasmithDispatcher/Private/DatasmithWorkerHandler.cpp
@@ -48,12 +48,12 @@ FDatasmithWorkerHandler::FDatasmithWorkerHandler(FDatasmithDispatcher& InDispatc
 	, bShouldTerminate(false)
 {
 	ThreadName = FString(TEXT("DatasmithWorkerHandler_")) + FString::FromInt(Id);
-	IOThread = FThread(*ThreadName, [this]() { Run(); } );
+	IOTask = UE::Tasks::Launch(*ThreadName, [this]() { Run(); } );
 }
 
 FDatasmithWorkerHandler::~FDatasmithWorkerHandler()
 {
-	IOThread.Join();
+	IOTask.Wait();
 }
 
 void FDatasmithWorkerHandler::StartWorkerProcess()
@@ -222,7 +222,7 @@ void FDatasmithWorkerHandler::RunInternal()
 				else
 				{
 					ValidateConnection();
-					if (ErrorState == EWorkerErrorState::WorkerProcess_Lost)
+					if (ErrorState == EWorkerErrorState::WorkerProcess_Lost || ErrorState == EWorkerErrorState::ConnectionLost)
 					{
 						if (CurrentTask.IsSet())
 						{
diff --git a/Engine/Plugins/Enterprise/DatasmithCADImporter/Source/DatasmithDispatcher/Public/DatasmithDispatcherConfig.h b/Engine/Plugins/Enterprise/DatasmithCADImporter/Source/DatasmithDispatcher/Public/DatasmithDispatcherConfig.h
--- a/Engine/Plugins/Enterprise/DatasmithCADImporter/Source/DatasmithDispatcher/Public/DatasmithDispatcherConfig.h
+++ b/Engine/Plugins/Enterprise/DatasmithCADImporter/Source/DatasmithDispatcher/Public/DatasmithDispatcherConfig.h
@@ -4,34 +4,36 @@
 
 namespace DatasmithDispatcher
 {
-	namespace Config
-	{
-		// Server side timeout for connection with the client
-		constexpr double AcceptTimeout_s = 3.0;
+namespace Config
+{
+
+// Server side timeout for connection with the client
+constexpr double AcceptTimeout_s = 3.0;
+
+// Client side timeout for connection with the server
+constexpr double ConnectTimeout_s = 3.0;
 
-		// Client side timeout for connection with the server
-		constexpr double ConnectTimeout_s = 3.0;
+// Maximum time we accept to wait when a command is sent.
+// If the command is not sent by this delay, the operation is considered as failed
+constexpr double SendCommandTimeout_s = 1.0;
 
-		// Maximum time we accept to wait when a command is sent.
-		// If the command is not sent by this delay, the operation is considered as failed
-		constexpr double SendCommandTimeout_s = 1.0;
+// When true, a Terminate command is sent to the processor and we wait for
+// TerminateTimeout_s seconds for it to close, before a hard TerminateProc.
+// When false, the process is directly closed with a TerminateProc.
+constexpr bool CloseProcessByCommand = true;
 
-		// When true, a Terminate command is sent to the processor and we wait for
-		// TerminateTimeout_s seconds for it to close, before a hard TerminateProc.
-		// When false, the process is directly closed with a TerminateProc.
-		constexpr bool CloseProcessByCommand = true;
+// Maximum time we accept to wait for the process to close after a Terminate
+// command has been sent.
+constexpr double TerminateTimeout_s = 1.0;
 
-		// Maximum time we accept to wait for the process to close after a Terminate
-		// command has been sent.
-		constexpr double TerminateTimeout_s = 1.0;
+// Delay between successive checks for a new task, and validation of the connection
+constexpr double IdleLoopDelay = 0.1;
 
-		// Delay between successive checks for a new task, and validation of the connection
-		constexpr double IdleLoopDelay = 0.1;
+// Delay between successive validation of the connection
+constexpr double ProcessingLoopDelay = 0.5;
 
-		// Delay between successive validation of the connection
-		constexpr double ProcessingLoopDelay = 0.5;
+// Number of worker instance that can be restarted when a crash occurs in en external process
+constexpr int32 MaxRestartAllowed = 1000;
 
-		// Number of worker instance that can be restarted when a crash occurs in en external process
-		constexpr int32 MaxRestartAllowed = 1000;
-	};
+} // ns Config
 } // ns DatasmithDispatcher
diff --git a/Engine/Plugins/Enterprise/DatasmithCADImporter/Source/DatasmithDispatcher/Public/DatasmithWorkerHandler.h b/Engine/Plugins/Enterprise/DatasmithCADImporter/Source/DatasmithDispatcher/Public/DatasmithWorkerHandler.h
--- a/Engine/Plugins/Enterprise/DatasmithCADImporter/Source/DatasmithDispatcher/Public/DatasmithWorkerHandler.h
+++ b/Engine/Plugins/Enterprise/DatasmithCADImporter/Source/DatasmithDispatcher/Public/DatasmithWorkerHandler.h
@@ -10,8 +10,7 @@
 #include "DatasmithDispatcherTask.h"
 
 #include "HAL/PlatformProcess.h"
-#include "HAL/Thread.h"
-
+#include "Tasks/Task.h"
 
 namespace DatasmithDispatcher
 {
@@ -55,8 +54,6 @@ private:
 	void StartWorkerProcess();
 	void ValidateConnection();
 
-// 	void RestartProcessor();
-
 	void ProcessCommand(ICommand& Command);
 	void ProcessCommand(FPingCommand& PingCommand);
 	void ProcessCommand(FCompletedTaskCommand& RunTaskCommand);
@@ -68,7 +65,7 @@ private:
 	// Send and receive commands
 	FNetworkServerNode NetworkInterface;
 	FCommandQueue CommandIO;
-	FThread IOThread;
+	UE::Tasks::FTask IOTask;
 	FString ThreadName;
 
 	// External process
diff --git a/Engine/Source/Programs/Enterprise/Datasmith/DatasmithCADWorker/Private/DatasmithCADWorkerImpl.cpp b/Engine/Source/Programs/Enterprise/Datasmith/DatasmithCADWorker/Private/DatasmithCADWorkerImpl.cpp
--- a/Engine/Source/Programs/Enterprise/Datasmith/DatasmithCADWorker/Private/DatasmithCADWorkerImpl.cpp
+++ b/Engine/Source/Programs/Enterprise/Datasmith/DatasmithCADWorker/Private/DatasmithCADWorkerImpl.cpp
@@ -9,12 +9,12 @@
 
 #include "HAL/FileManager.h"
 #include "HAL/PlatformProcess.h"
-#include "HAL/Thread.h"
 #include "Misc/CoreDelegates.h"
 #include "Misc/FileHelper.h"
 #include "Misc/Paths.h"
 #include "SocketSubsystem.h"
 #include "Sockets.h"
+#include "Tasks/Task.h"
 
 using namespace DatasmithDispatcher;
 
@@ -157,13 +157,13 @@ void FDatasmithCADWorkerImpl::ProcessCommand(const FRunTaskCommand& RunTaskComma
 
 	bProcessIsRunning = true;
 	int64 MaxDuration = DefineMaximumAllowedDuration(FileToProcess);
-	FThread TimeCheckerThread = FThread(TEXT("TimeCheckerThread"), [&]() { CheckDuration(FileToProcess, MaxDuration); });
+	UE::Tasks::FTask TimeChecker = UE::Tasks::Launch(TEXT("TimeChecker"), [&]() { CheckDuration(FileToProcess, MaxDuration); });
 
 	CADLibrary::FCADFileReader FileReader(ImportParameters, FileToProcess, EnginePluginsPath, CachePath);
 	ETaskState ProcessResult = FileReader.ProcessFile();
 
 	bProcessIsRunning = false;
-	TimeCheckerThread.Join();
+	TimeChecker.Wait();
 	
 	CompletedTask.ProcessResult = ProcessResult;
 
