The problem in this task is that each node of a branch of the scene graph can be in a different layer.
So if a leaf has a different layer than its ancestor, the branch has to be duplicated in the leaf's layer.
A second result is that empty branches are known only at the end. A process is therefore launched to recursively delete all the empty branches.

In the wait of option in the import widget, a CVar allow to activate this new import mode: ds.CADTranslator.Alias.LayersAsActors

Leaf with an hidden layer are no more built

diff --git a/Engine/Config/ConsoleVariables.ini b/Engine/Config/ConsoleVariables.ini
--- a/Engine/Config/ConsoleVariables.ini
+++ b/Engine/Config/ConsoleVariables.ini
@@ -114,18 +114,19 @@ fx.SkipVectorVMBackendOptimizations=1
 
 ; Uncomment to enable/disable temporary CAD processing file cache. These file will be use in a next import to avoid CAD file processing. 
 ; EnableCADCache = true enforce MaxImportThreads = 1
-; Default is enable
-;ds.CADTranslator.EnableCADCache=true
+; Default is true
+;ds.CADTranslator.EnableCADCache=false
 
 ; Uncomment to overwrite any existing cache associated with the file being imported. 
 ; Default is false
-;ds.CADTranslator.OverwriteCache=false
+;ds.CADTranslator.OverwriteCache=true
 
 ; Uncomment to enable/disable the timer that kill the worker if the import time is unusually long. With this time control, the load of the corrupted file 
 ; is canceled but the rest of the scene is imported. Default is enable
 ;ds.CADTranslator.EnableTimeControl=1
 
 ; Uncomment to enable/disable CADKernel tessellation. If disable, CAD SDK tessellator is used.
+; Default value is true
 ;ds.CADTranslator.DisableCADKernelTessellation=false
 
 ; Set Welding threshold for Heal/Sew stitching methods (in cm)
@@ -135,7 +136,7 @@ fx.SkipVectorVMBackendOptimizations=1
 ; Perform a welding of the mesh to try to stitch mesh cracks
 ; This welding is performed respecting the ds.CADTranslator.StitchingTolerance
 ; Default value is true
-;ds.CADTranslator.SewMeshIfNeeded=true
+;ds.CADTranslator.SewMeshIfNeeded=false
 
 ; Scale factor to change the unit of the DMU (Only applies to TechSoft import.)
 ; Default value of UnitScale is 1 i.e. unit = cm
@@ -147,7 +148,7 @@ fx.SkipVectorVMBackendOptimizations=1
 ; As this problem is rare, and to avoid the cost of the process, this option is disable by default. 
 ; To fix StaticMesh with duplicated triangles, the StaticMesh can be remesh with this option.
 ; Default value is false.
-;ds.CADTranslator.RemoveDuplicatedTriangle=false
+;ds.CADTranslator.RemoveDuplicatedTriangle=true
 
 ; Factor to allow to use smaller value than the defined minimal value of metric meshing parameters (i.e. Chord error > 0.005 cm, Max Edge Length > 1. cm)
 ; The used value of the meshing parameter is value * MeshingParameterFactor
@@ -156,13 +157,13 @@ fx.SkipVectorVMBackendOptimizations=1
 
 ; Alias Import Option
 ; Enable Sew action merges BReps according to their material i.e. only BReps associated with same material can be merged together.
-; Default is disable
-;ds.CADTranslator.Alias.SewByColor=false
+; Default is false
+;ds.CADTranslator.Alias.SewByMaterial=true
 
 ; Alias Import Option
 ; If true, the first level of actors in the outliner is the layers.
 ; Default is false
-;ds.CADTranslator.Alias.LayersAsActors=false
+;ds.CADTranslator.Alias.LayersAsActors=true
 
 ; JT Import Option
 ; Force to import embedded tessellation instead of meshing BRep model.
diff --git a/Engine/Plugins/Enterprise/DatasmithCADImporter/Source/CADTools/Private/CADOption.cpp b/Engine/Plugins/Enterprise/DatasmithCADImporter/Source/CADTools/Private/CADOption.cpp
--- a/Engine/Plugins/Enterprise/DatasmithCADImporter/Source/CADTools/Private/CADOption.cpp
+++ b/Engine/Plugins/Enterprise/DatasmithCADImporter/Source/CADTools/Private/CADOption.cpp
@@ -99,15 +99,6 @@ The used value of the meshing parameter is value * MeshingParameterFactor \n\
 Default value is 1.\n"),
 ECVF_Default);
 
-
-static bool bGAliasSewByColor = false;
-FAutoConsoleVariableRef GAliasSewByColor(
-	TEXT("ds.CADTranslator.Alias.SewByColor"),
-	bGAliasSewByColor,
-	TEXT("Enable Sew action merges BReps according to their material i.e. only BReps associated with same material can be merged together.\
-Default is disable\n"),
-ECVF_Default);
-
 uint32 GetTypeHash(const FImportParameters& ImportParameters)
 {
 	uint32 ParametersHash = ::GetTypeHash(ImportParameters.bGDisableCADKernelTessellation);
diff --git a/Engine/Plugins/Enterprise/DatasmithCADImporter/Source/DatasmithWireTranslator/Private/DatasmithWireTranslator.cpp b/Engine/Plugins/Enterprise/DatasmithCADImporter/Source/DatasmithWireTranslator/Private/DatasmithWireTranslator.cpp
--- a/Engine/Plugins/Enterprise/DatasmithCADImporter/Source/DatasmithWireTranslator/Private/DatasmithWireTranslator.cpp
+++ b/Engine/Plugins/Enterprise/DatasmithCADImporter/Source/DatasmithWireTranslator/Private/DatasmithWireTranslator.cpp
@@ -66,10 +66,32 @@ DEFINE_LOG_CATEGORY_STATIC(LogDatasmithWireTranslator, Log, All);
 #define WRONG_VERSION_TEXT "Unsupported version of Alias detected. Please upgrade to Alias 2021.3 (or later version)."
 #define CAD_INTERFACE_UNAVAILABLE "CAD Interface module is unavailable. Meshing will be done by Alias."
 
+#ifdef OPEN_MODEL_2023_0
+
+static bool bGAliasSewByMaterial = false;
+FAutoConsoleVariableRef GAliasSewByMaterial(
+	TEXT("ds.CADTranslator.Alias.SewByMaterial"),
+	bGAliasSewByMaterial,
+	TEXT("Enable Sew action merges BReps according to their material i.e. only BReps associated with same material can be merged together.\
+Default is disable\n"),
+ECVF_Default);
+
+static bool bGAliasLayerAsActor = false;
+FAutoConsoleVariableRef GAliasLayerAsActor(
+	TEXT("ds.CADTranslator.Alias.LayersAsActors"),
+	bGAliasLayerAsActor,
+	TEXT("If true, the first level of actors in the outliner is the layers.\
+Default is false\n"),
+ECVF_Default);
+
+#endif 
+
+
 namespace UE_DATASMITHWIRETRANSLATOR_NAMESPACE
 {
 
-static bool bGAliasSewByColor = false;
+static bool bGSewByMaterial = false;
+static bool bGLayersAsActors = false;
 
 static const FColor DefaultColor = FColor(200, 200, 200);
 
@@ -163,7 +185,7 @@ public:
 		LayerName = InLayerName;
 	};
 
-	bool IsCadData() const 
+	bool IsCadData() const
 	{
 		return bCadData;
 	}
@@ -247,10 +269,6 @@ public:
 		, SceneFullPath(InSceneFullName)
 		, SceneFileHash(0)
 		, AlRootNode(nullptr)
-		, FileVersion(0)
-		, ArchiveWireVersion(0)
-		, FileLength(0)
-		, NumCRCErrors(0)
 	{
 		// Set ProductName, ProductVersion in DatasmithScene for Analytics purpose
 		uint64 AliasFileVersion = FPlatformMisc::GetFileVersion(TEXT("libalias_api.dll"));
@@ -306,7 +324,7 @@ public:
 		}
 
 		CADLibrary::FImportParameters ImportParameters;
-		if(CADLibrary::FImportParameters::bGDisableCADKernelTessellation)
+		if (CADLibrary::FImportParameters::bGDisableCADKernelTessellation)
 		{
 			TSharedRef<FAliasModelToTechSoftConverter> AliasToTechSoftConverter = MakeShared<FAliasModelToTechSoftConverter>(ImportParameters);
 			CADModelConverter = AliasToTechSoftConverter;
@@ -319,10 +337,16 @@ public:
 			AliasBRepConverter = AliasToCADKernelConverter;
 		}
 
-		IConsoleVariable* CVarDSPlmXmlTranslatorEnabled = IConsoleManager::Get().FindConsoleVariable(TEXT("ds.CADTranslator.Alias.SewByColor"));
-		if (CVarDSPlmXmlTranslatorEnabled)
+		IConsoleVariable* CVarDSAliasSewByMaterialEnabled = IConsoleManager::Get().FindConsoleVariable(TEXT("ds.CADTranslator.Alias.SewByMaterial"));
+		if (CVarDSAliasSewByMaterialEnabled)
+		{
+			bGSewByMaterial = CVarDSAliasSewByMaterialEnabled->GetBool();
+		}
+
+		IConsoleVariable* CVarDSAliasLayersAsActorsEnabled = IConsoleManager::Get().FindConsoleVariable(TEXT("ds.CADTranslator.Alias.LayersAsActors"));
+		if (CVarDSAliasLayersAsActorsEnabled)
 		{
-			bGAliasSewByColor = CVarDSPlmXmlTranslatorEnabled->GetBool();
+			bGLayersAsActors = CVarDSAliasLayersAsActorsEnabled->GetBool();
 		}
 	}
 
@@ -347,21 +371,39 @@ private:
 		uint32 Uuid;  // Use for actor name
 		FString Label;
 		TSharedPtr<IDatasmithActorElement> ActorElement;
+		FString LayerName;
+		bool bLayerIsVisible = true;
+
+		TMap<FString, TSharedPtr<IDatasmithActorElement>> LayerToActorElement;
+		FDagNodeInfo* ParentInfo = nullptr;
+
+		FDagNodeInfo(FDagNodeInfo& InParentInfo)
+			: ParentInfo(&InParentInfo)
+		{}
+
+		FDagNodeInfo()
+			: ParentInfo(nullptr)
+		{}
+
+		void ExtractDagNodeInfo(AlDagNode& CurrentNode);
+		void ExtractDagNodeInfo(BodyData& CurrentNode);
+		void ExtractDagNodeLayer(const AlDagNode& InDagNode);
 	};
 
 	bool GetDagLeaves();
 	bool GetShader();
-	bool RecurseDagForLeaves(const TSharedPtr<AlDagNode>& DagNode, const FDagNodeInfo& ParentInfo);
-	void RecurseDagForLeavesNoMerge(const TSharedPtr<AlDagNode>& DagNode, const FDagNodeInfo& ParentInfo);
-	void DagForLeavesNoMerge(const TSharedPtr<AlDagNode>& DagNode, const FDagNodeInfo& ParentInfo);
-	bool ProcessAlGroupNode(AlDagNode& GroupNode, const FDagNodeInfo& ParentInfo);
-	bool ProcessAlShellNode(const TSharedPtr<AlDagNode>& ShellNode, const FDagNodeInfo& ParentInfo, const FString& ShaderName);
-	void ProcessBodyNode(const TSharedPtr<BodyData>& Body, const FDagNodeInfo& ParentInfo);
+	bool RecurseDagForLeaves(const TSharedPtr<AlDagNode>& DagNode, FDagNodeInfo& ParentInfo);
+	void RecurseDagForLeavesNoMerge(const TSharedPtr<AlDagNode>& DagNode, FDagNodeInfo& ParentInfo);
+	void DagForLeavesNoMerge(const TSharedPtr<AlDagNode>& DagNode, FDagNodeInfo& ParentInfo);
+	bool ProcessAlGroupNode(AlDagNode& GroupNode, FDagNodeInfo& ParentInfo);
+	void ProcessAlShellNode(const TSharedPtr<AlDagNode>& ShellNode, FDagNodeInfo& ParentInfo, const FString& ShaderName);
+	void ProcessBodyNode(const TSharedPtr<BodyData>& Body, FDagNodeInfo& ParentInfo);
 	TSharedPtr<IDatasmithMeshElement> FindOrAddMeshElement(const TSharedPtr<BodyData>& Body, const FDagNodeInfo& ParentInfo);
 	TSharedPtr<IDatasmithMeshElement> FindOrAddMeshElement(const TSharedPtr<AlDagNode>& ShellNode, const FDagNodeInfo& ParentInfo, const FString& ShaderName);
-	void GetDagNodeInfo(AlDagNode& GroupNode, const FDagNodeInfo& ParentInfo, FDagNodeInfo& CurrentNodeInfo);
-	void GetDagNodeInfo(BodyData& CurrentNode, const FDagNodeInfo& ParentInfo, FDagNodeInfo& CurrentNodeInfo);
-	void GetDagNodeMeta(const AlDagNode& InDagNode, TSharedPtr<IDatasmithActorElement> OutActorElement);
+	TSharedPtr<IDatasmithActorElement> FindOrAddParentActor(FDagNodeInfo& ParentInfo, const FString& LayerName);
+	TSharedPtr<IDatasmithActorElement> FindOrAddLayerActor(const FString& LayerName);
+
+	void RecurseDeleteEmptyActor(TSharedPtr<IDatasmithActorElement> Actor);
 
 	TOptional<FMeshDescription> GetMeshOfShellNode(AlDagNode& DagNode, TSharedRef<IDatasmithMeshElement> MeshElement, CADLibrary::FMeshParameters& MeshParameters);
 	TOptional<FMeshDescription> GetMeshOfNodeMesh(AlDagNode& DagNode, TSharedRef<IDatasmithMeshElement> MeshElement, CADLibrary::FMeshParameters& MeshParameters, AlMatrix4x4* AlMeshInvGlobalMatrix = nullptr);
@@ -445,21 +487,12 @@ private:
 	TMap<IDatasmithMeshElement*, TSharedPtr<AlDagNode>> MeshElementToAlDagNodeMap;
 
 	TMap<IDatasmithMeshElement*, TSharedPtr<BodyData>> MeshElementToBodyMap;
-	
+
 	TMap<IDatasmithMeshElement*, FString> MeshElementToShaderName;
 	TMap<FString, FColor> ShaderNameToColor;
 	TMap<FString, TSharedPtr<IDatasmithMaterialIDElement>> ShaderNameToUEMaterialId;
 
-	// start section information
-	int32 FileVersion;
-	int32 ArchiveWireVersion;
-
-	// length of archive returned by ON_BinaryArchive::Read3dmEndMark()
-	size_t FileLength;
-
-	// Number of crc errors found during archive reading.
-	// If > 0, then the archive is corrupt.
-	int32 NumCRCErrors;
+	TMap<FString, TSharedPtr<IDatasmithActorElement>> LayerNameToActor;
 
 	FDatasmithTessellationOptions TessellationOptions;
 
@@ -471,7 +504,7 @@ private:
 void FWireTranslatorImpl::SetTessellationOptions(const FDatasmithTessellationOptions& Options)
 {
 	TessellationOptions = Options;
-	CADModelConverter->SetImportParameters(Options.ChordTolerance, Options.MaxEdgeLength, Options.NormalTolerance, (CADLibrary::EStitchingTechnique) Options.StitchingTechnique);
+	CADModelConverter->SetImportParameters(Options.ChordTolerance, Options.MaxEdgeLength, Options.NormalTolerance, (CADLibrary::EStitchingTechnique)Options.StitchingTechnique);
 	SceneFileHash = HashCombine(Options.GetHash(), GetSceneFileHash(SceneFullPath, SceneName));
 }
 
@@ -1363,46 +1396,157 @@ bool FWireTranslatorImpl::GetShader()
 	return true;
 }
 
-
-
 bool FWireTranslatorImpl::GetDagLeaves()
 {
 	FDagNodeInfo RootContainer;
 	AlRootNode = TSharedPtr<AlDagNode>(AlUniverse::firstDagNode());
 	if (AlRootNode.IsValid())
 	{
-		return RecurseDagForLeaves(AlRootNode, RootContainer);
+		TArray<TSharedPtr<IDatasmithActorElement>> ActorsToRemove;
+		if (RecurseDagForLeaves(AlRootNode, RootContainer))
+		{
+			// Delete empty branches
+			for (int32 Index = DatasmithScene->GetActorsCount() - 1; Index >= 0; --Index)
+			{
+				TSharedPtr<IDatasmithActorElement> Actor = DatasmithScene->GetActor(Index);
+				if (Actor.IsValid() && !Actor->IsA(EDatasmithElementType::StaticMeshActor))
+				{
+					RecurseDeleteEmptyActor(Actor);
+					if (Actor->GetChildrenCount() == 0)
+					{
+						ActorsToRemove.Add(Actor);
+					}
+				}
+			}
+			for (const TSharedPtr<IDatasmithActorElement>& Actor : ActorsToRemove)
+			{
+				DatasmithScene->RemoveActor(Actor, EDatasmithActorRemovalRule::RemoveChildren);
+			}
+			return true;
+		}
 	}
+
 	return false;
 }
 
+void FWireTranslatorImpl::RecurseDeleteEmptyActor(TSharedPtr<IDatasmithActorElement> Actor)
+{
+	for (int32 Index = Actor->GetChildrenCount() - 1; Index >= 0; --Index)
+	{
+		TSharedPtr<IDatasmithActorElement> ActorChild = Actor->GetChild(Index);
+		if (ActorChild.IsValid() && !ActorChild->IsA(EDatasmithElementType::StaticMeshActor))
+		{
+			RecurseDeleteEmptyActor(ActorChild);
+			if (ActorChild->GetChildrenCount() == 0)
+			{
+				Actor->RemoveChild(ActorChild);
+			}
+		}
+	}
+}
 
-void FWireTranslatorImpl::GetDagNodeMeta(const AlDagNode& InDagNode, TSharedPtr<IDatasmithActorElement> OutActorElement)
+
+void FWireTranslatorImpl::FDagNodeInfo::ExtractDagNodeLayer(const AlDagNode& InDagNode)
 {
 	AlLayer* LayerPtr = InDagNode.layer();
 	if (LayerPtr)
 	{
 		TUniquePtr<AlLayer> Layer(LayerPtr);
-		FString LayerName = UTF8_TO_TCHAR(Layer->name());
-		OutActorElement->SetLayer(*LayerName);
+		LayerName = UTF8_TO_TCHAR(Layer->name());
+		bLayerIsVisible = !Layer->invisible();
+		if (ActorElement.IsValid())
+		{
+			ActorElement->SetLayer(*LayerName);
+			LayerToActorElement.Add(LayerName, ActorElement);
+		}
 	}
 }
 
-void FWireTranslatorImpl::GetDagNodeInfo(AlDagNode& CurrentNode, const FDagNodeInfo& ParentInfo, FDagNodeInfo& CurrentNodeInfo)
+void FWireTranslatorImpl::FDagNodeInfo::ExtractDagNodeInfo(AlDagNode& CurrentNode)
 {
-	CurrentNodeInfo.Label = UTF8_TO_TCHAR(CurrentNode.name());
+	Label = UTF8_TO_TCHAR(CurrentNode.name());
 	uint32 ThisGroupNodeUuid = OpenModelUtils::GetAlDagNodeUuid(CurrentNode);
-	CurrentNodeInfo.Uuid = HashCombine(ParentInfo.Uuid, ThisGroupNodeUuid);
+	Uuid = HashCombine(ParentInfo ? ParentInfo->Uuid : 1, ThisGroupNodeUuid);
+}
+
+void FWireTranslatorImpl::FDagNodeInfo::ExtractDagNodeInfo(BodyData& CurrentNode)
+{
+	if (!ParentInfo)
+	{
+		return;
+	}
+	Label = ParentInfo->Label;
+	CurrentNode.SetLabel(ParentInfo->Label);
+	Uuid = CurrentNode.GetUuid(ParentInfo->Uuid);
+}
+
+TSharedPtr<IDatasmithActorElement> FWireTranslatorImpl::FindOrAddLayerActor(const FString& LayerName)
+{
+	if (!bGLayersAsActors || LayerName.IsEmpty())
+	{
+		return TSharedPtr<IDatasmithActorElement>();
+	}
+
+	TSharedPtr<IDatasmithActorElement>* LayerActor = LayerNameToActor.Find(LayerName);
+	if (LayerActor)
+	{
+		return *LayerActor;
+	}
+
+	uint32 LayerUuid = GetTypeHash(LayerName);
+	TSharedRef<IDatasmithActorElement> NewLayerActor = FDatasmithSceneFactory::CreateActor(*OpenModelUtils::UuidToString(LayerUuid));
+	NewLayerActor->SetLabel(*LayerName);
+	NewLayerActor->SetLayer(*LayerName);
+	DatasmithScene->AddActor(NewLayerActor);
+	LayerNameToActor.Add(LayerName, NewLayerActor);
+	return NewLayerActor;
 }
 
-void FWireTranslatorImpl::GetDagNodeInfo(BodyData& CurrentNode, const FDagNodeInfo& ParentInfo, FDagNodeInfo& CurrentNodeInfo)
+TSharedPtr<IDatasmithActorElement> FWireTranslatorImpl::FindOrAddParentActor(FDagNodeInfo& ParentInfo, const FString& LayerName)
 {
-	CurrentNode.SetLabel(ParentInfo.Label);
-	CurrentNodeInfo.Label = ParentInfo.Label;
-	CurrentNodeInfo.Uuid = CurrentNode.GetUuid(ParentInfo.Uuid);
+	if (!ParentInfo.ActorElement.IsValid())
+	{
+		if (!LayerName.IsEmpty())
+		{
+			return FindOrAddLayerActor(LayerName);
+		}
+		return TSharedPtr<IDatasmithActorElement>();
+	}
+
+	if (!bGLayersAsActors)
+	{
+		return ParentInfo.ActorElement;
+	}
+
+	TSharedPtr<IDatasmithActorElement>* ActorParent = ParentInfo.LayerToActorElement.Find(LayerName);
+	if (ActorParent)
+	{
+		return *ActorParent;
+	}
+
+	uint32 LayerActorUuid = HashCombine(ParentInfo.Uuid, GetTypeHash(LayerName));
+
+	TSharedPtr<IDatasmithActorElement> LayerActorParent = FDatasmithSceneFactory::CreateActor(*OpenModelUtils::UuidToString(LayerActorUuid));
+	LayerActorParent->SetLabel(*ParentInfo.Label);
+	LayerActorParent->SetLayer(*LayerName);
+
+	ParentInfo.LayerToActorElement.Add(LayerName, LayerActorParent);
+
+	TSharedPtr<IDatasmithActorElement> GrandParentActor = FindOrAddParentActor(*ParentInfo.ParentInfo, LayerName);
+	if (GrandParentActor.IsValid())
+	{
+		GrandParentActor->AddChild(LayerActorParent);
+	}
+	else
+	{
+		TSharedPtr<IDatasmithActorElement> LayerActor = FindOrAddLayerActor(LayerName);
+		LayerActor->AddChild(LayerActorParent);
+	}
+
+	return LayerActorParent;
 }
 
-bool FWireTranslatorImpl::ProcessAlGroupNode(AlDagNode& GroupNode, const FDagNodeInfo& ParentInfo)
+bool FWireTranslatorImpl::ProcessAlGroupNode(AlDagNode& GroupNode, FDagNodeInfo& ParentInfo)
 {
 	AlGroupNode* AlGroup = GroupNode.asGroupNodePtr();
 	if (!AlIsValid(AlGroup))
@@ -1418,31 +1562,28 @@ bool FWireTranslatorImpl::ProcessAlGroupNode(AlDagNode& GroupNode, const FDagNod
 
 	TSharedPtr<AlDagNode> ChildNode(AlChildPtr);
 
-	FDagNodeInfo ThisGroupNodeInfo;
-	GetDagNodeInfo(GroupNode, ParentInfo, ThisGroupNodeInfo);
+	FDagNodeInfo ThisGroupNodeInfo(ParentInfo);
+	ThisGroupNodeInfo.ExtractDagNodeInfo(GroupNode);
 
 	ThisGroupNodeInfo.ActorElement = FDatasmithSceneFactory::CreateActor(*OpenModelUtils::UuidToString(ThisGroupNodeInfo.Uuid));
 	ThisGroupNodeInfo.ActorElement->SetLabel(*ThisGroupNodeInfo.Label);
-	GetDagNodeMeta(GroupNode, ThisGroupNodeInfo.ActorElement);
+	ThisGroupNodeInfo.ExtractDagNodeLayer(GroupNode);
 
-	RecurseDagForLeaves(ChildNode, ThisGroupNodeInfo);
+	// Apply local transform to actor element
+	OpenModelUtils::SetActorTransform(ThisGroupNodeInfo.ActorElement, GroupNode);
 
-	// add the resulting actor to the scene
-	if (OpenModelUtils::IsValidActor(ThisGroupNodeInfo.ActorElement))
+	TSharedPtr<IDatasmithActorElement> ParentActorElement = FindOrAddParentActor(ParentInfo, ThisGroupNodeInfo.LayerName);
+	if (ParentActorElement.IsValid())
 	{
-		// Apply local transform to actor element
-		OpenModelUtils::SetActorTransform(ThisGroupNodeInfo.ActorElement, GroupNode);
-
-		if (ParentInfo.ActorElement.IsValid())
-		{
-			ParentInfo.ActorElement->AddChild(ThisGroupNodeInfo.ActorElement);
-		}
-		else
-		{
-			DatasmithScene->AddActor(ThisGroupNodeInfo.ActorElement);
-		}
+		ParentActorElement->AddChild(ThisGroupNodeInfo.ActorElement);
+	}
+	else
+	{
+		DatasmithScene->AddActor(ThisGroupNodeInfo.ActorElement);
 	}
 
+	RecurseDagForLeaves(ChildNode, ThisGroupNodeInfo);
+
 	return true;
 }
 
@@ -1459,13 +1600,13 @@ TSharedPtr<IDatasmithMeshElement> FWireTranslatorImpl::FindOrAddMeshElement(cons
 	MeshElement->SetLightmapSourceUV(-1);
 
 	const TSet<FString>& ShaderNames = Body->GetShaderNames();
-	for(const FString& ShaderName : ShaderNames)
+	for (const FString& ShaderName : ShaderNames)
 	{
 		TSharedPtr<IDatasmithMaterialIDElement> MaterialElement = ShaderNameToUEMaterialId[ShaderName];
 		if (MaterialElement.IsValid())
 		{
 			const FColor* ShaderColor = ShaderNameToColor.Find(ShaderName);
-			if(ShaderColor)
+			if (ShaderColor)
 			{
 				int32 ShaderId = CreateShaderId(*ShaderColor);
 				MeshElement->SetMaterial(MaterialElement->GetName(), ShaderId);
@@ -1531,28 +1672,32 @@ TSharedPtr<IDatasmithMeshElement> FWireTranslatorImpl::FindOrAddMeshElement(cons
 	return MeshElement;
 }
 
-bool FWireTranslatorImpl::ProcessAlShellNode(const TSharedPtr<AlDagNode>& ShellNode, const FDagNodeInfo& ParentInfo, const FString& ShaderName)
+void FWireTranslatorImpl::ProcessAlShellNode(const TSharedPtr<AlDagNode>& ShellNode, FDagNodeInfo& ParentInfo, const FString& ShaderName)
 {
-	FDagNodeInfo ShellInfo;
-	GetDagNodeInfo(*ShellNode, ParentInfo, ShellInfo);
+	FDagNodeInfo ShellInfo(ParentInfo);
+	ShellInfo.ExtractDagNodeInfo(*ShellNode);
 
 	TSharedPtr<IDatasmithMeshElement> MeshElement = FindOrAddMeshElement(ShellNode, ShellInfo, ShaderName);
 	if (!MeshElement.IsValid())
 	{
-		return false;
+		return;
 	}
 
 	TSharedPtr<IDatasmithMeshActorElement> ActorElement = FDatasmithSceneFactory::CreateMeshActor(*OpenModelUtils::UuidToString(ShellInfo.Uuid));
 	if (!ActorElement.IsValid())
 	{
-		return false;
+		return;
 	}
 
 	ActorElement->SetLabel(*ShellInfo.Label);
 	ActorElement->SetStaticMeshPathName(MeshElement->GetName());
 	ShellInfo.ActorElement = ActorElement;
 
-	GetDagNodeMeta(*ShellNode, ActorElement);
+	ShellInfo.ExtractDagNodeLayer(*ShellNode);
+	if (!ShellInfo.bLayerIsVisible)
+	{
+		return;
+	}
 
 	OpenModelUtils::SetActorTransform(ShellInfo.ActorElement, *ShellNode);
 
@@ -1570,22 +1715,20 @@ bool FWireTranslatorImpl::ProcessAlShellNode(const TSharedPtr<AlDagNode>& ShellN
 		}
 	}
 
-	if (ActorElement.IsValid() && OpenModelUtils::IsValidActor(ActorElement))
+	TSharedPtr<IDatasmithActorElement> ParentElement = FindOrAddParentActor(ParentInfo, ShellInfo.LayerName);
+	if (ParentElement.IsValid())
 	{
-		if (ParentInfo.ActorElement.IsValid())
-		{
-			ParentInfo.ActorElement->AddChild(ActorElement);
-		}
-		else
-		{
-			DatasmithScene->AddActor(ActorElement);
-		}
+		ParentElement->AddChild(ActorElement);
+	}
+	else
+	{
+		DatasmithScene->AddActor(ActorElement);
 	}
-	return true;
 }
 
-void FWireTranslatorImpl::ProcessBodyNode(const TSharedPtr<BodyData>& Body, const FDagNodeInfo& ParentInfo)
+void FWireTranslatorImpl::ProcessBodyNode(const TSharedPtr<BodyData>& Body, FDagNodeInfo& ParentInfo)
 {
+
 	if (!Body.IsValid())
 	{
 		return;
@@ -1597,8 +1740,8 @@ void FWireTranslatorImpl::ProcessBodyNode(const TSharedPtr<BodyData>& Body, cons
 		return;
 	}
 
-	FDagNodeInfo ShellInfo;
-	GetDagNodeInfo(*Body, ParentInfo, ShellInfo);
+	FDagNodeInfo ShellInfo(ParentInfo);
+	ShellInfo.ExtractDagNodeInfo(*Body);
 
 	TSharedPtr<IDatasmithMeshElement> MeshElement = FindOrAddMeshElement(Body, ShellInfo);
 	if (!MeshElement.IsValid())
@@ -1612,7 +1755,7 @@ void FWireTranslatorImpl::ProcessBodyNode(const TSharedPtr<BodyData>& Body, cons
 		return;
 	}
 
-	ActorElement->SetLabel(*ShellInfo.Label);
+	ActorElement->SetLabel(ShellInfo.Label.IsEmpty() ? *Body->GetLayerName() : *ShellInfo.Label);
 	ActorElement->SetStaticMeshPathName(MeshElement->GetName());
 	ShellInfo.ActorElement = ActorElement;
 
@@ -1632,17 +1775,16 @@ void FWireTranslatorImpl::ProcessBodyNode(const TSharedPtr<BodyData>& Body, cons
 		}
 	}
 
-	if (ActorElement.IsValid() && OpenModelUtils::IsValidActor(ActorElement))
+	TSharedPtr<IDatasmithActorElement> ParentElement = FindOrAddParentActor(ParentInfo, Body->GetLayerName());
+	if (ParentElement.IsValid())
 	{
-		if (ParentInfo.ActorElement.IsValid())
-		{
-			ParentInfo.ActorElement->AddChild(ActorElement);
-		}
-		else
-		{
-			DatasmithScene->AddActor(ActorElement);
-		}
+		ParentElement->AddChild(ActorElement);
+	}
+	else
+	{
+		DatasmithScene->AddActor(ActorElement);
 	}
+
 	return;
 }
 
@@ -1662,8 +1804,8 @@ TSharedPtr<AlDagNode> GetNextNode(const TSharedPtr<AlDagNode>& DagNode)
 
 uint32 GetBodyGroupUuid(const FString& ShaderName, const FString& LayerName, bool bCadData)
 {
-	uint32 Uuid = HashCombine(GetTypeHash(LayerName), GetTypeHash(bCadData)); 
-	if(bGAliasSewByColor)
+	uint32 Uuid = HashCombine(GetTypeHash(LayerName), GetTypeHash(bCadData));
+	if (bGSewByMaterial)
 	{
 		Uuid = HashCombine(GetTypeHash(ShaderName), Uuid);
 	}
@@ -1691,6 +1833,10 @@ void FWireTranslatorImpl::AddNodeInBodyGroup(TSharedPtr<AlDagNode>& DagNode, con
 	{
 		const TUniquePtr<AlLayer> Layer(LayerPtr);
 		LayerName = UTF8_TO_TCHAR(Layer->name());
+		if (Layer->invisible())
+		{
+			return;
+		}
 	}
 
 	uint32 SetId = GetBodyGroupUuid(ShaderName, LayerName, bIsAPatch);
@@ -1712,7 +1858,7 @@ void FWireTranslatorImpl::AddNodeInBodyGroup(TSharedPtr<AlDagNode>& DagNode, con
 	Body->Add(DagNode, Color, ShaderName);
 }
 
-bool FWireTranslatorImpl::RecurseDagForLeaves(const TSharedPtr<AlDagNode>& FirstDagNode, const FDagNodeInfo& ParentInfo)
+bool FWireTranslatorImpl::RecurseDagForLeaves(const TSharedPtr<AlDagNode>& FirstDagNode, FDagNodeInfo& ParentInfo)
 {
 	if (TessellationOptions.StitchingTechnique != EDatasmithCADStitchingTechnique::StitchingSew)
 	{
@@ -1742,83 +1888,83 @@ bool FWireTranslatorImpl::RecurseDagForLeaves(const TSharedPtr<AlDagNode>& First
 		switch (objectType)
 		{
 			// Push all leaf nodes into 'leaves'
-			case kShellNodeType:
+		case kShellNodeType:
+		{
+			AlShellNode* ShellNode = DagNode->asShellNodePtr();
+			if (!AlIsValid(ShellNode))
 			{
-				AlShellNode* ShellNode = DagNode->asShellNodePtr();
-				if (!AlIsValid(ShellNode))
-				{
-					break;
-				}
-
-				AlShell* ShellPtr = ShellNode->shell();
-				if (!AlIsValid(ShellPtr))
-				{
-					break;
-				}
+				break;
+			}
 
-				TUniquePtr<AlShell> Shell(ShellPtr);
-				uint32 NbPatch = GetPatchCount(Shell);
+			AlShell* ShellPtr = ShellNode->shell();
+			if (!AlIsValid(ShellPtr))
+			{
+				break;
+			}
 
-				TUniquePtr<AlShader> Shader(Shell->firstShader());
-				if (Shader.IsValid())
-				{
-					ShaderName = Shader->name();
-				}
+			TUniquePtr<AlShell> Shell(ShellPtr);
+			uint32 NbPatch = GetPatchCount(Shell);
 
-				if (NbPatch == 1)
-				{
-					AddNodeInBodyGroup(DagNode, ShaderName, ShellToProcess, true, MaxSize);
-				}
-				else
-				{
-					ProcessAlShellNode(DagNode, ParentInfo, ShaderName);
-				}
-				break;
+			TUniquePtr<AlShader> Shader(Shell->firstShader());
+			if (Shader.IsValid())
+			{
+				ShaderName = Shader->name();
 			}
 
-			case kSurfaceNodeType:
+			if (NbPatch == 1)
 			{
-				AlSurfaceNode* SurfaceNode = DagNode->asSurfaceNodePtr();
-				TUniquePtr<AlSurface> Surface(SurfaceNode->surface());
-				if (Surface.IsValid())
-				{
-					TUniquePtr<AlShader> Shader(Surface->firstShader());
-					if (Shader.IsValid())
-					{
-						ShaderName = Shader->name();
-					}
-				}
 				AddNodeInBodyGroup(DagNode, ShaderName, ShellToProcess, true, MaxSize);
-				break;
 			}
+			else
+			{
+				ProcessAlShellNode(DagNode, ParentInfo, ShaderName);
+			}
+			break;
+		}
 
-			case kMeshNodeType:
+		case kSurfaceNodeType:
+		{
+			AlSurfaceNode* SurfaceNode = DagNode->asSurfaceNodePtr();
+			TUniquePtr<AlSurface> Surface(SurfaceNode->surface());
+			if (Surface.IsValid())
 			{
-				AlMeshNode* MeshNode = DagNode->asMeshNodePtr();
-				TUniquePtr<AlMesh> Mesh(MeshNode->mesh());
-				if (Mesh.IsValid())
+				TUniquePtr<AlShader> Shader(Surface->firstShader());
+				if (Shader.IsValid())
 				{
-					TUniquePtr<AlShader> Shader(Mesh->firstShader());
-					if (Shader.IsValid())
-					{
-						ShaderName = Shader->name();
-					}
+					ShaderName = Shader->name();
 				}
-				AddNodeInBodyGroup(DagNode, ShaderName, ShellToProcess, false, MaxSize);
-				break;
 			}
+			AddNodeInBodyGroup(DagNode, ShaderName, ShellToProcess, true, MaxSize);
+			break;
+		}
 
-			// Traverse down through groups
-			case kGroupNodeType:
+		case kMeshNodeType:
+		{
+			AlMeshNode* MeshNode = DagNode->asMeshNodePtr();
+			TUniquePtr<AlMesh> Mesh(MeshNode->mesh());
+			if (Mesh.IsValid())
 			{
-				ProcessAlGroupNode(*DagNode, ParentInfo);
-				break;
+				TUniquePtr<AlShader> Shader(Mesh->firstShader());
+				if (Shader.IsValid())
+				{
+					ShaderName = Shader->name();
+				}
 			}
+			AddNodeInBodyGroup(DagNode, ShaderName, ShellToProcess, false, MaxSize);
+			break;
+		}
 
-			default:
-			{
-				break;
-			}
+		// Traverse down through groups
+		case kGroupNodeType:
+		{
+			ProcessAlGroupNode(*DagNode, ParentInfo);
+			break;
+		}
+
+		default:
+		{
+			break;
+		}
 		}
 
 		DagNode = GetNextNode(DagNode);
@@ -1831,7 +1977,7 @@ bool FWireTranslatorImpl::RecurseDagForLeaves(const TSharedPtr<AlDagNode>& First
 	return true;
 }
 
-void FWireTranslatorImpl::RecurseDagForLeavesNoMerge(const TSharedPtr<AlDagNode>& FirstDagNode, const FDagNodeInfo& ParentInfo)
+void FWireTranslatorImpl::RecurseDagForLeavesNoMerge(const TSharedPtr<AlDagNode>& FirstDagNode, FDagNodeInfo& ParentInfo)
 {
 	TSharedPtr<AlDagNode> DagNode = FirstDagNode;
 	while (DagNode)
@@ -1841,7 +1987,7 @@ void FWireTranslatorImpl::RecurseDagForLeavesNoMerge(const TSharedPtr<AlDagNode>
 	}
 }
 
-void FWireTranslatorImpl::DagForLeavesNoMerge(const TSharedPtr<AlDagNode>& DagNode, const FDagNodeInfo& ParentInfo)
+void FWireTranslatorImpl::DagForLeavesNoMerge(const TSharedPtr<AlDagNode>& DagNode, FDagNodeInfo& ParentInfo)
 {
 	const char* ShaderName = nullptr;
 
@@ -1928,7 +2074,7 @@ TOptional<FMeshDescription> FWireTranslatorImpl::MeshDagNodeWithExternalMesher(A
 
 	const FColor* ColorPtr = nullptr;
 	const FString* ShaderName = MeshElementToShaderName.Find(&MeshElement.Get());
-	if(ShaderName)
+	if (ShaderName)
 	{
 		ColorPtr = ShaderNameToColor.Find(*ShaderName);
 	}
@@ -2054,7 +2200,7 @@ TOptional<FMeshDescription> FWireTranslatorImpl::GetMeshOfMeshBody(TSharedRef<Bo
 		const bool bMerge = true;
 		OpenModelUtils::TransferAlMeshToMeshDescription(*MeshPtr, *SlotMaterialName, MeshDescription, MeshParameters, bMerge);
 	}
-	
+
 	// Build edge meta data
 	FStaticMeshOperations::DetermineEdgeHardnessesFromVertexInstanceNormals(MeshDescription);
 
@@ -2246,8 +2392,8 @@ void FDatasmithWireTranslator::Initialize(FDatasmithTranslatorCapabilities& OutC
 	if (GIsEditor && !GEditor->PlayWorld && !GIsPlayInEditorWorld)
 	{
 #ifdef USE_OPENMODEL
- 		static void* DllHandle = FPlatformProcess::GetDllHandle(TEXT("libalias_api.dll"));
- 		if (DllHandle)
+		static void* DllHandle = FPlatformProcess::GetDllHandle(TEXT("libalias_api.dll"));
+		if (DllHandle)
 		{
 			// Check installed version of Alias Tools because binaries before 2021.3 are not compatible with Alias 2022
 			uint64 FileVersion = FPlatformMisc::GetFileVersion(TEXT("libalias_api.dll"));
@@ -2255,9 +2401,9 @@ void FDatasmithWireTranslator::Initialize(FDatasmithTranslatorCapabilities& OutC
 #ifdef OPEN_MODEL_2020
 			if (LibAlias2020_Version < FileVersion && FileVersion < LibAlias2021_Version)
 			{
-				TFunction<bool()> DisplayMessage = []() -> bool 
-				{ 
-					UE_LOG(LogDatasmithWireTranslator, Warning, TEXT(WRONG_VERSION_TEXT)); return true; 
+				TFunction<bool()> DisplayMessage = []() -> bool
+				{
+					UE_LOG(LogDatasmithWireTranslator, Warning, TEXT(WRONG_VERSION_TEXT)); return true;
 				};
 				static const bool bIsDispaly = DisplayMessage();
 				OutCapabilities.bIsEnabled = false;
@@ -2304,7 +2450,7 @@ bool FDatasmithWireTranslator::LoadScene(TSharedRef<IDatasmithScene> OutScene)
 
 	UE_LOG(LogDatasmithWireTranslator, Display, TEXT("CAD translation [%s]."), *Filename);
 	UE_LOG(LogDatasmithWireTranslator, Display, TEXT(" - Parsing Library:      %s"), TEXT("Alias"));
-	UE_LOG(LogDatasmithWireTranslator, Display, TEXT(" - Tessellation Library: %s") , CADLibrary::FImportParameters::bGDisableCADKernelTessellation ? TEXT("TechSoft") : TEXT("CADKernel"));
+	UE_LOG(LogDatasmithWireTranslator, Display, TEXT(" - Tessellation Library: %s"), CADLibrary::FImportParameters::bGDisableCADKernelTessellation ? TEXT("TechSoft") : TEXT("CADKernel"));
 
 	FString OutputPath = FPaths::ConvertRelativePathToFull(FPaths::Combine(FDatasmithWireTranslatorModule::Get().GetTempDir(), TEXT("Cache"), GetSource().GetSceneName()));
 	IFileManager::Get().MakeDirectory(*OutputPath, true);
