The crash was due to a call on a nullptr.
A nullptr check is added to avoid the cash

diff --git a/Engine/Plugins/Enterprise/DatasmithCADImporter/Source/DatasmithCADTranslator/Private/DatasmithSceneGraphBuilder.cpp b/Engine/Plugins/Enterprise/DatasmithCADImporter/Source/DatasmithCADTranslator/Private/DatasmithSceneGraphBuilder.cpp
--- a/Engine/Plugins/Enterprise/DatasmithCADImporter/Source/DatasmithCADTranslator/Private/DatasmithSceneGraphBuilder.cpp
+++ b/Engine/Plugins/Enterprise/DatasmithCADImporter/Source/DatasmithCADTranslator/Private/DatasmithSceneGraphBuilder.cpp
@@ -323,9 +323,12 @@ TSharedPtr<IDatasmithActorElement>  FDatasmithSceneBaseGraphBuilder::BuildInstan
 
 					FCadId RootId = 1;
 					Reference = &(SceneGraph->GetReference(RootId));
+					if (Reference)
+					{
+						Reference->SetGraphicProperties(UnloadedReference);
+					}
 				}
 			}
-			Reference->SetGraphicProperties(UnloadedReference);
 		}
 
 		if (!Reference)
