the problem was due to an uninitialized structure partially filled

diff --git a/Engine/Plugins/Enterprise/DatasmithCADImporter/Source/CADTools/Public/CADData.h b/Engine/Plugins/Enterprise/DatasmithCADImporter/Source/CADTools/Public/CADData.h
--- a/Engine/Plugins/Enterprise/DatasmithCADImporter/Source/CADTools/Public/CADData.h
+++ b/Engine/Plugins/Enterprise/DatasmithCADImporter/Source/CADTools/Public/CADData.h
@@ -77,12 +77,12 @@ public:
 
 public:
 	FString MaterialName;
-	FColor Diffuse;
-	FColor Ambient;
-	FColor Specular;
-	float Shininess;
-	float Transparency;
-	float Reflexion;
+	FColor Diffuse = FColor(255, 255, 255);
+	FColor Ambient = FColor(255, 255, 255);
+	FColor Specular = FColor(255, 255, 255);
+	float Shininess = 0.f;
+	float Transparency = 0.f;
+	float Reflexion = 0.f;
 	FString TextureName;
 };
 
