The propagation of the local unit was broken :
Due to the fact that the transform can be defined in the reference, the FArchiveReference is build before the extraction of the instance transform (and the local scale/unit)
So, this information wasn't propagate from the ArchiveInstance to the ArchiveReferenc
This is done at line 911 of TechSoftFileParser.cpp

The case of two instance referencing a Reference with a different scale wasn't manage.
This is done at  869 / 891 of TechSoftFileParser.cpp
I don't have usecase to verifiy this case.

diff --git a/Engine/Plugins/Enterprise/DatasmithCADImporter/Source/CADInterfaces/Private/CADSceneGraph.cpp b/Engine/Plugins/Enterprise/DatasmithCADImporter/Source/CADInterfaces/Private/CADSceneGraph.cpp
--- a/Engine/Plugins/Enterprise/DatasmithCADImporter/Source/CADInterfaces/Private/CADSceneGraph.cpp
+++ b/Engine/Plugins/Enterprise/DatasmithCADImporter/Source/CADInterfaces/Private/CADSceneGraph.cpp
@@ -131,10 +131,7 @@ FArchiveInstance& FArchiveSceneGraph::AddInstance(const FArchiveCADObject& Paren
 
 	int32 Index = Instances.Emplace(InstanceId, Parent);
 	CADIdToIndex.Add(Index);
-	FArchiveInstance& Instance = Instances[Index];
-	Instance.Unit = Parent.Unit;
-
-	return Instance;
+	return Instances[Index];
 }
 
 FArchiveInstance& FArchiveSceneGraph::GetInstance(FCadId CadId)
@@ -207,7 +204,6 @@ FArchiveReference& FArchiveSceneGraph::AddOccurence(FArchiveReference& Parent)
 	Parent.AddChild(Instance.Id);
 	FArchiveReference& Reference = AddReference(Instance);
 	Instance.ReferenceNodeId = Reference.Id;
-	Reference.Unit = Parent.Unit;
 	return Reference;
 }
 
diff --git a/Engine/Plugins/Enterprise/DatasmithCADImporter/Source/CADInterfaces/Private/TechSoftFileParser.cpp b/Engine/Plugins/Enterprise/DatasmithCADImporter/Source/CADInterfaces/Private/TechSoftFileParser.cpp
--- a/Engine/Plugins/Enterprise/DatasmithCADImporter/Source/CADInterfaces/Private/TechSoftFileParser.cpp
+++ b/Engine/Plugins/Enterprise/DatasmithCADImporter/Source/CADInterfaces/Private/TechSoftFileParser.cpp
@@ -162,7 +162,7 @@ void TraverseAttribute(const A3DMiscAttributeData& AttributeData, TMap<FString,
 			{
 				AttributeName += TEXT(" ") + AttributeTitle;
 			}
-			else if(Index > 0)
+			else if (Index > 0)
 			{
 				AttributeName += TEXT(" ") + FString::FromInt((int32)Index);
 			}
@@ -331,7 +331,7 @@ ECADParsingResult FTechSoftFileParser::Process()
 		return ECADParsingResult::FileNotFound;
 	}
 
-	A3DImport Import(TCHAR_TO_UTF8(*File.GetPathOfFileToLoad())); 
+	A3DImport Import(TCHAR_TO_UTF8(*File.GetPathOfFileToLoad()));
 
 	TechSoftFileParserImpl::SetIOOption(Import);
 
@@ -439,7 +439,7 @@ void FTechSoftFileParser::SewModel()
 {
 	CADLibrary::TUniqueTSObj<A3DSewOptionsData> SewData;
 	SewData->m_bComputePreferredOpenShellOrientation = false;
-	
+
 	TechSoftInterface::SewModel(ModelFile.Get(), CADLibrary::FImportParameters::GStitchingTolerance, SewData.GetPtr());
 }
 
@@ -480,7 +480,7 @@ void FTechSoftFileParser::GenerateBodyMesh(A3DRiRepresentationItem* Representati
 		{
 			TechSoftUtils::FillBodyMesh(NewBReps[Index], CADFileData.GetImportParameters(), Body.Unit, BodyMesh);
 		}
-	} 
+	}
 	else
 	{
 		TechSoftUtils::FillBodyMesh(Representation, CADFileData.GetImportParameters(), Body.Unit, BodyMesh);
@@ -536,11 +536,11 @@ void FTechSoftFileParser::GenerateBodyMesh(A3DRiRepresentationItem* Representati
 			// This will be used when the file is reloaded
 			JsonObject->SetNumberField(JSON_ENTRY_BODY_UNIT, Body.Unit);
 
-			if(Body.ColorUId)
+			if (Body.ColorUId)
 			{
 				JsonObject->SetNumberField(JSON_ENTRY_COLOR_NAME, Body.ColorUId);
 			}
-			if(Body.MaterialUId)
+			if (Body.MaterialUId)
 			{
 				JsonObject->SetNumberField(JSON_ENTRY_MATERIAL_NAME, Body.MaterialUId);
 			}
@@ -605,7 +605,7 @@ ECADParsingResult FTechSoftFileParser::TraverseModel()
 	ExtractSpecificMetaData(ModelFile.Get(), Reference);
 	Reference.Unit = FileUnit;
 
-	if(ModelFileData->m_uiPOccurrencesSize == 0)
+	if (ModelFileData->m_uiPOccurrencesSize == 0)
 	{
 		CADFileData.AddWarningMessages(FString::Printf(TEXT("File %s is empty."), *CADFileData.GetCADFileDescription().GetFileName()));
 		return ECADParsingResult::ProcessFailed;
@@ -785,7 +785,7 @@ void FTechSoftFileParser::TraverseReference(const A3DAsmProductOccurrence* A3DRe
 	FMatrix ReferenceMatrix = FMatrix::Identity;
 	A3DMiscTransformation* Transform = ReferenceData->m_pLocation;
 	ExtractTransformation(Transform, Reference);
-	
+
 	Reference.TransformMatrix = Reference.TransformMatrix * ReferenceMatrix;
 
 	FArchiveInstance EmptyInstance;
@@ -848,11 +848,11 @@ void FTechSoftFileParser::TraverseOccurrence(const A3DAsmProductOccurrence* Occu
 	BuildInstanceName(Instance, ParentReference);
 
 	A3DMiscTransformation* Transform = OccurrenceData->m_pLocation;
-	
+
 	A3DAsmProductOccurrence* ReferencePtr = ProductOccurrence::GetReference(*OccurrenceData);
 
 	// Is the Reference already processed ?
-	if(ReferencePtr)
+	if (ReferencePtr)
 	{
 		FCadId* ReferenceId = ReferenceCache.Find(ReferencePtr);
 		if (ReferenceId != nullptr)
@@ -865,6 +865,30 @@ void FTechSoftFileParser::TraverseOccurrence(const A3DAsmProductOccurrence* Occu
 			}
 
 			ExtractTransformation(Transform, Instance);
+
+			double ReferenceUnit = 1;
+			if (SceneGraph.IsAReference(Instance.ReferenceNodeId))
+			{
+				FArchiveReference& Reference = SceneGraph.GetReference(Instance.ReferenceNodeId);
+				ReferenceUnit = Reference.Unit;
+			}
+			else if (SceneGraph.IsAUnloadedReference(Instance.ReferenceNodeId))
+			{
+				FArchiveUnloadedReference& Reference = SceneGraph.GetUnloadedReference(Instance.ReferenceNodeId);
+				ReferenceUnit = Reference.Unit;
+			}
+
+			// check if the instance unit is nearly equal to the existing reference unit, otherwise add a scale component to the instance transform
+			double Scale = 1;
+			if (!FMath::IsNearlyZero(ReferenceUnit))
+			{
+				Scale = Instance.Unit / ReferenceUnit;
+				if (!FMath::IsNearlyEqual(Scale, 1.))
+				{
+					Instance.TransformMatrix.ApplyScale(Scale);
+				}
+			}
+
 			return;
 		}
 	}
@@ -884,6 +908,7 @@ void FTechSoftFileParser::TraverseOccurrence(const A3DAsmProductOccurrence* Occu
 	}
 
 	ExtractTransformation(Transform, Instance);
+	UnloadedReference.Unit = Instance.Unit;
 
 	if (UnloadedReference.bIsUnloaded)
 	{
@@ -896,7 +921,7 @@ void FTechSoftFileParser::TraverseOccurrence(const A3DAsmProductOccurrence* Occu
 		ProcessReference(OccurrencePtr, Instance, NewReference);
 	}
 
-	if(ReferencePtr)
+	if (ReferencePtr)
 	{
 		ReferenceCache.Add(ReferencePtr, Instance.ReferenceNodeId);
 	}
@@ -919,7 +944,7 @@ void FTechSoftFileParser::ProcessReference(const A3DAsmProductOccurrence* Occurr
 		OccurrenceData.FillFrom(CachedOccurrencePtr);
 	}
 
-	if(ProductOccurrence::HasNoPartNoChild(*OccurrenceData))
+	if (ProductOccurrence::HasNoPartNoChild(*OccurrenceData))
 	{
 		return;
 	}
@@ -960,7 +985,7 @@ void FTechSoftFileParser::CountUnderOccurrence(const A3DAsmProductOccurrence* Oc
 		A3DAsmProductOccurrence* ReferencePtr = ProductOccurrence::GetReference(*OccurrenceData);
 
 		// Is the Reference already processed ?
-		if(ReferencePtr)
+		if (ReferencePtr)
 		{
 			FCadId* ReferenceId = ReferenceCache.Find(ReferencePtr);
 			if (ReferenceId != nullptr)
@@ -1044,8 +1069,8 @@ void FTechSoftFileParser::ProcessPrototype(const A3DAsmProductOccurrence* InProt
 				}
 			}
 		}
-		
-		if (ProductOccurrence::HasPartOrChild(*PrototypeData) )
+
+		if (ProductOccurrence::HasPartOrChild(*PrototypeData))
 		{
 			OutReference.bIsUnloaded = false;
 			PrototypePtr = nullptr;
@@ -1075,7 +1100,7 @@ void FTechSoftFileParser::ProcessPrototype(const A3DAsmProductOccurrence* InProt
 
 	if (OutReference.bIsUnloaded)
 	{
-		if(OutReference.Label.IsEmpty())
+		if (OutReference.Label.IsEmpty())
 		{
 			OutReference.Label = OutReference.ExternalFile.GetFileName();
 		}
@@ -1278,7 +1303,7 @@ void FTechSoftFileParser::ExtractMetaData(const A3DEntity* Entity, FArchiveCADOb
 		if (OutObject.Label.IsEmpty() && MetaData->m_pcName && MetaData->m_pcName[0] != '\0')
 		{
 			FString Name = UTF8_TO_TCHAR(MetaData->m_pcName);
-			if(Name != TEXT("unnamed"))  // "unnamed" is create by Techsoft. This name is ignored 
+			if (Name != TEXT("unnamed"))  // "unnamed" is create by Techsoft. This name is ignored 
 			{
 				Name = TechSoftUtils::CleanLabel(Name);
 				OutObject.Label = Name;
@@ -1363,7 +1388,7 @@ void FTechSoftFileParser::BuildInstanceName(FArchiveInstance& InstanceData, cons
 		case ECADFormat::CATIA:
 			InstanceData.Label = TechSoftUtils::CleanCatiaInstanceLabel(InstanceData.Label);
 			break;
-		
+
 		case ECADFormat::CATIA_3DXML:
 			InstanceData.Label = TechSoftUtils::Clean3dxmlInstanceLabel(InstanceData.Label);
 			break;
@@ -1399,14 +1424,14 @@ void FTechSoftFileParser::BuildBodyName(FArchiveBody& Body, const FArchiveRefere
 		Body.Label = TechSoftUtils::CleanCreoLabel(Body.Label);
 	}
 
-	if(Body.IsNameDefined())
+	if (Body.IsNameDefined())
 	{
 		return;
 	}
 
 	if (Format == ECADFormat::CATIA)
 	{
-		if(Body.SetNameWithAttributeValue(TEXT("BodyID")))
+		if (Body.SetNameWithAttributeValue(TEXT("BodyID")))
 		{
 			return;
 		}
@@ -1518,7 +1543,7 @@ void FTechSoftFileParser::ExtractSpecificMetaData(const A3DAsmProductOccurrence*
 				A3DElementsByRefsetUg Refset = UnigraphicsSpecificData->m_asChildrenByRefsets[Index];
 			}
 
-			if(UnigraphicsSpecificData->m_uiSolidsByRefsetsSize)
+			if (UnigraphicsSpecificData->m_uiSolidsByRefsetsSize)
 			{
 				for (uint32 Index = 0; Index < UnigraphicsSpecificData->m_uiSolidsByRefsetsSize; ++Index)
 				{
@@ -1576,7 +1601,7 @@ FArchiveColor& FTechSoftFileParser::FindOrAddColor(uint32 ColorIndex, uint8 Alph
 	FArchiveColor& NewColor = CADFileData.AddColor(ColorHId);
 	NewColor.Color = TechSoftUtils::GetColorAt(ColorIndex);
 	NewColor.Color.A = Alpha;
-	
+
 	NewColor.UEMaterialUId = BuildColorUId(NewColor.Color);
 	return NewColor;
 }
@@ -1587,15 +1612,15 @@ FArchiveMaterial& FTechSoftFileParser::AddMaterialAt(uint32 MaterialIndexToSave,
 	FCADMaterial& Material = NewMaterial.Material;
 
 	TUniqueTSObjFromIndex<A3DGraphMaterialData> MaterialData(MaterialIndexToSave);
-	if(MaterialData.IsValid())
+	if (MaterialData.IsValid())
 	{
 		Material.Diffuse = TechSoftUtils::GetColorAt(MaterialData->m_uiDiffuse);
 		Material.Ambient = TechSoftUtils::GetColorAt(MaterialData->m_uiAmbient);
 		Material.Specular = TechSoftUtils::GetColorAt(MaterialData->m_uiSpecular);
 		Material.Shininess = MaterialData->m_dShininess;
-		if(GraphStyleData.m_bIsTransparencyDefined)
+		if (GraphStyleData.m_bIsTransparencyDefined)
 		{
-			Material.Transparency = 1. - GraphStyleData.m_ucTransparency/255.;
+			Material.Transparency = 1. - GraphStyleData.m_ucTransparency / 255.;
 		}
 		// todo: find how to convert Emissive color into ? reflexion coef...
 		// Material.Emissive = GetColor(MaterialData->m_uiEmissive);
@@ -1622,7 +1647,7 @@ FArchiveMaterial& FTechSoftFileParser::FindOrAddMaterial(FMaterialUId MaterialIn
 		if (TextureData.IsValid())
 		{
 			return AddMaterialAt(TextureData->m_uiMaterialIndex, MaterialIndex, GraphStyleData);
-			
+
 #ifdef NOTYETDEFINE
 			TUniqueTSObj<A3DGraphTextureDefinitionData> TextureDefinitionData(TextureData->m_uiTextureDefinitionIndex);
 			if (TextureDefinitionData.IsValid())
@@ -1694,7 +1719,7 @@ void FTechSoftFileParser::ExtractGraphStyleProperties(uint32 StyleIndex, FArchiv
 	}
 }
 
-FMatrix FTechSoftFileParser::ExtractTransformation3D(const A3DMiscTransformation* CartesianTransformation, double& InOutUnit)
+void FTechSoftFileParser::ExtractTransformation3D(const A3DMiscTransformation* CartesianTransformation, FArchiveCADObject& Component)
 {
 	TUniqueTSObj<A3DMiscCartesianTransformationData> CartesianTransformationData(CartesianTransformation);
 
@@ -1706,7 +1731,7 @@ FMatrix FTechSoftFileParser::ExtractTransformation3D(const A3DMiscTransformation
 
 		FVector ZVector = XVector ^ YVector;
 
-		Origin *= InOutUnit * FImportParameters::GUnitScale;
+		Origin *= Component.Unit * FImportParameters::GUnitScale;
 
 		const A3DVector3dData& A3DScale = CartesianTransformationData->m_sScale;
 		FVector3d Scale(A3DScale.m_dX, A3DScale.m_dY, A3DScale.m_dZ);
@@ -1716,25 +1741,26 @@ FMatrix FTechSoftFileParser::ExtractTransformation3D(const A3DMiscTransformation
 		YVector *= Scale.Y;
 		ZVector *= Scale.Z;
 
-		InOutUnit *= UniformScale;
+		Component.Unit *= UniformScale;
 
-		FMatrix Matrix(XVector, YVector, ZVector, FVector::ZeroVector);
+		Component.TransformMatrix = FMatrix(XVector, YVector, ZVector, FVector::ZeroVector);
 
 		if (CartesianTransformationData->m_ucBehaviour & kA3DTransformationMirror)
 		{
-			Matrix.M[2][0] *= -1;
-			Matrix.M[2][1] *= -1;
-			Matrix.M[2][2] *= -1;
+			Component.TransformMatrix.M[2][0] *= -1;
+			Component.TransformMatrix.M[2][1] *= -1;
+			Component.TransformMatrix.M[2][2] *= -1;
 		}
 
-		Matrix.SetOrigin(Origin);
-		return Matrix;
+		Component.TransformMatrix.SetOrigin(Origin);
+	}
+	else
+	{
+		Component.TransformMatrix = FMatrix::Identity;
 	}
-
-	return FMatrix::Identity;
 }
 
-FMatrix FTechSoftFileParser::ExtractGeneralTransformation(const A3DMiscTransformation* GeneralTransformation, double& InOutUnit)
+void FTechSoftFileParser::ExtractGeneralTransformation(const A3DMiscTransformation* GeneralTransformation, FArchiveCADObject& Component)
 {
 	TUniqueTSObj<A3DMiscGeneralTransformationData> GeneralTransformationData(GeneralTransformation);
 	if (GeneralTransformationData.IsValid())
@@ -1753,32 +1779,35 @@ FMatrix FTechSoftFileParser::ExtractGeneralTransformation(const A3DMiscTransform
 		FVector3d Scale = Transform.GetScale3D();
 		if (Scale.Equals(FVector3d::OneVector, KINDA_SMALL_NUMBER))
 		{
-			const double TranslationScale = InOutUnit * FImportParameters::GUnitScale;
+			const double TranslationScale = Component.Unit * FImportParameters::GUnitScale;
 			for (Index = 0; Index < 3; ++Index, ++Index)
 			{
 				Matrix.M[3][Index] *= TranslationScale;
 			}
-			return Matrix;
+			Component.TransformMatrix = Matrix;
 		}
+		else
+		{
+			FVector3d Translation = Transform.GetTranslation();
+			Translation *= FImportParameters::GUnitScale;
 
-		FVector3d Translation = Transform.GetTranslation();
-		Translation *= FImportParameters::GUnitScale;
-
-		double UniformScale = TechSoftFileParserImpl::ExtractUniformScale(Scale);
-		InOutUnit *= UniformScale;
-
-		FQuat4d Rotation = Transform.GetRotation();
+			double UniformScale = TechSoftFileParserImpl::ExtractUniformScale(Scale);
+			Component.Unit *= UniformScale;
 
-		FTransform3d NewTransform;
-		NewTransform.SetScale3D(Scale);
-		NewTransform.SetRotation(Rotation);
+			FQuat4d Rotation = Transform.GetRotation();
 
-		FMatrix NewMatrix = NewTransform.ToMatrixWithScale();
-		NewMatrix.SetOrigin(Translation);
+			FTransform3d NewTransform;
+			NewTransform.SetScale3D(Scale);
+			NewTransform.SetRotation(Rotation);
 
-		return NewMatrix;
+			Component.TransformMatrix = NewTransform.ToMatrixWithScale();
+			Component.TransformMatrix.SetOrigin(Translation);
+		}
+	}
+	else
+	{
+		Component.TransformMatrix = FMatrix::Identity;
 	}
-	return FMatrix::Identity;
 }
 
 void FTechSoftFileParser::ExtractTransformation(const A3DMiscTransformation* Transformation3D, FArchiveCADObject& Component)
@@ -1793,11 +1822,11 @@ void FTechSoftFileParser::ExtractTransformation(const A3DMiscTransformation* Tra
 
 	if (Type == kA3DTypeMiscCartesianTransformation)
 	{
-		Component.TransformMatrix = ExtractTransformation3D(Transformation3D, Component.Unit);
+		ExtractTransformation3D(Transformation3D, Component);
 	}
 	else if (Type == kA3DTypeMiscGeneralTransformation)
 	{
-		Component.TransformMatrix = ExtractGeneralTransformation(Transformation3D, Component.Unit);
+		ExtractGeneralTransformation(Transformation3D, Component);
 	}
 }
 
@@ -1806,7 +1835,7 @@ void FTechSoftFileParser::ExtractCoordinateSystem(const A3DRiCoordinateSystem* C
 	TUniqueTSObj<A3DRiCoordinateSystemData> CoordinateSystemData(CoordinateSystem);
 	if (CoordinateSystemData.IsValid())
 	{
-		OutMetaData.TransformMatrix = ExtractTransformation3D(CoordinateSystemData->m_pTransformation, OutMetaData.Unit);
+		ExtractTransformation3D(CoordinateSystemData->m_pTransformation, OutMetaData);
 	}
 	else
 	{
diff --git a/Engine/Plugins/Enterprise/DatasmithCADImporter/Source/CADInterfaces/Private/TechSoftFileParser.h b/Engine/Plugins/Enterprise/DatasmithCADImporter/Source/CADInterfaces/Private/TechSoftFileParser.h
--- a/Engine/Plugins/Enterprise/DatasmithCADImporter/Source/CADInterfaces/Private/TechSoftFileParser.h
+++ b/Engine/Plugins/Enterprise/DatasmithCADImporter/Source/CADInterfaces/Private/TechSoftFileParser.h
@@ -139,8 +139,8 @@ private:
 	// Transform
 	void ExtractCoordinateSystem(const A3DRiCoordinateSystem* CoordinateSystem, FArchiveCADObject& Component);
 	void ExtractTransformation(const A3DMiscTransformation* Transformation3d, FArchiveCADObject& Component);
-	FMatrix ExtractGeneralTransformation(const A3DMiscTransformation* GeneralTransformation, double& InOutUnit);
-	FMatrix ExtractTransformation3D(const A3DMiscTransformation* CartesianTransformation, double& InOutUnit);
+	void ExtractGeneralTransformation(const A3DMiscTransformation* GeneralTransformation, FArchiveCADObject& Component);
+	void ExtractTransformation3D(const A3DMiscTransformation* CartesianTransformation, FArchiveCADObject& Component);
 
 #endif
 
