Improve delation of old CADCache folders

diff --git a/Engine/Plugins/Enterprise/DatasmithCADImporter/Source/DatasmithCADTranslator/Private/DatasmithCADTranslatorModule.cpp b/Engine/Plugins/Enterprise/DatasmithCADImporter/Source/DatasmithCADTranslator/Private/DatasmithCADTranslatorModule.cpp
--- a/Engine/Plugins/Enterprise/DatasmithCADImporter/Source/DatasmithCADTranslator/Private/DatasmithCADTranslatorModule.cpp
+++ b/Engine/Plugins/Enterprise/DatasmithCADImporter/Source/DatasmithCADTranslator/Private/DatasmithCADTranslatorModule.cpp
@@ -8,6 +8,7 @@
 
 #include "HAL/IConsoleManager.h"
 #include "HAL/FileManager.h"
+#include "HAL/PlatformFileManager.h"
 #include "Misc/Paths.h"
 #include "Modules/ModuleInterface.h"
 #include "Modules/ModuleManager.h"
@@ -19,13 +20,23 @@ void FDatasmithCADTranslatorModule::StartupModule()
 {
 	const int32 CacheVersion = FCADToolsModule::Get().GetCacheVersion();
 
-	// Create temporary directory which will be used by CoreTech to store tessellation data
-	for (int32 Version = 0; Version < CacheVersion; ++Version)
+	// Delete incompatible cache directory 
+	const FString Request = FPaths::ConvertRelativePathToFull(FPaths::Combine(FPaths::ProjectIntermediateDir(), TEXT("DatasmithCADCache"), TEXT("*")));
+	const FString CurrentCadCacheName = FString::FromInt(CacheVersion);
+
+	TArray<FString> CadCacheContents;
+	IFileManager::Get().FindFiles(CadCacheContents, *Request, false, true);
+
+	for (const FString& Directory : CadCacheContents)
 	{
-		FString OldCacheDir = FPaths::ConvertRelativePathToFull(FPaths::Combine(FPaths::ProjectIntermediateDir(), TEXT("DatasmithCADCache"), *FString::FromInt(Version)));
-		IFileManager::Get().DeleteDirectory(*OldCacheDir, true, true);
+		if (Directory != CurrentCadCacheName)
+		{
+			const FString OldCacheDirectory = FPaths::ConvertRelativePathToFull(FPaths::Combine(FPaths::ProjectIntermediateDir(), TEXT("DatasmithCADCache"), *Directory));
+			IFileManager::Get().DeleteDirectory(*OldCacheDirectory, true, true);
+		}
 	}
 
+	// Create root cache directory which will be used by cad library sdk to store import data
 	CacheDir = FPaths::ConvertRelativePathToFull(FPaths::Combine(FPaths::ProjectIntermediateDir(), TEXT("DatasmithCADCache"), *FString::FromInt(CacheVersion)));
 	if (!IFileManager::Get().MakeDirectory(*CacheDir, true))
 	{
