TS give unclosed external loop for sperical surface in specific case. A process is added to find this case and close the loop.
TS return wrong loop as the external loop. A process is added to find the external loop.

diff --git a/Engine/Plugins/Enterprise/DatasmithCADImporter/Source/CADInterfaces/Private/TechSoftBridge.cpp b/Engine/Plugins/Enterprise/DatasmithCADImporter/Source/CADInterfaces/Private/TechSoftBridge.cpp
--- a/Engine/Plugins/Enterprise/DatasmithCADImporter/Source/CADInterfaces/Private/TechSoftBridge.cpp
+++ b/Engine/Plugins/Enterprise/DatasmithCADImporter/Source/CADInterfaces/Private/TechSoftBridge.cpp
@@ -217,7 +217,7 @@ UE::CADKernel::FBody* FTechSoftBridge::GetBody(A3DRiBrepModel* A3DBRepModel)
 UE::CADKernel::FBody* FTechSoftBridge::AddBody(A3DRiBrepModel* A3DBRepModel, TMap<FString, FString> MetaData, const double InBodyScale)
 {
 #ifdef CADKERNEL_DEV
-	Report.BodyCount++;
+	Report->BodyCount++;
 #endif
 
 	//UE::CADKernel working unit is mm
@@ -319,7 +319,7 @@ void FTechSoftBridge::TraverseShell(const A3DTopoShell* A3DShell, TSharedRef<UE:
 	TSharedRef<UE::CADKernel::FShell> Shell = UE::CADKernel::FEntity::MakeShared<UE::CADKernel::FShell>();
 	Body->AddShell(Shell);
 #ifdef CADKERNEL_DEV
-	Report.ShellCount++;
+	Report->ShellCount++;
 #endif
 	AddMetadata(MetaData, *Shell);
 
@@ -344,7 +344,7 @@ TSharedPtr<UE::CADKernel::FCurve> FTechSoftBridge::AddCurve(const A3DCrvBase* A3
 	if (Ret == A3D_SUCCESS)
 	{
 #ifdef CADKERNEL_DEV
-		Report.CurveCount++;
+		Report->CurveCount++;
 #endif
 		switch (eType)
 		{
@@ -389,7 +389,7 @@ TSharedPtr<UE::CADKernel::FCurve> FTechSoftBridge::AddCurve(const A3DCrvBase* A3
 TSharedPtr<UE::CADKernel::FCurve> FTechSoftBridge::AddCurveLine(const A3DCrvLine* A3DCurve, const TechSoftUtils::FUVReparameterization& UVReparameterization)
 {
 #ifdef CADKERNEL_DEV
-	Report.CurveLineCount++;
+	Report->CurveLineCount++;
 #endif
 
 	if (bUseCurveAsNurbs)
@@ -412,7 +412,7 @@ TSharedPtr<UE::CADKernel::FCurve> FTechSoftBridge::AddCurveLine(const A3DCrvLine
 TSharedPtr<UE::CADKernel::FCurve> FTechSoftBridge::AddCurveCircle(const A3DCrvCircle* A3DCurve, const TechSoftUtils::FUVReparameterization& UVReparameterization)
 {
 #ifdef CADKERNEL_DEV
-	Report.CurveCircleCount++;
+	Report->CurveCircleCount++;
 #endif
 
 	if (bUseCurveAsNurbs)
@@ -434,7 +434,7 @@ TSharedPtr<UE::CADKernel::FCurve> FTechSoftBridge::AddCurveCircle(const A3DCrvCi
 TSharedPtr<UE::CADKernel::FCurve> FTechSoftBridge::AddCurveEllipse(const A3DCrvEllipse* A3DCurve, const TechSoftUtils::FUVReparameterization& UVReparameterization)
 {
 #ifdef CADKERNEL_DEV
-	Report.CurveEllipseCount++;
+	Report->CurveEllipseCount++;
 #endif
 
 	if (bUseCurveAsNurbs)
@@ -456,7 +456,7 @@ TSharedPtr<UE::CADKernel::FCurve> FTechSoftBridge::AddCurveEllipse(const A3DCrvE
 TSharedPtr<UE::CADKernel::FCurve> FTechSoftBridge::AddCurveParabola(const A3DCrvParabola* A3DCurve, const TechSoftUtils::FUVReparameterization& UVReparameterization)
 {
 #ifdef CADKERNEL_DEV
-	Report.CurveParabolaCount++;
+	Report->CurveParabolaCount++;
 #endif
 
 	if (bUseCurveAsNurbs)
@@ -478,7 +478,7 @@ TSharedPtr<UE::CADKernel::FCurve> FTechSoftBridge::AddCurveParabola(const A3DCrv
 TSharedPtr<UE::CADKernel::FCurve> FTechSoftBridge::AddCurveHyperbola(const A3DCrvHyperbola* A3DCurve, const TechSoftUtils::FUVReparameterization& UVReparameterization)
 {
 #ifdef CADKERNEL_DEV
-	Report.CurveHyperbolaCount++;
+	Report->CurveHyperbolaCount++;
 #endif
 
 	if (bUseCurveAsNurbs)
@@ -500,7 +500,7 @@ TSharedPtr<UE::CADKernel::FCurve> FTechSoftBridge::AddCurveHyperbola(const A3DCr
 TSharedPtr<UE::CADKernel::FCurve> FTechSoftBridge::AddCurveHelix(const A3DCrvHelix* A3DCurve, const TechSoftUtils::FUVReparameterization& UVReparameterization)
 {
 #ifdef CADKERNEL_DEV
-	Report.CurveHelixCount++;
+	Report->CurveHelixCount++;
 #endif
 
 	if (bUseCurveAsNurbs)
@@ -522,7 +522,7 @@ TSharedPtr<UE::CADKernel::FCurve> FTechSoftBridge::AddCurveHelix(const A3DCrvHel
 TSharedPtr<UE::CADKernel::FCurve> FTechSoftBridge::AddCurvePolyLine(const A3DCrvPolyLine* A3DCurve, const TechSoftUtils::FUVReparameterization& UVReparameterization)
 {
 #ifdef CADKERNEL_DEV
-	Report.CurvePolyLineCount++;
+	Report->CurvePolyLineCount++;
 #endif
 
 	if (bUseCurveAsNurbs)
@@ -544,7 +544,7 @@ TSharedPtr<UE::CADKernel::FCurve> FTechSoftBridge::AddCurvePolyLine(const A3DCrv
 TSharedPtr<UE::CADKernel::FCurve> FTechSoftBridge::AddCurveComposite(const A3DCrvComposite* A3DCurve, const TechSoftUtils::FUVReparameterization& UVReparameterization)
 {
 #ifdef CADKERNEL_DEV
-	Report.CurveCompositeCount++;
+	Report->CurveCompositeCount++;
 #endif
 
 	if (bUseCurveAsNurbs)
@@ -590,7 +590,7 @@ TSharedPtr<UE::CADKernel::FCurve> AddCurveNurbsFromData(A3DCrvNurbsData& A3DNurb
 TSharedPtr<UE::CADKernel::FCurve> FTechSoftBridge::AddCurveNurbs(const A3DCrvNurbs* A3DNurbs, const TechSoftUtils::FUVReparameterization& UVReparameterization)
 {
 #ifdef CADKERNEL_DEV
-	Report.CurveLineCount++;
+	Report->CurveLineCount++;
 #endif
 
 	TUniqueTSObj<A3DCrvNurbsData> CrvNurbsData(A3DNurbs);
@@ -605,7 +605,7 @@ TSharedPtr<UE::CADKernel::FCurve> FTechSoftBridge::AddCurveNurbs(const A3DCrvNur
 TSharedPtr<UE::CADKernel::FCurve> FTechSoftBridge::AddCurveAsNurbs(const A3DCrvBase* A3DCurve, const TechSoftUtils::FUVReparameterization& UVReparameterization)
 {
 #ifdef CADKERNEL_DEV
-	Report.CurveAsNurbsCount++;
+	Report->CurveAsNurbsCount++;
 #endif
 
 	TUniqueTSObj<A3DCrvNurbsData> NurbsData;
@@ -625,7 +625,7 @@ TSharedPtr<UE::CADKernel::FCurve> FTechSoftBridge::AddCurveAsNurbs(const A3DCrvB
 TSharedPtr<UE::CADKernel::FTopologicalEdge> FTechSoftBridge::AddEdge(const A3DTopoCoEdge* A3DCoedge, const TSharedRef<UE::CADKernel::FSurface>& Surface, const TechSoftUtils::FUVReparameterization& UVReparameterization, UE::CADKernel::EOrientation& OutOrientation)
 {
 #ifdef CADKERNEL_DEV
-	Report.EdgeCount++;
+	Report->EdgeCount++;
 #endif
 
 	TUniqueTSObj<A3DTopoCoEdgeData> CoEdgeData(A3DCoedge);
@@ -687,7 +687,7 @@ TSharedPtr<UE::CADKernel::FTopologicalEdge> FTechSoftBridge::AddEdge(const A3DTo
 TSharedPtr<UE::CADKernel::FTopologicalLoop> FTechSoftBridge::AddLoop(const A3DTopoLoop* A3DLoop, const TSharedRef<UE::CADKernel::FSurface>& Surface, const TechSoftUtils::FUVReparameterization& UVReparameterization, const bool bIsExternalLoop)
 {
 #ifdef CADKERNEL_DEV
-	Report.LoopCount++;
+	Report->LoopCount++;
 #endif
 
 	TArray<TSharedPtr<UE::CADKernel::FTopologicalEdge>> Edges;
@@ -697,7 +697,7 @@ TSharedPtr<UE::CADKernel::FTopologicalLoop> FTechSoftBridge::AddLoop(const A3DTo
 	if (!TopoLoopData.IsValid())
 	{
 #ifdef CADKERNEL_DEV
-		Report.DegeneratedLoopCount++;
+		Report->DegeneratedLoopCount++;
 #endif
 		return TSharedPtr<UE::CADKernel::FTopologicalLoop>();
 	}
@@ -710,7 +710,7 @@ TSharedPtr<UE::CADKernel::FTopologicalLoop> FTechSoftBridge::AddLoop(const A3DTo
 		if (!Edge.IsValid())
 		{
 #ifdef CADKERNEL_DEV
-			Report.DegeneratedEdgeCount++;
+			Report->DegeneratedEdgeCount++;
 #endif
 			continue;
 		}
@@ -722,7 +722,7 @@ TSharedPtr<UE::CADKernel::FTopologicalLoop> FTechSoftBridge::AddLoop(const A3DTo
 	if (Edges.Num() == 0)
 	{
 #ifdef CADKERNEL_DEV
-		Report.DegeneratedLoopCount++;
+		Report->DegeneratedLoopCount++;
 #endif
 		return TSharedPtr<UE::CADKernel::FTopologicalLoop>();
 	}
@@ -732,8 +732,10 @@ TSharedPtr<UE::CADKernel::FTopologicalLoop> FTechSoftBridge::AddLoop(const A3DTo
 
 void FTechSoftBridge::AddFace(const A3DTopoFace* A3DFace, UE::CADKernel::EOrientation Orientation, TSharedRef<UE::CADKernel::FShell>& Shell, uint32 ShellIndex)
 {
+	using namespace UE::CADKernel;
+
 #ifdef CADKERNEL_DEV
-	Report.FaceCount++;
+	Report->FaceCount++;
 #endif
 
 	FArchiveCADObject MetaData;
@@ -743,19 +745,19 @@ void FTechSoftBridge::AddFace(const A3DTopoFace* A3DFace, UE::CADKernel::EOrient
 	if (!TopoFaceData.IsValid())
 	{
 #ifdef CADKERNEL_DEV
-		Report.FailedFaceCount++;
+		Report->FailedFaceCount++;
 #endif
 		return;
 	}
 
 	const A3DSurfBase* A3DSurface = TopoFaceData->m_pSurface;
 	TechSoftUtils::FUVReparameterization UVReparameterization;
-	TSharedPtr<UE::CADKernel::FSurface> SurfacePtr = AddSurface(A3DSurface, UVReparameterization);
+	TSharedPtr<FSurface> SurfacePtr = AddSurface(A3DSurface, UVReparameterization);
 	if (!SurfacePtr.IsValid())
 	{
 #ifdef CADKERNEL_DEV
-		Report.DegeneratedSurfaceCount++;
-		Report.FailedFaceCount++; 
+		Report->DegeneratedSurfaceCount++;
+		Report->FailedFaceCount++; 
 #endif
 		return;
 	}
@@ -765,12 +767,12 @@ void FTechSoftBridge::AddFace(const A3DTopoFace* A3DFace, UE::CADKernel::EOrient
 		SwapOrientation(Orientation);
 	}
 
-	TSharedRef<UE::CADKernel::FSurface> Surface = SurfacePtr.ToSharedRef();
-	TSharedRef<UE::CADKernel::FTopologicalFace> Face = UE::CADKernel::FEntity::MakeShared<UE::CADKernel::FTopologicalFace>(Surface);
+	TSharedRef<FSurface> Surface = SurfacePtr.ToSharedRef();
+	TSharedRef<FTopologicalFace> Face = FEntity::MakeShared<FTopologicalFace>(Surface);
 
 	if (TopoFaceData->m_bHasTrimDomain)
 	{
-		UE::CADKernel::FSurfacicBoundary SurfaceBoundary = TechSoftUtils::GetSurfacicBoundary(TopoFaceData->m_sSurfaceDomain, UVReparameterization);
+		FSurfacicBoundary SurfaceBoundary = TechSoftUtils::GetSurfacicBoundary(TopoFaceData->m_sSurfaceDomain, UVReparameterization);
 		Surface->TrimBoundaryTo(SurfaceBoundary);
 	}
 
@@ -780,26 +782,28 @@ void FTechSoftBridge::AddFace(const A3DTopoFace* A3DFace, UE::CADKernel::EOrient
 	}
 	else
 	{
-		TArray<TSharedPtr<UE::CADKernel::FTopologicalLoop>> Loops;
+		TArray<TSharedPtr<FTopologicalLoop>> Loops;
+
+		const uint32 OuterLoopIndex = TopoFaceData->m_uiOuterLoopIndex;
 
 		for (A3DUns32 Index = 0; Index < TopoFaceData->m_uiLoopSize; ++Index)
 		{
-			const bool bIsExternalLoop = (Index == TopoFaceData->m_uiOuterLoopIndex);
-			TSharedPtr<UE::CADKernel::FTopologicalLoop> Loop = AddLoop(TopoFaceData->m_ppLoops[Index], Surface, UVReparameterization, bIsExternalLoop);
+			const bool bIsExternalLoop = (Index == OuterLoopIndex);
+			TSharedPtr<FTopologicalLoop> Loop = AddLoop(TopoFaceData->m_ppLoops[Index], Surface, UVReparameterization, bIsExternalLoop);
 			if (!Loop.IsValid())
 			{
 				continue;
 			}
 
-			TArray<UE::CADKernel::FPoint2D> LoopSampling;
+			TArray<FPoint2D> LoopSampling;
 			Loop->Get2DSampling(LoopSampling);
-			UE::CADKernel::FAABB2D Boundary;
+			FAABB2D Boundary;
 			Boundary += LoopSampling;
 			Loop->Boundary.Set(Boundary.GetMin(), Boundary.GetMax());
 
 			// Check if the loop is not composed with only degenerated edge
 			bool bDegeneratedLoop = true;
-			for (const UE::CADKernel::FOrientedEdge& Edge : Loop->GetEdges())
+			for (const FOrientedEdge& Edge : Loop->GetEdges())
 			{
 				if (!Edge.Entity->IsDegenerated())
 				{
@@ -815,10 +819,63 @@ void FTechSoftBridge::AddFace(const A3DTopoFace* A3DFace, UE::CADKernel::EOrient
 			Loops.Add(Loop);
 		}
 
+		if(Loops.Num() > 1)
+		{
+			// Check external loop
+			TSharedPtr<FTopologicalLoop> ExternalLoop;
+			FSurfacicBoundary ExternalBoundary;
+			ExternalBoundary.Init();
+			for (const TSharedPtr<FTopologicalLoop>& Loop : Loops)
+			{
+				// fast but not accurate test to check if the loop is inside an other loop based of the bounding box 
+				switch(Loop->Boundary.IsInside(ExternalBoundary, Surface->GetIsoTolerances()))
+				{
+				case ESituation::Undefined:
+				{
+					// accurate test to check if the loop is inside an other loop 
+					if (!Loop->IsInside(*ExternalLoop))
+					{
+						ExternalBoundary = Loop->Boundary;
+						ExternalLoop = Loop;
+					}
+					break;
+				}
+
+				case ESituation::Outside:
+				{
+					ExternalBoundary = Loop->Boundary;
+					ExternalLoop = Loop;
+					break;
+				}
+
+				default:
+					break;
+				}
+			}
+
+			if(!ExternalLoop->IsExternal())
+			{
+#ifdef CADKERNEL_DEV
+				Report->WrongExternalLoopCount++;
+				FMessage::Printf(Log, TEXT("ERROR: Face %d had wrong external loop\n"), Face->GetId());
+#endif
+	
+				for (TSharedPtr<FTopologicalLoop>& Loop : Loops)
+				{
+					if(Loop->IsExternal())
+					{
+						Loop->SetInternal();
+						break;
+					}
+				}
+				ExternalLoop->SetExternal();
+			}
+		}
+
 		if (Loops.Num() == 0)
 		{
 #ifdef CADKERNEL_DEV
-			Report.FailedFaceCount++;
+			Report->FailedFaceCount++;
 #endif
 			Face->SetAsDegenerated();
 			Face->SetDeleted();
@@ -829,7 +886,7 @@ void FTechSoftBridge::AddFace(const A3DTopoFace* A3DFace, UE::CADKernel::EOrient
 			int32 DoubtfulLoopOrientationCount = 0;
 			Face->AddLoops(Loops, DoubtfulLoopOrientationCount);
 #ifdef CADKERNEL_DEV
-			Report.DoubtfulLoopOrientationCount += DoubtfulLoopOrientationCount;
+			Report->DoubtfulLoopOrientationCount += DoubtfulLoopOrientationCount;
 #endif
 		}
 	}
@@ -844,7 +901,7 @@ static bool bUseSurfaceAsNurbs = true;
 TSharedPtr<UE::CADKernel::FSurface> FTechSoftBridge::AddSurface(const A3DSurfBase* A3DSurface, TechSoftUtils::FUVReparameterization& OutUVReparameterization)
 {
 #ifdef CADKERNEL_DEV
-	Report.SurfaceCount++;
+	Report->SurfaceCount++;
 #endif
 
 	FArchiveCADObject MetaData;
@@ -921,7 +978,7 @@ TSharedPtr<UE::CADKernel::FSurface> FTechSoftBridge::AddSurface(const A3DSurfBas
 TSharedPtr<UE::CADKernel::FSurface> FTechSoftBridge::AddConeSurface(const A3DSurfBase* Surface, TechSoftUtils::FUVReparameterization& OutUVReparameterization)
 {
 #ifdef CADKERNEL_DEV
-	Report.ConeSurfaceCount++;
+	Report->ConeSurfaceCount++;
 #endif
 
 	TUniqueTSObj<A3DSurfConeData> A3DConeData(Surface);
@@ -945,7 +1002,7 @@ TSharedPtr<UE::CADKernel::FSurface> FTechSoftBridge::AddConeSurface(const A3DSur
 TSharedPtr<UE::CADKernel::FSurface> FTechSoftBridge::AddCylinderSurface(const A3DSurfBase* Surface, TechSoftUtils::FUVReparameterization& OutUVReparameterization)
 {
 #ifdef CADKERNEL_DEV
-	Report.CylinderSurfaceCount++;
+	Report->CylinderSurfaceCount++;
 #endif
 
 	TUniqueTSObj<A3DSurfCylinderData> A3DCylinderData(Surface);
@@ -965,7 +1022,7 @@ TSharedPtr<UE::CADKernel::FSurface> FTechSoftBridge::AddCylinderSurface(const A3
 TSharedPtr<UE::CADKernel::FSurface> FTechSoftBridge::AddLinearTransfoSurface(const A3DSurfBase* Surface, TechSoftUtils::FUVReparameterization& OutUVReparameterization)
 {
 #ifdef CADKERNEL_DEV
-	Report.LinearTransfoSurfaceCount++;
+	Report->LinearTransfoSurfaceCount++;
 #endif
 
 	if (bUseSurfaceAsNurbs)
@@ -979,7 +1036,7 @@ TSharedPtr<UE::CADKernel::FSurface> FTechSoftBridge::AddLinearTransfoSurface(con
 TSharedPtr<UE::CADKernel::FSurface> FTechSoftBridge::AddNurbsSurface(const A3DSurfNurbs* Nurbs, TechSoftUtils::FUVReparameterization& OutUVReparameterization)
 {
 #ifdef CADKERNEL_DEV
-	Report.NurbsSurfaceCount++;
+	Report->NurbsSurfaceCount++;
 #endif
 
 	TUniqueTSObj<A3DSurfNurbsData> A3DNurbsData(Nurbs);
@@ -994,7 +1051,7 @@ TSharedPtr<UE::CADKernel::FSurface> FTechSoftBridge::AddNurbsSurface(const A3DSu
 TSharedPtr<UE::CADKernel::FSurface> FTechSoftBridge::AddOffsetSurface(const A3DSurfBase* Surface, TechSoftUtils::FUVReparameterization& OutUVReparameterization)
 {
 #ifdef CADKERNEL_DEV
-	Report.OffsetSurfaceCount++;
+	Report->OffsetSurfaceCount++;
 #endif
 
 	if (bUseSurfaceAsNurbs)
@@ -1008,7 +1065,7 @@ TSharedPtr<UE::CADKernel::FSurface> FTechSoftBridge::AddOffsetSurface(const A3DS
 TSharedPtr<UE::CADKernel::FSurface> FTechSoftBridge::AddPlaneSurface(const A3DSurfPlane* Surface, TechSoftUtils::FUVReparameterization& OutUVReparameterization)
 {
 #ifdef CADKERNEL_DEV
-	Report.PlaneSurfaceCount++;
+	Report->PlaneSurfaceCount++;
 #endif
 
 	TUniqueTSObj<A3DSurfPlaneData> A3DPlaneData(Surface);
@@ -1028,7 +1085,7 @@ TSharedPtr<UE::CADKernel::FSurface> FTechSoftBridge::AddPlaneSurface(const A3DSu
 TSharedPtr<UE::CADKernel::FSurface> FTechSoftBridge::AddRevolutionSurface(const A3DSurfRevolution* Surface, TechSoftUtils::FUVReparameterization& OutUVReparameterization)
 {
 #ifdef CADKERNEL_DEV
-	Report.RevolutionSurfaceCount++;
+	Report->RevolutionSurfaceCount++;
 #endif
 
 	if (bUseSurfaceAsNurbs)
@@ -1048,7 +1105,7 @@ TSharedPtr<UE::CADKernel::FSurface> FTechSoftBridge::AddRevolutionSurface(const
 TSharedPtr<UE::CADKernel::FSurface> FTechSoftBridge::AddRuledSurface(const A3DSurfBase* Surface, TechSoftUtils::FUVReparameterization& OutUVReparameterization)
 {
 #ifdef CADKERNEL_DEV
-	Report.RuledSurfaceCount++;
+	Report->RuledSurfaceCount++;
 #endif
 
 	if (bUseSurfaceAsNurbs)
@@ -1068,7 +1125,7 @@ TSharedPtr<UE::CADKernel::FSurface> FTechSoftBridge::AddRuledSurface(const A3DSu
 TSharedPtr<UE::CADKernel::FSurface> FTechSoftBridge::AddSphereSurface(const A3DSurfBase* Surface, TechSoftUtils::FUVReparameterization& OutUVReparameterization)
 {
 #ifdef CADKERNEL_DEV
-	Report.SphereSurfaceCount++;
+	Report->SphereSurfaceCount++;
 #endif
 
 	TUniqueTSObj<A3DSurfSphereData> A3DSphereData(Surface);
@@ -1087,7 +1144,7 @@ TSharedPtr<UE::CADKernel::FSurface> FTechSoftBridge::AddSphereSurface(const A3DS
 TSharedPtr<UE::CADKernel::FSurface> FTechSoftBridge::AddTorusSurface(const A3DSurfBase* Surface, TechSoftUtils::FUVReparameterization& OutUVReparameterization)
 {
 #ifdef CADKERNEL_DEV
-	Report.TorusSurfaceCount++;
+	Report->TorusSurfaceCount++;
 #endif
 
 	TUniqueTSObj<A3DSurfTorusData> A3DTorusData(Surface);
@@ -1105,7 +1162,7 @@ TSharedPtr<UE::CADKernel::FSurface> FTechSoftBridge::AddTorusSurface(const A3DSu
 TSharedPtr<UE::CADKernel::FSurface> FTechSoftBridge::AddBlend01Surface(const A3DSurfBase* Surface, TechSoftUtils::FUVReparameterization& OutUVReparameterization)
 {
 #ifdef CADKERNEL_DEV
-	Report.Blend01SurfaceCount++;
+	Report->Blend01SurfaceCount++;
 #endif
 
 	if (bUseSurfaceAsNurbs)
@@ -1125,7 +1182,7 @@ TSharedPtr<UE::CADKernel::FSurface> FTechSoftBridge::AddBlend01Surface(const A3D
 TSharedPtr<UE::CADKernel::FSurface> FTechSoftBridge::AddBlend02Surface(const A3DSurfBase* Surface, TechSoftUtils::FUVReparameterization& OutUVReparameterization)
 {
 #ifdef CADKERNEL_DEV
-	Report.Blend02SurfaceCount++;
+	Report->Blend02SurfaceCount++;
 #endif
 
 	if (bUseSurfaceAsNurbs)
@@ -1145,7 +1202,7 @@ TSharedPtr<UE::CADKernel::FSurface> FTechSoftBridge::AddBlend02Surface(const A3D
 TSharedPtr<UE::CADKernel::FSurface> FTechSoftBridge::AddBlend03Surface(const A3DSurfBase* Surface, TechSoftUtils::FUVReparameterization& OutUVReparameterization)
 {
 #ifdef CADKERNEL_DEV
-	Report.Blend03SurfaceCount++;
+	Report->Blend03SurfaceCount++;
 #endif
 
 	if (bUseSurfaceAsNurbs)
@@ -1159,7 +1216,7 @@ TSharedPtr<UE::CADKernel::FSurface> FTechSoftBridge::AddBlend03Surface(const A3D
 TSharedPtr<UE::CADKernel::FSurface> FTechSoftBridge::AddCylindricalSurface(const A3DSurfBase* Surface, TechSoftUtils::FUVReparameterization& OutUVReparameterization)
 {
 #ifdef CADKERNEL_DEV
-	Report.CylindricalSurfaceCount++;
+	Report->CylindricalSurfaceCount++;
 #endif
 
 	if (bUseSurfaceAsNurbs)
@@ -1179,7 +1236,7 @@ TSharedPtr<UE::CADKernel::FSurface> FTechSoftBridge::AddCylindricalSurface(const
 TSharedPtr<UE::CADKernel::FSurface> FTechSoftBridge::AddPipeSurface(const A3DSurfBase* Surface, TechSoftUtils::FUVReparameterization& OutUVReparameterization)
 {
 #ifdef CADKERNEL_DEV
-	Report.PipeSurfaceCount++;
+	Report->PipeSurfaceCount++;
 #endif
 
 	if (bUseSurfaceAsNurbs)
@@ -1199,7 +1256,7 @@ TSharedPtr<UE::CADKernel::FSurface> FTechSoftBridge::AddPipeSurface(const A3DSur
 TSharedPtr<UE::CADKernel::FSurface> FTechSoftBridge::AddExtrusionSurface(const A3DSurfBase* Surface, TechSoftUtils::FUVReparameterization& OutUVReparameterization)
 {
 #ifdef CADKERNEL_DEV
-	Report.ExtrusionSurfaceCount++;
+	Report->ExtrusionSurfaceCount++;
 #endif
 
 	if (bUseSurfaceAsNurbs)
@@ -1219,7 +1276,7 @@ TSharedPtr<UE::CADKernel::FSurface> FTechSoftBridge::AddExtrusionSurface(const A
 TSharedPtr<UE::CADKernel::FSurface> FTechSoftBridge::AddSurfaceFromCurves(const A3DSurfBase* Surface, TechSoftUtils::FUVReparameterization& OutUVReparameterization)
 {
 #ifdef CADKERNEL_DEV
-	Report.SurfaceFromCurvesCount++;
+	Report->SurfaceFromCurvesCount++;
 #endif
 
 	if (bUseSurfaceAsNurbs)
@@ -1239,7 +1296,7 @@ TSharedPtr<UE::CADKernel::FSurface> FTechSoftBridge::AddSurfaceFromCurves(const
 TSharedPtr<UE::CADKernel::FSurface> FTechSoftBridge::AddTransformSurface(const A3DSurfBase* Surface, TechSoftUtils::FUVReparameterization& OutUVReparameterization)
 {
 #ifdef CADKERNEL_DEV
-	Report.TransformSurfaceCount++;
+	Report->TransformSurfaceCount++;
 #endif
 
 	if (bUseSurfaceAsNurbs)
@@ -1293,7 +1350,7 @@ TSharedPtr<UE::CADKernel::FSurface> FTechSoftBridge::AddSurfaceNurbs(const A3DSu
 TSharedPtr<UE::CADKernel::FSurface> FTechSoftBridge::AddSurfaceAsNurbs(const A3DSurfBase* Surface, TechSoftUtils::FUVReparameterization& OutUVReparameterization)
 {
 #ifdef CADKERNEL_DEV
-	Report.SurfaceAsNurbsCount++;
+	Report->SurfaceAsNurbsCount++;
 #endif
 
 	TUniqueTSObj<A3DSurfNurbsData> A3DNurbsData;
diff --git a/Engine/Plugins/Enterprise/DatasmithCADImporter/Source/CADInterfaces/Private/TechSoftBridge.h b/Engine/Plugins/Enterprise/DatasmithCADImporter/Source/CADInterfaces/Private/TechSoftBridge.h
--- a/Engine/Plugins/Enterprise/DatasmithCADImporter/Source/CADInterfaces/Private/TechSoftBridge.h
+++ b/Engine/Plugins/Enterprise/DatasmithCADImporter/Source/CADInterfaces/Private/TechSoftBridge.h
@@ -169,6 +169,10 @@ private:
 
 	UE::CADKernel::FSession& Session;
 	UE::CADKernel::FModel& Model;
+#ifdef CADKERNEL_DEV
+	UE::CADKernel::FCADFileReport* Report;
+#endif
+
 
 	const double GeometricTolerance;
 	const double SquareGeometricTolerance;
@@ -186,6 +190,13 @@ public:
 	UE::CADKernel::FBody* GetBody(A3DRiBrepModel* A3DBRepModel);
 	const A3DRiBrepModel* GetA3DBody(UE::CADKernel::FBody* BRepModel);
 
+#ifdef CADKERNEL_DEV
+	void SetReport(UE::CADKernel::FCADFileReport& InReport)
+	{
+		Report = &InReport;
+	}
+#endif
+
 private:
 
 	void TraverseBrepData(const A3DTopoBrepData* A3DBrepData, TSharedRef<UE::CADKernel::FBody>& Body);
diff --git a/Engine/Source/Runtime/Datasmith/CADKernel/Private/CADKernel/Core/Database.cpp b/Engine/Source/Runtime/Datasmith/CADKernel/Private/CADKernel/Core/Database.cpp
--- a/Engine/Source/Runtime/Datasmith/CADKernel/Private/CADKernel/Core/Database.cpp
+++ b/Engine/Source/Runtime/Datasmith/CADKernel/Private/CADKernel/Core/Database.cpp
@@ -274,9 +274,6 @@ void FDatabase::SerializeSelection(FCADKernelArchive& Ar, const TArray<FIdent>&
 		case EEntity::Body:
 			TmpModel->Add(StaticCastSharedRef<FBody>(Entity.ToSharedRef()));
 			break;
-		case EEntity::TopologicalFace:
-			TmpModel->Add(StaticCastSharedRef<FTopologicalFace>(Entity.ToSharedRef()));
-			break;
 		default:
 			AddEntityToSave(EntityId);
 			break;
@@ -484,7 +481,6 @@ void FDatabase::CleanArchiveEntities()
 		}
 	}
 
-	Model->Append(IndependantFaces);
 	Model->Append(NewBodies);
 }
 
@@ -629,10 +625,6 @@ void FDatabase::ExpandSelection(TSharedPtr<FEntity> Entity, const TSet<EEntity>&
 			ExpandSelection(Body, Filter, Selection);
 		}
 
-		for (const TSharedPtr<FTopologicalFace>& Domain : InModel->GetFaces())
-		{
-			ExpandSelection(Domain, Filter, Selection);
-		}
 		break;
 	}
 
diff --git a/Engine/Source/Runtime/Datasmith/CADKernel/Private/CADKernel/Geo/Curves/BezierCurve.cpp b/Engine/Source/Runtime/Datasmith/CADKernel/Private/CADKernel/Geo/Curves/BezierCurve.cpp
--- a/Engine/Source/Runtime/Datasmith/CADKernel/Private/CADKernel/Geo/Curves/BezierCurve.cpp
+++ b/Engine/Source/Runtime/Datasmith/CADKernel/Private/CADKernel/Geo/Curves/BezierCurve.cpp
@@ -54,6 +54,14 @@ TSharedPtr<FEntityGeom> FBezierCurve::ApplyMatrix(const FMatrixH& InMatrix) cons
 	return FEntity::MakeShared<FBezierCurve>(TransformedPoles);
 }
 
+void FBezierCurve::Offset(const FPoint& OffsetDirection)
+{
+	for (FPoint& Pole : Poles)
+	{
+		Pole += OffsetDirection;
+	}
+}
+
 #ifdef CADKERNEL_DEV
 FInfoEntity& FBezierCurve::GetInfo(FInfoEntity& Info) const
 {
diff --git a/Engine/Source/Runtime/Datasmith/CADKernel/Private/CADKernel/Geo/Curves/BezierCurve.h b/Engine/Source/Runtime/Datasmith/CADKernel/Private/CADKernel/Geo/Curves/BezierCurve.h
--- a/Engine/Source/Runtime/Datasmith/CADKernel/Private/CADKernel/Geo/Curves/BezierCurve.h
+++ b/Engine/Source/Runtime/Datasmith/CADKernel/Private/CADKernel/Geo/Curves/BezierCurve.h
@@ -47,11 +47,13 @@ public:
 	}
 
 	virtual TSharedPtr<FEntityGeom> ApplyMatrix(const FMatrixH& InMatrix) const override;
+	virtual void Offset(const FPoint& OffsetDirection) override;
 
 	virtual void EvaluatePoint(double Coordinate, FCurvePoint& OutPoint, int32 DerivativeOrder = 0) const override;
 
 	virtual void ExtendTo(const FPoint& Point) override;
 
+
 };
 
 } // namespace UE::CADKernel
diff --git a/Engine/Source/Runtime/Datasmith/CADKernel/Private/CADKernel/Geo/Curves/BoundedCurve.cpp b/Engine/Source/Runtime/Datasmith/CADKernel/Private/CADKernel/Geo/Curves/BoundedCurve.cpp
--- a/Engine/Source/Runtime/Datasmith/CADKernel/Private/CADKernel/Geo/Curves/BoundedCurve.cpp
+++ b/Engine/Source/Runtime/Datasmith/CADKernel/Private/CADKernel/Geo/Curves/BoundedCurve.cpp
@@ -59,6 +59,11 @@ TSharedPtr<FEntityGeom> FBoundedCurve::ApplyMatrix(const FMatrixH& InMatrix) con
 	return FEntity::MakeShared<FBoundedCurve>(TransformedCurve.ToSharedRef(), Boundary, Dimension);
 }
 
+void FBoundedCurve::Offset(const FPoint& OffsetDirection)
+{
+	Curve->Offset(OffsetDirection);
+}
+
 #ifdef CADKERNEL_DEV
 FInfoEntity& FBoundedCurve::GetInfo(FInfoEntity& Info) const
 {
diff --git a/Engine/Source/Runtime/Datasmith/CADKernel/Private/CADKernel/Geo/Curves/BoundedCurve.h b/Engine/Source/Runtime/Datasmith/CADKernel/Private/CADKernel/Geo/Curves/BoundedCurve.h
--- a/Engine/Source/Runtime/Datasmith/CADKernel/Private/CADKernel/Geo/Curves/BoundedCurve.h
+++ b/Engine/Source/Runtime/Datasmith/CADKernel/Private/CADKernel/Geo/Curves/BoundedCurve.h
@@ -61,6 +61,7 @@ public:
 	}
 
 	virtual TSharedPtr<FEntityGeom> ApplyMatrix(const FMatrixH& InMatrix) const override;
+	virtual void Offset(const FPoint& OffsetDirection) override;
 
 	virtual void EvaluatePoint(double Coordinate, FCurvePoint& OutPoint, int32 DerivativeOrder = 0) const override;
 	virtual void Evaluate2DPoint(double Coordinate, FCurvePoint2D& OutPoint, int32 DerivativeOrder = 0) const override;
diff --git a/Engine/Source/Runtime/Datasmith/CADKernel/Private/CADKernel/Geo/Curves/CompositeCurve.cpp b/Engine/Source/Runtime/Datasmith/CADKernel/Private/CADKernel/Geo/Curves/CompositeCurve.cpp
--- a/Engine/Source/Runtime/Datasmith/CADKernel/Private/CADKernel/Geo/Curves/CompositeCurve.cpp
+++ b/Engine/Source/Runtime/Datasmith/CADKernel/Private/CADKernel/Geo/Curves/CompositeCurve.cpp
@@ -184,6 +184,14 @@ TSharedPtr<FEntityGeom> FCompositeCurve::ApplyMatrix(const FMatrixH& InMatrix) c
 	return FEntity::MakeShared<FCompositeCurve>(TransformedCurves);
 }
 
+void FCompositeCurve::Offset(const FPoint& OffsetDirection)
+{
+	for (FOrientedCurve& Curve : Curves)
+	{
+		Curve.Entity->Offset(OffsetDirection);
+	}
+}
+
 #ifdef CADKERNEL_DEV
 FInfoEntity& FCompositeCurve::GetInfo(FInfoEntity& Info) const
 {
diff --git a/Engine/Source/Runtime/Datasmith/CADKernel/Private/CADKernel/Geo/Curves/CompositeCurve.h b/Engine/Source/Runtime/Datasmith/CADKernel/Private/CADKernel/Geo/Curves/CompositeCurve.h
--- a/Engine/Source/Runtime/Datasmith/CADKernel/Private/CADKernel/Geo/Curves/CompositeCurve.h
+++ b/Engine/Source/Runtime/Datasmith/CADKernel/Private/CADKernel/Geo/Curves/CompositeCurve.h
@@ -74,6 +74,7 @@ public:
 	}
 
 	virtual TSharedPtr<FEntityGeom> ApplyMatrix(const FMatrixH& InMatrix) const override;
+	virtual void Offset(const FPoint& OffsetDirection) override;
 
 	virtual TSharedPtr<FCurve> GetCurve(int32 Index) const
 	{
diff --git a/Engine/Source/Runtime/Datasmith/CADKernel/Private/CADKernel/Geo/Curves/EllipseCurve.cpp b/Engine/Source/Runtime/Datasmith/CADKernel/Private/CADKernel/Geo/Curves/EllipseCurve.cpp
--- a/Engine/Source/Runtime/Datasmith/CADKernel/Private/CADKernel/Geo/Curves/EllipseCurve.cpp
+++ b/Engine/Source/Runtime/Datasmith/CADKernel/Private/CADKernel/Geo/Curves/EllipseCurve.cpp
@@ -10,6 +10,12 @@ TSharedPtr<FEntityGeom> FEllipseCurve::ApplyMatrix(const FMatrixH& InMatrix) con
 	return FEntity::MakeShared<FEllipseCurve>(NewMatrix, RadiusU, RadiusV, Boundary);
 }
 
+void FEllipseCurve::Offset(const FPoint& OffsetDirection)
+{
+	FMatrixH Offset = FMatrixH::MakeTranslationMatrix(OffsetDirection);
+	Matrix *= Offset;
+}
+
 #ifdef CADKERNEL_DEV
 FInfoEntity& FEllipseCurve::GetInfo(FInfoEntity& Info) const
 {
diff --git a/Engine/Source/Runtime/Datasmith/CADKernel/Private/CADKernel/Geo/Curves/EllipseCurve.h b/Engine/Source/Runtime/Datasmith/CADKernel/Private/CADKernel/Geo/Curves/EllipseCurve.h
--- a/Engine/Source/Runtime/Datasmith/CADKernel/Private/CADKernel/Geo/Curves/EllipseCurve.h
+++ b/Engine/Source/Runtime/Datasmith/CADKernel/Private/CADKernel/Geo/Curves/EllipseCurve.h
@@ -16,14 +16,12 @@ protected:
 	FMatrixH Matrix;
 	double   RadiusU;
 	double   RadiusV;
-	FMatrixH InverseMatrix;
 
 	FEllipseCurve(const FMatrixH& InMatrix, double InRadiusU, double InRadiusV, int8 InDimension = 3)
 		: FCurve(FLinearBoundary(0, PI * 2.), InDimension)
 		, Matrix(InMatrix)
 		, RadiusU(InRadiusU)
 		, RadiusV(InRadiusV)
-		, InverseMatrix(InMatrix.GetInverse())
 	{
 	}
 
@@ -32,7 +30,6 @@ protected:
 		, Matrix(InMatrix)
 		, RadiusU(InRadiusU)
 		, RadiusV(InRadiusV)
-		, InverseMatrix(InMatrix.GetInverse())
 	{
 	}
 
@@ -46,7 +43,6 @@ public:
 		Ar << Matrix;
 		Ar << RadiusU;
 		Ar << RadiusV;
-		Ar << InverseMatrix;
 	}
 
 #ifdef CADKERNEL_DEV
@@ -69,6 +65,7 @@ public:
 	}
 
 	virtual TSharedPtr<FEntityGeom> ApplyMatrix(const FMatrixH& InMatrix) const override;
+	virtual void Offset(const FPoint& OffsetDirection) override;
 
 	virtual void EvaluatePoint(double Coordinate, FCurvePoint& OutPoint, int32 DerivativeOrder = 0) const override
 	{
diff --git a/Engine/Source/Runtime/Datasmith/CADKernel/Private/CADKernel/Geo/Curves/HyperbolaCurve.cpp b/Engine/Source/Runtime/Datasmith/CADKernel/Private/CADKernel/Geo/Curves/HyperbolaCurve.cpp
--- a/Engine/Source/Runtime/Datasmith/CADKernel/Private/CADKernel/Geo/Curves/HyperbolaCurve.cpp
+++ b/Engine/Source/Runtime/Datasmith/CADKernel/Private/CADKernel/Geo/Curves/HyperbolaCurve.cpp
@@ -10,6 +10,12 @@ TSharedPtr<FEntityGeom> FHyperbolaCurve::ApplyMatrix(const FMatrixH& InMatrix) c
 	return FEntity::MakeShared<FHyperbolaCurve>(NewMatrix, SemiMajorAxis, SemiImaginaryAxis, Boundary);
 }
 
+void FHyperbolaCurve::Offset(const FPoint& OffsetDirection)
+{
+	FMatrixH Offset = FMatrixH::MakeTranslationMatrix(OffsetDirection);
+	Matrix *= Offset;
+}
+
 #ifdef CADKERNEL_DEV
 FInfoEntity& FHyperbolaCurve::GetInfo(FInfoEntity& Info) const
 {
diff --git a/Engine/Source/Runtime/Datasmith/CADKernel/Private/CADKernel/Geo/Curves/HyperbolaCurve.h b/Engine/Source/Runtime/Datasmith/CADKernel/Private/CADKernel/Geo/Curves/HyperbolaCurve.h
--- a/Engine/Source/Runtime/Datasmith/CADKernel/Private/CADKernel/Geo/Curves/HyperbolaCurve.h
+++ b/Engine/Source/Runtime/Datasmith/CADKernel/Private/CADKernel/Geo/Curves/HyperbolaCurve.h
@@ -52,6 +52,7 @@ public:
 	}
 
 	virtual TSharedPtr<FEntityGeom> ApplyMatrix(const FMatrixH& InMatrix) const override;
+	virtual void Offset(const FPoint& OffsetDirection) override;
 
 	virtual void EvaluatePoint(double Coordinate, FCurvePoint& OutPoint, int32 DerivativeOrder = 0) const override
 	{
diff --git a/Engine/Source/Runtime/Datasmith/CADKernel/Private/CADKernel/Geo/Curves/NURBSCurve.cpp b/Engine/Source/Runtime/Datasmith/CADKernel/Private/CADKernel/Geo/Curves/NURBSCurve.cpp
--- a/Engine/Source/Runtime/Datasmith/CADKernel/Private/CADKernel/Geo/Curves/NURBSCurve.cpp
+++ b/Engine/Source/Runtime/Datasmith/CADKernel/Private/CADKernel/Geo/Curves/NURBSCurve.cpp
@@ -51,7 +51,7 @@ TSharedPtr<FEntityGeom> FNURBSCurve::ApplyMatrix(const FMatrixH& InMatrix) const
 	TArray<FPoint> TransformedPoles;
 
 	TransformedPoles.Reserve(Poles.Num());
-	for (FPoint Pole : Poles) 
+	for (const FPoint& Pole : Poles) 
 	{
 		TransformedPoles.Emplace(InMatrix.Multiply(Pole));
 	}
@@ -59,6 +59,15 @@ TSharedPtr<FEntityGeom> FNURBSCurve::ApplyMatrix(const FMatrixH& InMatrix) const
 	return FEntity::MakeShared<FNURBSCurve>(Degree, NodalVector, TransformedPoles, Weights, Dimension);
 }
 
+void FNURBSCurve::Offset(const FPoint& OffsetDirection)
+{
+	for (FPoint& Pole : Poles)
+	{
+		Pole += OffsetDirection;
+	}
+	Finalize();
+}
+
 #ifdef CADKERNEL_DEV
 FInfoEntity& FNURBSCurve::GetInfo(FInfoEntity& Info) const
 {
diff --git a/Engine/Source/Runtime/Datasmith/CADKernel/Private/CADKernel/Geo/Curves/NURBSCurve.h b/Engine/Source/Runtime/Datasmith/CADKernel/Private/CADKernel/Geo/Curves/NURBSCurve.h
--- a/Engine/Source/Runtime/Datasmith/CADKernel/Private/CADKernel/Geo/Curves/NURBSCurve.h
+++ b/Engine/Source/Runtime/Datasmith/CADKernel/Private/CADKernel/Geo/Curves/NURBSCurve.h
@@ -107,6 +107,7 @@ public:
 	}
 
 	virtual TSharedPtr<FEntityGeom> ApplyMatrix(const FMatrixH& InMatrix) const override;
+	virtual void Offset(const FPoint& OffsetDirection) override;
 
 	virtual void EvaluatePoint(double Coordinate, FCurvePoint& OutPoint, int32 DerivativeOrder = 0) const override
 	{
diff --git a/Engine/Source/Runtime/Datasmith/CADKernel/Private/CADKernel/Geo/Curves/ParabolaCurve.cpp b/Engine/Source/Runtime/Datasmith/CADKernel/Private/CADKernel/Geo/Curves/ParabolaCurve.cpp
--- a/Engine/Source/Runtime/Datasmith/CADKernel/Private/CADKernel/Geo/Curves/ParabolaCurve.cpp
+++ b/Engine/Source/Runtime/Datasmith/CADKernel/Private/CADKernel/Geo/Curves/ParabolaCurve.cpp
@@ -13,6 +13,13 @@ TSharedPtr<FEntityGeom> FParabolaCurve::ApplyMatrix(const FMatrixH& InMatrix) co
 	return FEntity::MakeShared<FParabolaCurve>(NewMatrix, FocalDistance, Boundary, Dimension);
 }
 
+void FParabolaCurve::Offset(const FPoint& OffsetDirection)
+{
+	FMatrixH Offset = FMatrixH::MakeTranslationMatrix(OffsetDirection);
+	Matrix *= Offset;
+}
+
+
 #ifdef CADKERNEL_DEV
 FInfoEntity& FParabolaCurve::GetInfo(FInfoEntity& Info) const
 {
diff --git a/Engine/Source/Runtime/Datasmith/CADKernel/Private/CADKernel/Geo/Curves/ParabolaCurve.h b/Engine/Source/Runtime/Datasmith/CADKernel/Private/CADKernel/Geo/Curves/ParabolaCurve.h
--- a/Engine/Source/Runtime/Datasmith/CADKernel/Private/CADKernel/Geo/Curves/ParabolaCurve.h
+++ b/Engine/Source/Runtime/Datasmith/CADKernel/Private/CADKernel/Geo/Curves/ParabolaCurve.h
@@ -43,6 +43,7 @@ public:
 	}
 
 	virtual TSharedPtr<FEntityGeom> ApplyMatrix(const FMatrixH& InMatrix) const override;
+	virtual void Offset(const FPoint& OffsetDirection) override;
 
 	virtual void ExtendTo(const FPoint& Point) override
 	{
diff --git a/Engine/Source/Runtime/Datasmith/CADKernel/Private/CADKernel/Geo/Curves/PolylineCurve.h b/Engine/Source/Runtime/Datasmith/CADKernel/Private/CADKernel/Geo/Curves/PolylineCurve.h
--- a/Engine/Source/Runtime/Datasmith/CADKernel/Private/CADKernel/Geo/Curves/PolylineCurve.h
+++ b/Engine/Source/Runtime/Datasmith/CADKernel/Private/CADKernel/Geo/Curves/PolylineCurve.h
@@ -117,6 +117,14 @@ public:
 		return FEntity::MakeShared<PolylineType>(NewPoints, this->Coordinates);
 	}
 
+	virtual void Offset(const FPoint& OffsetDirection) override
+	{
+		for (PointType& Pole : this->Points)
+		{
+			Pole += OffsetDirection;
+		}
+	}
+
 	virtual void ExtendTo(const FPoint& DesiredPoint) override
 	{
 		PolylineTools::ExtendTo(this->Points, (PointType)DesiredPoint);
diff --git a/Engine/Source/Runtime/Datasmith/CADKernel/Private/CADKernel/Geo/Curves/RestrictionCurve.cpp b/Engine/Source/Runtime/Datasmith/CADKernel/Private/CADKernel/Geo/Curves/RestrictionCurve.cpp
--- a/Engine/Source/Runtime/Datasmith/CADKernel/Private/CADKernel/Geo/Curves/RestrictionCurve.cpp
+++ b/Engine/Source/Runtime/Datasmith/CADKernel/Private/CADKernel/Geo/Curves/RestrictionCurve.cpp
@@ -11,6 +11,13 @@ void FRestrictionCurve::ExtendTo(const FPoint2D& Point)
 	EvaluateSurfacicPolyline(Polyline);
 }
 
+void FRestrictionCurve::Offset2D(const FPoint2D& Offset)
+{
+	Curve2D->Offset(Offset);
+	EvaluateSurfacicPolyline(Polyline);
+}
+
+
 #ifdef CADKERNEL_DEV
 FInfoEntity& FRestrictionCurve::GetInfo(FInfoEntity& Info) const
 {
diff --git a/Engine/Source/Runtime/Datasmith/CADKernel/Private/CADKernel/Geo/Curves/SegmentCurve.cpp b/Engine/Source/Runtime/Datasmith/CADKernel/Private/CADKernel/Geo/Curves/SegmentCurve.cpp
--- a/Engine/Source/Runtime/Datasmith/CADKernel/Private/CADKernel/Geo/Curves/SegmentCurve.cpp
+++ b/Engine/Source/Runtime/Datasmith/CADKernel/Private/CADKernel/Geo/Curves/SegmentCurve.cpp
@@ -13,6 +13,12 @@ TSharedPtr<FEntityGeom> FSegmentCurve::ApplyMatrix(const FMatrixH& InMatrix) con
 	return FEntity::MakeShared<FSegmentCurve>(TransformedStartPoint, TransformedEndPoint, Dimension);
 }
 
+void FSegmentCurve::Offset(const FPoint& OffsetDirection)
+{
+	StartPoint += OffsetDirection;
+	EndPoint += OffsetDirection;
+}
+
 #ifdef CADKERNEL_DEV
 FInfoEntity& FSegmentCurve::GetInfo(FInfoEntity& Info) const
 {
diff --git a/Engine/Source/Runtime/Datasmith/CADKernel/Private/CADKernel/Geo/Curves/SegmentCurve.h b/Engine/Source/Runtime/Datasmith/CADKernel/Private/CADKernel/Geo/Curves/SegmentCurve.h
--- a/Engine/Source/Runtime/Datasmith/CADKernel/Private/CADKernel/Geo/Curves/SegmentCurve.h
+++ b/Engine/Source/Runtime/Datasmith/CADKernel/Private/CADKernel/Geo/Curves/SegmentCurve.h
@@ -43,6 +43,7 @@ public:
 	}
 
 	virtual TSharedPtr<FEntityGeom> ApplyMatrix(const FMatrixH& InMatrix) const override;
+	virtual void Offset(const FPoint& OffsetDirection) override;
 
 	virtual void EvaluatePoint(double Coordinate, FCurvePoint& OutPoint, int32 DerivativeOrder = 0) const override
 	{
diff --git a/Engine/Source/Runtime/Datasmith/CADKernel/Private/CADKernel/Topo/Body.cpp b/Engine/Source/Runtime/Datasmith/CADKernel/Private/CADKernel/Topo/Body.cpp
--- a/Engine/Source/Runtime/Datasmith/CADKernel/Private/CADKernel/Topo/Body.cpp
+++ b/Engine/Source/Runtime/Datasmith/CADKernel/Private/CADKernel/Topo/Body.cpp
@@ -35,7 +35,7 @@ void FBody::RemoveEmptyShell(FModel& Model)
 	if (NewShells.IsEmpty())
 	{
 		Delete();
-		Model.RemoveBody(this);
+		Model.Remove(this);
 	}
 	else
 	{
@@ -43,6 +43,17 @@ void FBody::RemoveEmptyShell(FModel& Model)
 	}
 }
 
+void FBody::Remove(const FTopologicalShapeEntity* ShellToRemove)
+{
+	if (!ShellToRemove)
+	{
+		return;
+	}
+
+	int32 Index = Shells.IndexOfByPredicate([&](const TSharedPtr<FShell>& Shell) { return (Shell.Get() == ShellToRemove); });
+	Shells.RemoveAt(Index);
+}
+
 #ifdef CADKERNEL_DEV
 FInfoEntity& FBody::GetInfo(FInfoEntity& Info) const
 {
diff --git a/Engine/Source/Runtime/Datasmith/CADKernel/Private/CADKernel/Topo/FaceAnalyzer.cpp b/Engine/Source/Runtime/Datasmith/CADKernel/Private/CADKernel/Topo/FaceAnalyzer.cpp
new file mode 100644
--- /dev/null
+++ b/Engine/Source/Runtime/Datasmith/CADKernel/Private/CADKernel/Topo/FaceAnalyzer.cpp
@@ -0,0 +1,285 @@
+// Copyright Epic Games, Inc. All Rights Reserved.
+
+#include "CADKernel/Topo/FaceAnalyzer.h"
+
+#include "CADKernel/Math/Point.h"
+#include "CADKernel/Topo/TopologicalEdge.h"
+#include "CADKernel/Topo/TopologicalFace.h"
+#include "CADKernel/Topo/TopologicalLoop.h"
+#include "CADKernel/UI/Display.h"
+
+namespace UE::CADKernel
+{
+
+#ifdef CADKERNEL_DEV
+FIdent Topo::FEdgeSegment::LastId = 0;
+#endif
+
+void FFaceAnalyzer::FindClosedSegments(Topo::FThinFaceContext& Context)
+{
+#ifdef DEBUG_FIND_THIN_FACE
+	F3DDebugSession _(TEXT("Closed Segment"));
+#endif
+
+	// For each segment, the nearest segment is search
+	// If segment is from an inner loop, the nearest could not be from the same loop
+	for (Topo::FEdgeSegment* Segment : Context.LoopSegments)
+	{
+		const FPoint& Point = Segment->GetMiddle();
+		const FTopologicalEdge* Edge = Segment->GetEdge();
+
+		//FPoint ClosedPoint;
+		Topo::FEdgeSegment* ClosedSegment = nullptr;
+
+		double MinSquareDistance = HUGE_VALUE;
+
+		for (Topo::FEdgeSegment* Candidate : Context.LoopSegments)
+		{
+			if (Candidate->GetEdge() == Edge)
+			{
+				continue;
+			}
+
+			const FPoint& CenterPointCandidate = Candidate->GetMiddle();
+			const double SquareDistanceToCenter = CenterPointCandidate.SquareDistance(Point);
+
+			// if the distance of Point with the middle of candidate segment is biggest than MaxSpace, then the projection of the point cannot be smaller than the tolerance, 
+			if (SquareDistanceToCenter > MinSquareDistance)
+			{
+				continue;
+			}
+
+			const FPoint& FirstPointCandidate = Candidate->GetExtemity(ELimit::Start);
+			const FPoint& SecondPointCandidate = Candidate->GetExtemity(ELimit::End);
+
+			double Coordinate;
+			FPoint Projection = ProjectPointOnSegment(Point, FirstPointCandidate, SecondPointCandidate, Coordinate, true);
+
+			double SquareDistance = Point.SquareDistance(Projection);
+			if (SquareDistance > MinSquareDistance)
+			{
+				continue;
+			}
+
+			// check the angle between segments. As they are opposite, the cosAngle as to be close to -1
+			double CosAngle = Segment->ComputeCosAngleOf(Candidate);
+			if (CosAngle > -0.5) // Angle < 3Pi/4 (135 deg)
+			{
+				continue;
+			}
+
+			MinSquareDistance = SquareDistance;
+			ClosedSegment = Candidate;
+		}
+
+		if (ClosedSegment)
+		{
+			Segment->SetClosedSegment(ClosedSegment, MinSquareDistance);
+
+#ifdef DEBUG_FIND_THIN_FACE
+			//F3DDebugSession _(TEXT("Segment"));
+			DisplaySegment(Segment->GetExtemity(ELimit::Start), Segment->GetExtemity(ELimit::End), 0, EVisuProperty::BlueCurve);
+			DisplaySegment(ClosedSegment->GetExtemity(ELimit::Start), ClosedSegment->GetExtemity(ELimit::End), 0, EVisuProperty::RedCurve);
+			//Wait();
+#endif
+		}
+	}
+}
+
+
+void FFaceAnalyzer::Analyze(Topo::FThinFaceContext& Context)
+{
+	const FTopologicalEdge* Edge = Context.LoopSegments[0]->GetEdge();
+	double MaxSquareDistance = 0;
+	double MedSquareDistance = 0;
+	double EdgeLength = 0;
+
+	TFunction<void()> SetEdgeMaxGap = [&]()
+	{
+		if (EdgeLength * 1.2 > Edge->Length())
+		{
+			MedSquareDistance /= EdgeLength;
+		}
+		else
+		{
+			MedSquareDistance = -1;
+			MaxSquareDistance = -1;
+		}
+
+		if (MaxSquareDistance < 2*FMath::Square(Edge->Length()))
+		{
+			Context.EdgeSquareDistance.Add(MedSquareDistance);
+			Context.EdgeMaxSquareDistance.Add(MaxSquareDistance);
+
+			Context.MaxSquareDistance = FMath::Max(Context.MaxSquareDistance, MaxSquareDistance);
+			Context.SideEdgeLength += Edge->Length();
+		}
+		else
+		{
+			Context.EdgeSquareDistance.Add(-1.);
+			Context.EdgeMaxSquareDistance.Add(-1.);
+		}
+		Context.ExternalLoopLength += Edge->Length();
+	};
+
+	Context.EdgeSquareDistance.Reserve(Context.Loop.EdgeCount());
+	Context.EdgeMaxSquareDistance.Reserve(Context.Loop.EdgeCount());
+
+	for (Topo::FEdgeSegment* Segment : Context.LoopSegments)
+	{
+		if (Segment->GetEdge() != Edge)
+		{
+			SetEdgeMaxGap();
+
+			Edge = Segment->GetEdge();
+			MaxSquareDistance = 0;
+			MedSquareDistance = 0;
+			EdgeLength = 0;
+		}
+
+		if (Segment->GetClosedSquareDistance() > 0)
+		{
+			MedSquareDistance += Segment->GetLength() * Segment->GetClosedSquareDistance();
+			EdgeLength += Segment->GetLength();
+			if (MaxSquareDistance < Segment->GetClosedSquareDistance())
+			{
+				MaxSquareDistance = Segment->GetClosedSquareDistance();
+			}
+		}
+	}
+
+	SetEdgeMaxGap();
+}
+
+
+void FFaceAnalyzer::BuildLoopSegments(Topo::FThinFaceContext& Context)
+{
+	double Length = 0;
+	double SegmentLength = Tolerance / 5.;
+
+	const TArray<FOrientedEdge>& Edges = Context.Loop.GetEdges();
+
+	Context.LoopSegments.Empty();
+
+	for (const FOrientedEdge& OrientedEdge : Edges)
+	{
+		const FTopologicalEdge* Edge = OrientedEdge.Entity.Get();
+		TArray<double> Coordinates;
+		TArray<FPoint> Points;
+		Edge->GetCurve()->GetDiscretizationPoints(Edge->GetBoundary(), Coordinates, Points);
+
+		if (OrientedEdge.Direction == EOrientation::Back)
+		{
+			Algo::Reverse(Coordinates);
+			Algo::Reverse(Points);
+		}
+			
+		int32 PointCount = Points.Num();
+
+		Context.LoopSegments.Reserve(Context.LoopSegments.Num() + PointCount);
+
+		for (int32 ISegment = 0; ISegment < PointCount - 1; ISegment++)
+		{
+			const int32 Index1 = ISegment;	
+			const int32 Index2 = ISegment + 1;
+			Topo::FEdgeSegment& CurrentSeg = Context.SegmentFatory.New();
+			CurrentSeg.SetBoundarySegment(Edge, Coordinates[Index1], Coordinates[Index2], Points[Index1], Points[Index2]);
+			Context.LoopSegments.Add(&CurrentSeg);
+		}
+	}
+}
+
+bool FFaceAnalyzer::IsThinFace(double& OutGapSize)
+{
+#ifdef DEBUG_THIN_FACE
+	F3DDebugSession _(TEXT("Thin Surface"));
+#endif
+
+	Topo::FThinFaceContext Context(*Face.GetExternalLoop());
+
+	FTimePoint StartTime = FChrono::Now();
+	BuildLoopSegments(Context);
+	Chronos.BuildLoopSegmentsTime = FChrono::Elapse(StartTime);
+
+#ifdef CADKERNEL_DEV
+	DisplayLoopSegments(Context);
+#endif
+
+	StartTime = FChrono::Now();
+	FindClosedSegments(Context);
+	Chronos.FindClosedSegmentTime = FChrono::Elapse(StartTime);
+#ifdef CADKERNEL_DEV
+	DisplayClosedSegments(Context);
+#endif
+
+	StartTime = FChrono::Now();
+	Analyze(Context);
+	Chronos.AnalyzeClosedSegmentTime = FChrono::Elapse(StartTime);
+
+	if (Context.MaxSquareDistance < SquareTolerance)
+	{
+		OutGapSize = sqrt(Context.MaxSquareDistance);
+		return true;
+	}
+
+	return false;
+}
+
+#ifdef CADKERNEL_DEV
+void FFaceAnalyzer::DisplayLoopSegments(Topo::FThinFaceContext& Context)
+{
+#ifdef DEBUG_THIN_FACE
+	const FTopologicalEdge* Edge = nullptr;
+	Edge = Context.LoopSegments[0]->GetEdge();
+
+	F3DDebugSession _(TEXT("BoundarySegment"));
+
+#ifdef DEBUG_THIN_FACE_EDGE
+	Open3DDebugSession(FString::Printf(TEXT("Edge %d\n"), Edge->GetId()));
+#endif
+
+	for (Topo::FEdgeSegment* EdgeSegment : Context.LoopSegments)
+	{
+		if (Edge != EdgeSegment->GetEdge())
+		{
+			Edge = EdgeSegment->GetEdge();
+#ifdef DEBUG_THIN_FACE_EDGE
+			Close3DDebugSession();
+			Open3DDebugSession(FString::Printf(TEXT("Edge %d\n"), Edge->GetId()));
+#endif
+		}
+		DisplaySegment(EdgeSegment->GetExtemity(ELimit::Start), EdgeSegment->GetExtemity(ELimit::End), EdgeSegment->GetId(), EVisuProperty::GreenCurve, true);
+		DisplayPoint(EdgeSegment->GetExtemity(ELimit::Start), EVisuProperty::BluePoint);
+	}
+#ifdef DEBUG_THIN_FACE_EDGE
+	Close3DDebugSession();
+#endif
+	//Wait();
+#endif
+}
+
+void FFaceAnalyzer::DisplayClosedSegments(Topo::FThinFaceContext& Context)
+{
+#ifdef DEBUG_THIN_FACE
+	F3DDebugSession _(TEXT("Closed Segment"));
+	for (const Topo::FEdgeSegment* Segment : Context.LoopSegments)
+	{
+		F3DDebugSession _(TEXT("Closed"));
+		if (Segment->GetClosedSegment())
+		{
+			DisplaySegment(Segment->GetExtemity(ELimit::Start), Segment->GetExtemity(ELimit::End), Segment->GetId(), EVisuProperty::BlueCurve);
+			Topo::FEdgeSegment* Parallel = Segment->GetClosedSegment();
+			DisplaySegment(Parallel->GetExtemity(ELimit::Start), Parallel->GetExtemity(ELimit::End), Parallel->GetId(), EVisuProperty::BlueCurve);
+
+			double Coordinate;
+			FPoint Projection = ProjectPointOnSegment(Segment->GetMiddle(), Parallel->GetExtemity(ELimit::Start), Parallel->GetExtemity(ELimit::End), Coordinate, true);
+
+			DisplaySegment(Projection, Segment->GetMiddle(), Segment->GetId(), EVisuProperty::RedCurve);
+		}
+	}
+	//Wait();
+#endif
+}
+#endif
+}
+
diff --git a/Engine/Source/Runtime/Datasmith/CADKernel/Private/CADKernel/Topo/FaceAnalyzer.h b/Engine/Source/Runtime/Datasmith/CADKernel/Private/CADKernel/Topo/FaceAnalyzer.h
new file mode 100644
--- /dev/null
+++ b/Engine/Source/Runtime/Datasmith/CADKernel/Private/CADKernel/Topo/FaceAnalyzer.h
@@ -0,0 +1,203 @@
+// Copyright Epic Games, Inc. All Rights Reserved.
+#pragma once
+
+#include "CADKernel/Core/Chrono.h"
+#include "CADKernel/Core/HaveStates.h"
+#include "CADKernel/Core/Factory.h"
+#include "CADKernel/Core/Types.h"
+#include "CADKernel/Math/Point.h"
+#include "CADKernel/Math/Geometry.h"
+
+
+//#define DEBUG_THIN_FACE
+//#define DISPLAY_THIN_FACE // to display found thin zone
+
+
+namespace UE::CADKernel
+{
+class FTopologicalFace;
+class FTopologicalEdge;
+class FTopologicalLoop;
+
+namespace Topo
+{
+
+struct FFaceAnalyzerChronos
+{
+	FDuration BuildLoopSegmentsTime;
+	FDuration FindClosedSegmentTime;
+	FDuration AnalyzeClosedSegmentTime;
+};
+
+class FEdgeSegment : public FHaveStates
+{
+private:
+	const FTopologicalEdge* Edge = nullptr;
+	double Coordinates[2] = {0. ,0.};
+	FPoint Points[2];
+	FPoint Middle;
+	FPoint Vector;
+
+	FEdgeSegment* ClosedSegment = nullptr;
+
+	double MiddleAxis = 0.;
+	double SquareDistanceToClosedSegment = 0.;
+
+	double Length;
+
+#ifdef CADKERNEL_DEV
+	FIdent Id;
+	static FIdent LastId;
+#endif
+
+public:
+	FEdgeSegment() = default;
+
+	virtual ~FEdgeSegment() = default;
+
+	void SetBoundarySegment(const FTopologicalEdge* InEdge, double InStartU, double InEndU, const FPoint& InStartPoint, const FPoint& InEndPoint)
+	{
+		Edge = InEdge;
+		Coordinates[ELimit::Start] = InStartU;
+		Coordinates[ELimit::End] = InEndU;
+		Points[ELimit::Start] = InStartPoint;
+		Points[ELimit::End] = InEndPoint;
+		Middle = Points[ELimit::Start].Middle(Points[ELimit::End]);
+		Vector = Points[ELimit::End] - Points[ELimit::Start];
+
+		ClosedSegment = nullptr;
+
+		SquareDistanceToClosedSegment = -1.;
+		Length = Points[ELimit::Start].Distance(Points[ELimit::End]);
+
+		MiddleAxis = Middle.DiagonalAxisCoordinate();
+
+#ifdef CADKERNEL_DEV
+		Id = LastId++;
+#endif
+	};
+
+	const FPoint GetVector() const
+	{
+		return Vector;
+	}
+
+#ifdef CADKERNEL_DEV
+	FIdent GetId() const
+	{
+		return Id;
+	}
+#endif
+
+	const FTopologicalEdge* GetEdge() const
+	{
+		return Edge;
+	}
+
+	double GetLength() const
+	{
+		return Length;
+	}
+
+	FPoint GetMiddle() const
+	{
+		return Middle;
+	}
+
+	constexpr const FPoint& GetExtemity(const ELimit Limit) const
+	{
+		return Points[Limit];
+	}
+
+	constexpr double GetCoordinate(const ELimit Limit) const
+	{
+		return Coordinates[Limit];
+	}
+
+	/**
+	 * Compute the slope of the input Segment according to this.
+	 */
+	double ComputeCosAngleOf(const FEdgeSegment* Segment)
+	{
+		return GetVector().ComputeCosinus(Segment->GetVector());
+	}
+
+	FEdgeSegment* GetClosedSegment() const
+	{
+		return ClosedSegment;
+	}
+
+	void SetClosedSegment(FEdgeSegment* InSegmentA, double InSquareDistance)
+	{
+		ClosedSegment = InSegmentA;
+		SquareDistanceToClosedSegment = InSquareDistance;
+	}
+
+	double GetClosedSquareDistance() const
+	{
+		return SquareDistanceToClosedSegment;
+	}
+
+	FPoint ProjectPoint(const FPoint& PointToProject, double& SegmentU) const
+	{
+		return ProjectPointOnSegment(PointToProject, Points[ELimit::Start], Points[ELimit::End], SegmentU, true);
+	}
+};
+
+struct FThinFaceContext
+{
+	const FTopologicalLoop& Loop;
+
+	TArray<Topo::FEdgeSegment*> LoopSegments;
+	double ExternalLoopLength = 0.0;
+	TFactory<Topo::FEdgeSegment> SegmentFatory;
+
+	TArray<double> EdgeSquareDistance;
+	TArray<double> EdgeMaxSquareDistance;
+
+	double MaxSquareDistance = 0;
+	double SideEdgeLength = 0;
+
+	FThinFaceContext(const FTopologicalLoop& InLoop)
+		: Loop(InLoop)
+	{
+	}
+};
+
+}
+
+class FFaceAnalyzer
+{
+
+public:
+	Topo::FFaceAnalyzerChronos Chronos;
+
+protected:
+	double Tolerance;
+	double SquareTolerance;
+
+	FTopologicalFace& Face;
+
+public:
+
+	FFaceAnalyzer(FTopologicalFace& InFace, double InTol)
+		: Tolerance(InTol)
+		, SquareTolerance(FMath::Square(InTol))
+		, Face(InFace)
+	{
+ 	}
+
+	bool IsThinFace(double& OutGapSize);
+
+private:
+
+	void BuildLoopSegments(Topo::FThinFaceContext& Context);
+	void FindClosedSegments(Topo::FThinFaceContext& Context);
+	void Analyze(Topo::FThinFaceContext& Context);
+
+#ifdef CADKERNEL_DEV
+	void DisplayClosedSegments(Topo::FThinFaceContext& Context);
+	void DisplayLoopSegments(Topo::FThinFaceContext& Context);
+#endif
+};
+}
diff --git a/Engine/Source/Runtime/Datasmith/CADKernel/Private/CADKernel/Topo/Model.cpp b/Engine/Source/Runtime/Datasmith/CADKernel/Private/CADKernel/Topo/Model.cpp
--- a/Engine/Source/Runtime/Datasmith/CADKernel/Private/CADKernel/Topo/Model.cpp
+++ b/Engine/Source/Runtime/Datasmith/CADKernel/Private/CADKernel/Topo/Model.cpp
@@ -21,9 +21,6 @@ void FModel::AddEntity(TSharedRef<FTopologicalEntity> Entity)
 	case EEntity::Body:
 		Add(StaticCastSharedRef<FBody>(Entity));
 		break;
-	case EEntity::TopologicalFace:
-		Add(StaticCastSharedRef<FTopologicalFace>(Entity));
-		break;
 	default:
 		break;
 	}
@@ -35,8 +32,6 @@ bool FModel::Contains(TSharedPtr<FTopologicalEntity> Entity)
 	{
 	case EEntity::Body:
 		return Bodies.Find(StaticCastSharedPtr<FBody>(Entity)) != INDEX_NONE;
-	case EEntity::TopologicalFace:
-		return Faces.Find(StaticCastSharedPtr<FTopologicalFace>(Entity)) != INDEX_NONE;
 	default:
 		return false;
 	}
@@ -50,9 +45,6 @@ void FModel::RemoveEntity(TSharedPtr<FTopologicalEntity> Entity)
 	case EEntity::Body:
 		RemoveBody(StaticCastSharedPtr<FBody>(Entity));
 		break;
-	case EEntity::TopologicalFace:
-		RemoveFace(StaticCastSharedPtr<FTopologicalFace>(Entity));
-		break;
 	default:
 		break;
 	}
@@ -79,7 +71,6 @@ int32 FModel::FaceCount() const
 	{
 		FaceCount += Body->FaceCount();
 	}
-	FaceCount += Faces.Num();
 	return FaceCount;
 }
 
@@ -89,15 +80,6 @@ void FModel::GetFaces(TArray<FTopologicalFace*>& OutFaces)
 	{
 		Body->GetFaces(OutFaces);
 	}
-
-	for (TSharedPtr<FTopologicalFace>& Face : Faces)
-	{
-		if (!Face->HasMarker1())
-		{
-			OutFaces.Add(Face.Get());
-			Face->SetMarker1();
-		}
-	}
 }
 
 void FModel::SpreadBodyOrientation()
@@ -112,19 +94,10 @@ void FModel::SpreadBodyOrientation()
 FInfoEntity& FModel::GetInfo(FInfoEntity& Info) const
 {
 	return FTopologicalShapeEntity::GetInfo(Info)
-		.Add(TEXT("Bodies"), Bodies)
-		.Add(TEXT("Faces"), Faces);
+		.Add(TEXT("Bodies"), Bodies);
 }
 #endif
 
-
-// Topo functions
-void FModel::MergeInto(TSharedPtr<FBody> Body, TArray<TSharedPtr<FTopologicalEntity>>& InEntities)
-{
-
-}
-
-
 struct FBodyShell
 {
 	TSharedPtr<FBody> Body;
@@ -203,11 +176,6 @@ void FModel::FillTopologyReport(FTopologyReport& Report) const
 	{
 		Body->FillTopologyReport(Report);
 	}
-
-	for (TSharedPtr<FTopologicalFace> Face : Faces)
-	{
-		Face->FillTopologyReport(Report);
-	}
 }
 #endif
 
@@ -219,69 +187,4 @@ void FModel::Orient()
 	}
 }
 
-/**
- * Fore each shell of each body, try to stitch topological gap
- */
-void FModel::HealModelTopology(double JoiningTolerance)
-{
-
-}
-
-
-
-//void FModel::FixModelTopology()
-void FModel::FixModelTopology(double JoiningTolerance)
-{
-	//HealModelTopology(JoiningTolerance);
-
-	//// Find Isolated Shell
-	//TArray<TSharedPtr<FShell>> IsolatedShell;
-
-	//for (TSharedPtr<FBody> Body : Bodies)
-	//{
-	//	for (TSharedPtr<FShell> Shell : Body->GetShells())
-	//	{
-	//		if (Shell->IsOpenShell())
-	//		{
-	//			IsolatedShell.Add(Shell);
-	//		}
-	//	}
-	//}
-
-	//FJoiner Joiner(IsolatedShell, JoiningTolerance);
-	//Joiner.JoinFaces();
-
-	/*
-	Body->GetFaces(Surfaces);
-	TSharedRef<FBody> MergedBody = FEntity::MakeShared<FBody>();
-	TSharedRef<FShell> MergedShell = FEntity::MakeShared<FShell>();
-
-	MergedBody->AddShell(MergedShell);
-	TArray<TSharedPtr<FBody>> NewBodies;
-	NewBodies.Reserve(Bodies.Num());
-
-	for (TSharedPtr<FBody>& Body : Bodies)
-	{
-		if (Body->GetShells().Num() == 1)
-		{
-			TSharedPtr<FShell> Shell = Body->GetShells()[0];
-			if (Shell->FaceCount() < 3)
-			{
-				MergedShell->Merge(Shell);
-				Body->Empty();
-			}
-		}
-		else
-		{
-			NewBodies.Emplace(Body);
-			Body.Reset();
-		}
-	}
-
-	NewBodies.Emplace(MergedBody);
-
-	Swap(NewBodies, Bodies);
-	*/
-}
-
 }
\ No newline at end of file
diff --git a/Engine/Source/Runtime/Datasmith/CADKernel/Private/CADKernel/Topo/Shell.cpp b/Engine/Source/Runtime/Datasmith/CADKernel/Private/CADKernel/Topo/Shell.cpp
--- a/Engine/Source/Runtime/Datasmith/CADKernel/Private/CADKernel/Topo/Shell.cpp
+++ b/Engine/Source/Runtime/Datasmith/CADKernel/Private/CADKernel/Topo/Shell.cpp
@@ -79,6 +79,17 @@ void FShell::Add(TSharedRef<FTopologicalFace> InTopologicalFace, EOrientation Or
 	Face->SetHost(this);
 }
 
+void FShell::Remove(const FTopologicalShapeEntity* FaceToRemove)
+{
+	if (!FaceToRemove)
+	{
+		return;
+	}
+
+	int32 Index = TopologicalFaces.IndexOfByPredicate([&](const FOrientedFace& Face) { return (Face.Entity.Get() == FaceToRemove); });
+	TopologicalFaces.RemoveAt(Index);
+}
+
 #ifdef CADKERNEL_DEV
 FInfoEntity& FShell::GetInfo(FInfoEntity& Info) const
 {
diff --git a/Engine/Source/Runtime/Datasmith/CADKernel/Private/CADKernel/Topo/TopologicalEdge.cpp b/Engine/Source/Runtime/Datasmith/CADKernel/Private/CADKernel/Topo/TopologicalEdge.cpp
--- a/Engine/Source/Runtime/Datasmith/CADKernel/Private/CADKernel/Topo/TopologicalEdge.cpp
+++ b/Engine/Source/Runtime/Datasmith/CADKernel/Private/CADKernel/Topo/TopologicalEdge.cpp
@@ -559,6 +559,7 @@ bool FTopologicalEdge::IsSameDirection(const FTopologicalEdge& Edge) const
 FInfoEntity& FTopologicalEdge::GetInfo(FInfoEntity& Info) const
 {
 	return FTopologicalEntity::GetInfo(Info)
+		.Add(TEXT("IsDegenerated"), IsDegenerated())
 		.Add(TEXT("Link"), TopologicalLink)
 		.Add(TEXT("Curve"), Curve)
 		.Add(TEXT("Vertex1"), StartVertex)
@@ -1070,6 +1071,10 @@ bool FTopologicalEdge::IsSharpEdge()
 	return false;
 }
 
+void FTopologicalEdge::Offset2D(const FPoint2D& OffsetDirection)
+{
+	Curve->Offset2D(OffsetDirection);
+}
 
 #ifdef CADKERNEL_DEV
 FInfoEntity& FEdgeLink::GetInfo(FInfoEntity& Info) const
diff --git a/Engine/Source/Runtime/Datasmith/CADKernel/Private/CADKernel/Topo/TopologicalEntity.cpp b/Engine/Source/Runtime/Datasmith/CADKernel/Private/CADKernel/Topo/TopologicalEntity.cpp
--- a/Engine/Source/Runtime/Datasmith/CADKernel/Private/CADKernel/Topo/TopologicalEntity.cpp
+++ b/Engine/Source/Runtime/Datasmith/CADKernel/Private/CADKernel/Topo/TopologicalEntity.cpp
@@ -8,6 +8,7 @@ namespace UE::CADKernel
 FInfoEntity& FTopologicalEntity::GetInfo(FInfoEntity& Info) const
 {
 	return FEntity::GetInfo(Info)
+		.Add(TEXT("IsDeleted"), IsDeleted())
 		.Add(TEXT("IsMesh"), IsMeshed());
 }
 #endif
diff --git a/Engine/Source/Runtime/Datasmith/CADKernel/Private/CADKernel/Topo/TopologicalFace.cpp b/Engine/Source/Runtime/Datasmith/CADKernel/Private/CADKernel/Topo/TopologicalFace.cpp
--- a/Engine/Source/Runtime/Datasmith/CADKernel/Private/CADKernel/Topo/TopologicalFace.cpp
+++ b/Engine/Source/Runtime/Datasmith/CADKernel/Private/CADKernel/Topo/TopologicalFace.cpp
@@ -263,13 +263,25 @@ void FTopologicalFace::RemoveLoop(const TSharedPtr<FTopologicalLoop>& Loop)
 	}
 }
 
-void FTopologicalFace::RemoveLinksWithNeighbours()
+void FTopologicalFace::Disjoin(TArray<FTopologicalEdge*>& NewBorderEdges)
 {
+	NewBorderEdges.Reserve(NewBorderEdges.Num() + EdgeCount());
 	for (const TSharedPtr<FTopologicalLoop>& Loop : GetLoops())
 	{
 		for (const FOrientedEdge& Edge : Loop->GetEdges())
 		{
+			const TArray<FTopologicalEdge*> Twins = Edge.Entity->GetTwinEntities();
+			for (FTopologicalEdge* TwinEdge : Twins)
+			{
+				if (TwinEdge != Edge.Entity.Get())
+				{
+					NewBorderEdges.Add(TwinEdge);
+				}
+			}
 			Edge.Entity->RemoveFromLink();
+			Edge.Entity->SetMarker1();
+			Edge.Entity->GetStartVertex()->RemoveFromLink();
+			Edge.Entity->GetEndVertex()->RemoveFromLink();
 		}
 	}
 }
@@ -632,6 +644,19 @@ void FTopologicalFace::DefineSurfaceType()
 	}
 }
 
+const TSharedPtr<FTopologicalLoop> FTopologicalFace::GetExternalLoop() const
+{
+	for (const TSharedPtr<FTopologicalLoop>& Loop : GetLoops())
+	{
+		if (Loop->IsExternal())
+		{
+			return Loop;
+		}
+	}
+	return TSharedPtr<FTopologicalLoop>();
+}
+
+
 
 void FFaceSubset::SetMainShell(TMap<FTopologicalShapeEntity*, int32>& ShellToFaceCount)
 {
diff --git a/Engine/Source/Runtime/Datasmith/CADKernel/Private/CADKernel/Topo/TopologicalLoop.cpp b/Engine/Source/Runtime/Datasmith/CADKernel/Private/CADKernel/Topo/TopologicalLoop.cpp
--- a/Engine/Source/Runtime/Datasmith/CADKernel/Private/CADKernel/Topo/TopologicalLoop.cpp
+++ b/Engine/Source/Runtime/Datasmith/CADKernel/Private/CADKernel/Topo/TopologicalLoop.cpp
@@ -22,6 +22,32 @@ TSharedPtr<FTopologicalLoop> FTopologicalLoop::Make(const TArray<TSharedPtr<FTop
 	TSharedRef<FTopologicalLoop> LoopRef = FEntity::MakeShared<FTopologicalLoop>(InEdges, InEdgeDirections, bIsExternalLoop);
 	FTopologicalLoop& Loop = *LoopRef;
 
+	const FSurface& Surface = *InEdges[0]->GetCurve()->GetCarrierSurface();
+	const bool bIsSphere = Surface.GetSurfaceType() == ESurface::Sphere;
+	if (bIsSphere && InEdges.Num() == 2)
+	{
+		// Workaround of TechSoft Bug (#jira https://techsoft3d.atlassian.net/servicedesk/customer/portal/7/SDHE-19879)
+		// If the loop is composed of two coincident edges in the parametric space, linking the poles. Offset one edge of 2Pi
+		const FTopologicalEdge& Edge0 = *InEdges[0];
+		FTopologicalEdge& Edge1 = *InEdges[1];
+
+		const FPoint2D StartPoint = Edge0.Approximate2DPoint(Edge0.GetBoundary().GetMin());
+		const FPoint2D EndPoint = Edge0.Approximate2DPoint(Edge0.GetBoundary().GetMax());
+		if (FMath::IsNearlyEqual(FMath::Abs(StartPoint.V), DOUBLE_HALF_PI) && FMath::IsNearlyEqual(FMath::Abs(EndPoint.V), DOUBLE_HALF_PI))
+		{
+			const double TolU = Surface.GetIsoTolerance(EIso::IsoU);
+			const bool Edge0IsIso = Edge0.GetCurve()->IsIso(EIso::IsoU, TolU);
+			const bool Edge1IsIso = Edge1.GetCurve()->IsIso(EIso::IsoU, TolU);
+			const FPoint2D Edge1EndPoint = Edge1.Approximate2DPoint(InEdgeDirections[0]== InEdgeDirections[1] ? Edge1.GetBoundary().GetMax() : Edge1.GetBoundary().GetMin());
+			double Distance = Edge1EndPoint.SquareDistance(StartPoint);
+			if (Distance < FMath::Square(TolU))
+			{
+				const FPoint2D Offset(TWO_PI, 0.);
+				Edge1.Offset2D(Offset);
+			}
+		}
+	}
+
 	Loop.EnsureLogicalClosing(GeometricTolerance);
 	Loop.RemoveDegeneratedEdges();
 
@@ -116,7 +142,7 @@ EOrientation FTopologicalLoop::GetDirection(TSharedPtr<FTopologicalEdge>& InEdge
 }
 
 
-void FTopologicalLoop::Get2DSampling(TArray<FPoint2D>& LoopSampling)
+void FTopologicalLoop::Get2DSampling(TArray<FPoint2D>& LoopSampling) const
 {
 	int32 PointCount = 0;
 	for (const FOrientedEdge& Edge : Edges)
@@ -134,7 +160,7 @@ void FTopologicalLoop::Get2DSampling(TArray<FPoint2D>& LoopSampling)
 	LoopSampling.Emplace(LoopSampling[0]);
 }
 
-bool FTopologicalLoop::Get2DSamplingWithoutDegeneratedEdges(TArray<FPoint2D>& LoopSampling)
+bool FTopologicalLoop::Get2DSamplingWithoutDegeneratedEdges(TArray<FPoint2D>& LoopSampling) const
 {
 	double LoopLength = 0;
 	int32 EdgeCount = 0;
@@ -211,6 +237,8 @@ bool FTopologicalLoop::Get2DSamplingWithoutDegeneratedEdges(TArray<FPoint2D>& Lo
  * @return false if the orientation is doubtful
  * 
  */
+
+//#define DEBUG_ORIENT
 bool FTopologicalLoop::Orient()
 {
 	bool bSucceed = true;
@@ -226,6 +254,8 @@ bool FTopologicalLoop::Orient()
 	}
 
 #ifdef DEBUG_ORIENT
+	bool bDisplayDebug = (this->GetFace()->GetId() == 735);
+	if(bDisplayDebug)
 	{
 		F3DDebugSession _(*FString::Printf(TEXT("Loop before orientation")));
 		DisplayOrientedPolyline(LoopSampling, EVisuProperty::BlueCurve);
@@ -328,6 +358,7 @@ bool FTopologicalLoop::Orient()
 			Slop = ComputePositiveSlope(LoopSampling[Index], LoopSampling[NextIndex], LoopSampling[PreviousIndex]);
 
 #ifdef DEBUG_ORIENT
+			if (bDisplayDebug)
 			{
 				F3DDebugSession _(*FString::Printf(TEXT("Fix Pic Node %f"), Slop));
 				DisplayPoint(LoopSampling[Index], EVisuProperty::BluePoint, Index);
@@ -384,27 +415,30 @@ bool FTopologicalLoop::Orient()
 			else
 			{
 #ifdef DEBUG_ORIENT
-				F3DDebugSession _(*FString::Printf(TEXT("Pic case")));
+				if (bDisplayDebug)
 				{
-					F3DDebugSession _(*FString::Printf(TEXT("Loop")));
-					DisplayPolyline(LoopSampling, EVisuProperty::BlueCurve);
-				}
-				{
-					F3DDebugSession _(*FString::Printf(TEXT("Next")));
-					DisplaySegment(LoopSampling[NextIndex], LoopSampling[Index], EVisuProperty::YellowCurve);
-				}
-				{
-					F3DDebugSession _(*FString::Printf(TEXT("Node %f"), Slop));
-					DisplayPoint(LoopSampling[Index], EVisuProperty::RedPoint, Index);
+					F3DDebugSession _(*FString::Printf(TEXT("Pic case")));
+					{
+						F3DDebugSession _(*FString::Printf(TEXT("Loop")));
+						DisplayPolyline(LoopSampling, EVisuProperty::BlueCurve);
+					}
+					{
+						F3DDebugSession _(*FString::Printf(TEXT("Next")));
+						DisplaySegment(LoopSampling[NextIndex], LoopSampling[Index], EVisuProperty::YellowCurve);
+					}
+					{
+						F3DDebugSession _(*FString::Printf(TEXT("Node %f"), Slop));
+						DisplayPoint(LoopSampling[Index], EVisuProperty::RedPoint, Index);
+					}
+					Wait();
 				}
-				Wait();
 #endif
 				UndefinedOrientationCount++;
 			}
 		}
 
 #ifdef DEBUG_ORIENT
-		if (false /*Surface->GetId() == 826*/)
+		if (bDisplayDebug)
 		{
 			F3DDebugSession _(*FString::Printf(TEXT("Pic case")));
 			{
@@ -436,22 +470,25 @@ bool FTopologicalLoop::Orient()
 	if ((WrongOrientationCount != 0 && GoodOrientationCount != 0) || UndefinedOrientationCount > FMath::Max(WrongOrientationCount, GoodOrientationCount))
 	{
 #ifdef DEBUG_ORIENT
-		F3DDebugSession GraphicSession(TEXT("Points of evaluation"));
+		if (bDisplayDebug)
 		{
-			F3DDebugSession G(*FString::Printf(TEXT("Loop Discretization %d"), Face.Pin()->GetId()));
-			DisplayPolyline(LoopSampling, EVisuProperty::BlueCurve);
-			for (int32 Index = 0; Index < LoopSampling.Num(); ++Index)
+			F3DDebugSession GraphicSession(TEXT("Points of evaluation"));
 			{
-				DisplayPoint(LoopSampling[Index], Index);
+				F3DDebugSession G(*FString::Printf(TEXT("Loop Discretization %d"), Face->GetId()));
+				DisplayPolyline(LoopSampling, EVisuProperty::BlueCurve);
+				for (int32 Index = 0; Index < LoopSampling.Num(); ++Index)
+				{
+					DisplayPoint(LoopSampling[Index], Index);
+				}
 			}
-		}
 
-		for (int32 Index : ExtremityIndex)
-		{
-			F3DDebugSession G(*FString::Printf(TEXT("Seg UMin Loop Discretization %d"), Face.Pin()->GetId()));
-			DisplayPoint(LoopSampling[Index], EVisuProperty::RedPoint, Index);
+			for (int32 Index : ExtremityIndex)
+			{
+				F3DDebugSession G(*FString::Printf(TEXT("Seg UMin Loop Discretization %d"), Face->GetId()));
+				DisplayPoint(LoopSampling[Index], EVisuProperty::RedPoint, Index);
+			}
+			Wait();
 		}
-		Wait();
 #endif
 
 		bSucceed = false;
@@ -464,15 +501,14 @@ bool FTopologicalLoop::Orient()
 	}
 
 #ifdef DEBUG_ORIENT
+	if (bDisplayDebug)
 	{
 		LoopSampling.Empty();
-		if (!Get2DSamplingWithoutDegeneratedEdges(LoopSampling))
+		if (Get2DSamplingWithoutDegeneratedEdges(LoopSampling))
 		{
-			return;
+			F3DDebugSession _(*FString::Printf(TEXT("Loop oriented")));
+			DisplayOrientedPolyline(LoopSampling, EVisuProperty::BlueCurve);
 		}
-
-		F3DDebugSession _(*FString::Printf(TEXT("Loop oriented")));
-		DisplayOrientedPolyline(LoopSampling, EVisuProperty::BlueCurve);
 	}
 #endif
 
@@ -925,14 +961,20 @@ void FTopologicalLoop::EnsureLogicalClosing(const double Tolerance3D)
 
 			double CosAngle = Gap2D.ComputeCosinus(EdgeTangent);
 			constexpr double CosFlatTangent = 0.98;  // ~10 deg: In 2D, the distortion due to the parametric space, degenerated area and else imposes to be more careful before extending a curve.
-			if (CosAngle > CosFlatTangent)
+
+			//   ----------------------*         ------------**--------*
+            //               *---------*    Or                         *
+			//               |				                           |
+			//               |				                           |
+			//               |				                           |
+			if (FMath::Abs(CosAngle) > CosFlatTangent)
 			{
 				OrientedEdge.Entity->ExtendTo(OrientedEdge.Direction == EOrientation::Front, PreviousExtremity.Point2D, PreviousEdgeEndVertex);
 			}
 			else
 			{
 				CosAngle = Gap2D.ComputeCosinus(PreviousTangent);
-				if (CosAngle < -CosFlatTangent)
+				if (FMath::Abs(CosAngle) > CosFlatTangent)
 				{
 					PreviousEdge.Entity->ExtendTo(PreviousEdge.Direction == EOrientation::Back, EdgeExtremities[ExtremityIndex].Point2D, EdgeStartVertex);
 				}
@@ -949,6 +991,11 @@ void FTopologicalLoop::EnsureLogicalClosing(const double Tolerance3D)
 						PreviousEdgeEndVertex->Link(*EdgeStartVertex);
 						++Index;
 					}
+					else
+					{
+						// joins two vertices
+						PreviousEdgeEndVertex->Link(*EdgeStartVertex);
+					}
 				}
 			}
 		}
@@ -1032,4 +1079,99 @@ void FTopologicalLoop::FindBreaks(TArray<TSharedPtr<FTopologicalVertex>>& OutBre
 	}
 }
 
+namespace TopologicalLoopImpl
+{
+void FindLoopIntersectionsWithIso(const EIso Iso, const double IsoParameter, const TArray<FPoint2D>& Loop, TArray<double>& OutIntersections)
+{
+	OutIntersections.Empty(8);
+
+	int32 UIndex = Iso == EIso::IsoU ? 0 : 1;
+	int32 VIndex = Iso == EIso::IsoU ? 1 : 0;
+
+	TFunction<void(const FPoint2D&, const FPoint2D&)> ComputeIntersection = [&](const FPoint2D& Point1, const FPoint2D& Point2)
+	{
+		if (IsoParameter > Point1[UIndex] && IsoParameter <= Point2[UIndex])
+		{
+			double Intersection = (IsoParameter - Point1[UIndex]) / (Point2[UIndex] - Point1[UIndex]) * (Point2[VIndex] - Point1[VIndex]) + Point1[VIndex];
+			OutIntersections.Add((IsoParameter - Point1[UIndex]) / (Point2[UIndex] - Point1[UIndex]) * (Point2[VIndex] - Point1[VIndex]) + Point1[VIndex]);
+		}
+	};
+
+	const FPoint2D* Point1 = &Loop.Last();
+	for (const FPoint2D& Point2 : Loop)
+	{
+		if (!FMath::IsNearlyEqual((*Point1)[UIndex], Point2[UIndex]))
+		{
+			if ((*Point1)[UIndex] < Point2[UIndex])
+			{
+				ComputeIntersection(*Point1, Point2);
+			}
+			else
+			{
+				ComputeIntersection(Point2, *Point1);
+			}
+		}
+		Point1 = &Point2;
+	}
+
+	Algo::Sort(OutIntersections);
+}
+};
+
+bool FTopologicalLoop::IsInside(const FTopologicalLoop& OtherLoop) const
+{
+	TArray<FPoint2D> Sampling2D;
+	Get2DSampling(Sampling2D);
+
+	TArray<FPoint2D> OtherSampling2D;
+	OtherLoop.Get2DSampling(OtherSampling2D);
+
+
+	int32 InsidePoint = 0;
+	int32 OutsidePoint = 0;
+
+	TFunction<void(const FPoint2D&, EIso)> CountLeftIntersection = [&](const FPoint2D& Point, EIso Iso)
+	{
+		TArray<double> Intersections;
+		TopologicalLoopImpl::FindLoopIntersectionsWithIso(Iso, Point[Iso], Sampling2D, Intersections);
+
+		int32 IntersectionLeftCount = 0;
+		const EIso OtherIso = Iso == EIso::IsoU ? EIso::IsoV : EIso::IsoU;
+
+		for (const double& Intersection : Intersections)
+		{
+			if (Intersection < Point[OtherIso])
+			{
+				IntersectionLeftCount++;
+			}
+			else
+			{
+				break;
+			}
+		}
+
+		if (IntersectionLeftCount % 2 == 0)
+		{
+			OutsidePoint++;
+		}
+		else
+		{
+			InsidePoint++;
+		}
+	};
+
+	const int32 Step = 7;
+	const int32 OtherSampling2DCount = OtherSampling2D.Num();
+	const int32 Increment = FMath::Max(OtherSampling2DCount / Step, 1);
+	for (int32 Index = 0; Index < OtherSampling2DCount; Index += Increment)
+	{
+		const FPoint2D& TestPoint = OtherSampling2D[Index];
+		CountLeftIntersection(TestPoint, EIso::IsoU);
+		CountLeftIntersection(TestPoint, EIso::IsoV);
+	}
+
+	return InsidePoint > OutsidePoint;
+}
+
+
 }
\ No newline at end of file
diff --git a/Engine/Source/Runtime/Datasmith/CADKernel/Private/CADKernel/Topo/Topomaker.cpp b/Engine/Source/Runtime/Datasmith/CADKernel/Private/CADKernel/Topo/Topomaker.cpp
--- a/Engine/Source/Runtime/Datasmith/CADKernel/Private/CADKernel/Topo/Topomaker.cpp
+++ b/Engine/Source/Runtime/Datasmith/CADKernel/Private/CADKernel/Topo/Topomaker.cpp
@@ -6,6 +6,7 @@
 #include "CADKernel/Core/Session.h"
 #include "CADKernel/Geo/Curves/RestrictionCurve.h"
 #include "CADKernel/Topo/Body.h"
+#include "CADKernel/Topo/FaceAnalyzer.h"
 #include "CADKernel/Topo/Model.h"
 #include "CADKernel/Topo/Shell.h"
 #include "CADKernel/Topo/TopologicalEdge.h"
@@ -23,9 +24,9 @@ namespace UE::CADKernel
 
 FTopomaker::FTopomaker(FSession& InSession, double InTolerance)
 	: Session(InSession)
-	, SewTolerance(InTolerance * UE_DOUBLE_SQRT_2)
-	, SewToleranceSquare(FMath::Square(SewTolerance))
 {
+	SetTolerance(InTolerance);
+
 	int32 ShellCount = 0;
 	for (const TSharedPtr<FBody>& Body : Session.GetModel().GetBodies())
 	{
@@ -47,18 +48,17 @@ FTopomaker::FTopomaker(FSession& InSession, double InTolerance)
 FTopomaker::FTopomaker(FSession& InSession, const TArray<TSharedPtr<FTopologicalFace>>& InFaces, double InTolerance)
 	: Session(InSession)
 	, Faces(InFaces)
-	, SewTolerance(InTolerance * UE_DOUBLE_SQRT_2)
-	, SewToleranceSquare(FMath::Square(SewTolerance))
 {
+	SetTolerance(InTolerance);
 }
 
 FTopomaker::FTopomaker(FSession& InSession, const TArray<TSharedPtr<FShell>>& InShells, double InTolerance)
 	: Session(InSession)
-	, SewTolerance(InTolerance)
-	, SewToleranceSquare(FMath::Square(SewTolerance))
 {
+	SetTolerance(InTolerance);
+
 	Shells.Reserve(InShells.Num());
-	for(const TSharedPtr<FShell>& Shell : InShells)
+	for (const TSharedPtr<FShell>& Shell : InShells)
 	{
 		Shells.Add(Shell.Get());
 	}
@@ -102,9 +102,9 @@ void FTopomaker::RemoveFacesFromShell()
 	{
 		if (Face->GetHost() != nullptr)
 		{
-			ShellSet.Add((FShell*) Face->GetHost());
+			ShellSet.Add((FShell*)Face->GetHost());
 			Face->ResetHost();
-		}	
+		}
 	}
 
 	for (FShell* Shell : ShellSet)
@@ -211,7 +211,6 @@ void FTopomaker::MergeCoincidentVertices(TArray<FTopologicalVertex*>& VerticesTo
 {
 	FTimePoint StartTime = FChrono::Now();
 
-	const double JoiningVerticesToleranceSquare = 2 * SewToleranceSquare;
 	const double WeigthTolerance = 3 * SewTolerance;
 
 	int32 VertexNum = (int32)VerticesToMerge.Num();
@@ -224,7 +223,7 @@ void FTopomaker::MergeCoincidentVertices(TArray<FTopologicalVertex*>& VerticesTo
 
 	for (FTopologicalVertex* Vertex : VerticesToMerge)
 	{
-		VerticesWeight.Add(Vertex->GetCoordinates().X + Vertex->GetCoordinates().Y + Vertex->GetCoordinates().Z);
+		VerticesWeight.Add(Vertex->GetCoordinates().DiagonalAxisCoordinate());
 	}
 
 	for (int32 Index = 0; Index < VerticesToMerge.Num(); ++Index)
@@ -263,7 +262,7 @@ void FTopomaker::MergeCoincidentVertices(TArray<FTopologicalVertex*>& VerticesTo
 			}
 
 			double DistanceSqr = OtherVertex->GetLinkActiveEntity()->SquareDistance(Barycenter);
-			if (DistanceSqr < JoiningVerticesToleranceSquare)
+			if (DistanceSqr < SewToleranceSquare)
 			{
 				OtherVertex->SetMarker1();
 				Vertex->Link(*OtherVertex);
@@ -394,6 +393,70 @@ void FTopomaker::MergeBorderVerticesWithCoincidentOtherVertices(TArray<FTopologi
 	Swap(NewVertices, Vertices);
 }
 
+void FTopomaker::MergeCoincidentEdges(FTopologicalEdge* Edge)
+{
+	const FTopologicalVertex& StartVertex = *Edge->GetStartVertex()->GetLinkActiveEntity();
+	const FTopologicalVertex& EndVertex = *Edge->GetEndVertex()->GetLinkActiveEntity();
+
+	TArray<FTopologicalEdge*> ConnectedEdges;
+	StartVertex.GetConnectedEdges(EndVertex, ConnectedEdges);
+	const int32 ConnectedEdgeCount = ConnectedEdges.Num();
+	if (ConnectedEdgeCount == 1)
+	{
+		return;
+	}
+
+	const bool bFirstEdgeBorder = Edge->IsBorder();
+
+	for (int32 Index = 0; Index < ConnectedEdgeCount; ++Index)
+	{
+		FTopologicalEdge* SecondEdge = ConnectedEdges[Index];
+		if (SecondEdge == Edge || !SecondEdge->IsActiveEntity() || SecondEdge->IsDegenerated())
+		{
+			continue;
+		}
+
+		const bool bSecondEdgeBorder = Edge->IsBorder();
+
+		// Process only if at least one edge is Border
+		if (!bFirstEdgeBorder && !bSecondEdgeBorder)
+		{
+			continue;
+		}
+
+		const FPoint StartTangentEdge = Edge->GetTangentAt(StartVertex);
+		const FPoint StartTangentOtherEdge = SecondEdge->GetTangentAt(StartVertex);
+
+		double CosAngle = StartTangentEdge.ComputeCosinus(StartTangentOtherEdge);
+		if (CosAngle < 0.9)
+		{
+			continue;
+		}
+
+		if (Edge->GetFace() != SecondEdge->GetFace())
+		{
+			Edge->Link(*SecondEdge, SewTolerance);
+		}
+	}
+}
+
+void FTopomaker::MergeCoincidentEdges(TArray<FTopologicalEdge*>& EdgesToProcess)
+{
+	FTimePoint StartTime = FChrono::Now();
+
+	for (FTopologicalEdge* Edge : EdgesToProcess)
+	{
+		if (Edge->IsDegenerated() || !Edge->IsBorder())
+		{
+			continue;
+		}
+		MergeCoincidentEdges(Edge);
+	}
+
+	FDuration StepDuration = FChrono::Elapse(StartTime);
+	FChrono::PrintClockElapse(EVerboseLevel::Log, TEXT("    "), TEXT("Merge coincident edges"), StepDuration);
+}
+
 void FTopomaker::MergeCoincidentEdges(TArray<FTopologicalVertex*>& VerticesToProcess)
 {
 	FTimePoint StartTime = FChrono::Now();
@@ -402,7 +465,7 @@ void FTopomaker::MergeCoincidentEdges(TArray<FTopologicalVertex*>& VerticesToPro
 	{
 		TArray<FTopologicalEdge*> ConnectedEdges;
 		Vertex->GetConnectedEdges(ConnectedEdges);
-		int32 ConnectedEdgeCount = ConnectedEdges.Num();
+		const int32 ConnectedEdgeCount = ConnectedEdges.Num();
 		if (ConnectedEdgeCount == 1)
 		{
 			continue;
@@ -416,51 +479,12 @@ void FTopomaker::MergeCoincidentEdges(TArray<FTopologicalVertex*>& VerticesToPro
 				continue;
 			}
 
-			bool bFirstEdgeBorder = Edge->IsBorder();
-			FTopologicalVertex& EndVertex = *Edge->GetOtherVertex(*Vertex)->GetLinkActiveEntity();
-
-			for (int32 EdgeJ = EdgeI + 1; EdgeJ < ConnectedEdgeCount; ++EdgeJ)
-			{
-				FTopologicalEdge* SecondEdge = ConnectedEdges[EdgeJ];
-				if (!SecondEdge->IsActiveEntity() || SecondEdge->IsDegenerated())
-				{
-					continue;
-				}
-
-				bool bSecondEdgeBorder = Edge->IsBorder();
-
-				// Process only if at least one edge is Border
-				if (!bFirstEdgeBorder && !bSecondEdgeBorder)
-				{
-					continue;
-				}
-
-				FTopologicalVertex& OtherEdgeEndVertex = *SecondEdge->GetOtherVertex(*Vertex)->GetLinkActiveEntity();
-
-				if (&OtherEdgeEndVertex == &EndVertex)
-				{
-					FPoint StartTangentEdge = Edge->GetTangentAt(*Vertex);
-					FPoint StartTangentOtherEdge = SecondEdge->GetTangentAt(*Vertex);
-
-					double CosAngle = StartTangentEdge.ComputeCosinus(StartTangentOtherEdge);
-					if (CosAngle < 0.9)
-					{
-						continue;
-					}
-
-					if (Edge->GetFace() != SecondEdge->GetFace())
-					{
-						Edge->Link(*SecondEdge, SewTolerance);
-					}
-				}
-			}
+			MergeCoincidentEdges(Edge);
 		}
 	}
 
 	FDuration StepDuration = FChrono::Elapse(StartTime);
-
 	FChrono::PrintClockElapse(EVerboseLevel::Log, TEXT("    "), TEXT("Merge coincident edges"), StepDuration);
-
 }
 
 TSharedPtr<FTopologicalVertex> FTopomaker::SplitAndLink(FTopologicalVertex& StartVertex, FTopologicalEdge& EdgeToLink, FTopologicalEdge& EdgeToSplit)
@@ -519,7 +543,7 @@ void FTopomaker::StitchParallelEdges(TArray<FTopologicalVertex*>& VerticesToProc
 
 		TArray<FTopologicalEdge*> ConnectedEdges;
 		Vertex->GetConnectedEdges(ConnectedEdges);
-		int32 ConnectedEdgeCount = ConnectedEdges.Num();
+		const int32 ConnectedEdgeCount = ConnectedEdges.Num();
 		if (ConnectedEdgeCount == 1)
 		{
 			continue;
@@ -717,7 +741,7 @@ void FTopomaker::MergeUnconnectedAdjacentEdges()
 						Edge = nullptr;
 						continue;
 					}
-					
+
 					FPoint StartTangentOtherEdge = Edge->GetTangentAt(*StartVertex);
 					double CosAngle = StartTangentEdge.ComputeCosinus(StartTangentOtherEdge);
 					if (CosAngle > 0.9)
@@ -729,12 +753,6 @@ void FTopomaker::MergeUnconnectedAdjacentEdges()
 
 				if (ParallelEdge)
 				{
-					//FMessage::Printf(Log, TEXT("Face %d UnconnectedAdjacentEdges\n"), Face->GetId());
-					//for (auto Edge : Candidates)
-					//{
-					//	FMessage::Printf(Log, TEXT("       - %d \n"), Edge.Entity->GetId());
-					//}
-
 					TSharedPtr<FTopologicalEdge> NewEdge = FTopologicalEdge::CreateEdgeByMergingEdges(Candidates, StartVertex, EndVertex.ToSharedRef());
 					if (!NewEdge.IsValid())
 					{
@@ -807,7 +825,7 @@ void FTopomaker::UnlinkNonManifoldVertex()
 		{
 			ensureCADKernel(!TwinVertex->GetDirectConnectedEdges().IsEmpty());
 			FTopologicalEdge* Edge = TwinVertex->GetDirectConnectedEdges()[0];
-			FShell* Shell = (FShell*) Edge->GetLoop()->GetFace()->GetHost();
+			FShell* Shell = (FShell*)Edge->GetLoop()->GetFace()->GetHost();
 			if (Shell != nullptr)
 			{
 				ShellToVertexCount.FindOrAdd(Shell)++;
@@ -827,7 +845,7 @@ void FTopomaker::UnlinkNonManifoldVertex()
 			{
 				ensureCADKernel(!TwinVertex->GetDirectConnectedEdges().IsEmpty());
 				FTopologicalEdge* Edge = TwinVertex->GetDirectConnectedEdges()[0];
-				FShell* Shell = (FShell*) Edge->GetLoop()->GetFace()->GetHost();
+				FShell* Shell = (FShell*)Edge->GetLoop()->GetFace()->GetHost();
 				if (Shell != nullptr)
 				{
 					ShellToVertices.FindOrAdd(Shell).Add(TwinVertex);
@@ -851,17 +869,147 @@ void FTopomaker::UnlinkNonManifoldVertex()
 	}
 }
 
-// =========================================================================================================================================================================================================
-// =========================================================================================================================================================================================================
-// =========================================================================================================================================================================================================
-//
-//
-//                                                                            NOT YET REVIEWED
-//
-//
-// =========================================================================================================================================================================================================
-// =========================================================================================================================================================================================================
-// =========================================================================================================================================================================================================
+void FindBorderVertex(const TArray<FTopologicalEdge*>& Edges, TArray<FTopologicalVertex*>& OutBorderVertices)
+{
+	OutBorderVertices.Reserve(Edges.Num());
+
+	TFunction<void(FTopologicalVertex&)> AddVertex = [&](FTopologicalVertex& Vertex)
+	{
+		if (!Vertex.GetLinkActiveEntity()->HasMarker1())
+		{
+			OutBorderVertices.Add(&*Vertex.GetLinkActiveEntity());
+			Vertex.GetLinkActiveEntity()->SetMarker1();
+		}
+	};
+
+	for (FTopologicalEdge* Edge : Edges)
+	{
+		if (Edge->IsBorder())
+		{
+			AddVertex(*Edge->GetStartVertex());
+			AddVertex(*Edge->GetEndVertex());
+		}
+	}
+
+	for (FTopologicalVertex* Vertex : OutBorderVertices)
+	{
+		Vertex->ResetMarker1();
+	}
+
+}
+
+
+//#define DEBUG_THIN_FACE
+void FTopomaker::RemoveThinFaces()
+{
+	FTimePoint StartTime = FChrono::Now();
+
+	F3DDebugSession _(TEXT("RemoveThinFaces"));
+
+	// A thin surface is a surface no larger than ThinToleranceFactor * SewTolerance
+	constexpr double ThinToleranceFactor = 5.;
+	const double ThinTolerance = SewTolerance * ThinToleranceFactor;
+
+	TArray<FTopologicalEdge*> NewBorderEdges;
+	TArray<FTopologicalFace*> DeletedFaces;
+
+	// Find thin faces
+	for (TSharedPtr<FTopologicalFace> Face : Faces)
+	{
+		FFaceAnalyzer Analyer(*Face, ThinTolerance);
+		double GapSize = 0;
+		if (Analyer.IsThinFace(GapSize))
+		{
+			Face->Disjoin(NewBorderEdges);
+			DeletedFaces.Add(Face.Get());
+			Face->Delete();
+			Face->RemoveOfHost();
+#ifdef CADKERNEL_DEV
+			Report.AddThinFace();
+#endif
+		}
+	}
+
+#ifdef DEBUG_THIN_FACE
+	for (const FTopologicalFace* Face : DeletedFaces)
+	{
+		F3DDebugSession _(TEXT("Face"));
+		Display(*Face);
+	}
+#endif
+
+	// Remove deleted edges of NewBorderEdges
+	{
+		int32 NewIndex = 0;
+		for (int32 Index = 0; Index < NewBorderEdges.Num(); ++Index)
+		{
+			if (!NewBorderEdges[Index]->HasMarker1())
+			{
+				NewBorderEdges[NewIndex] = NewBorderEdges[Index];
+				NewIndex++;
+			}
+		}
+		NewBorderEdges.SetNum(NewIndex);
+	}
+
+	const double SmallEdgeTolerance = ThinTolerance * DOUBLE_HALF_PI; // in the case of the thin surface is a Half cylinder, the diameter is ThinTolerance, so the half circle circumference is ThinTolerance * Pi / 2
+	// Collapse edges with a length smaller than SmallEdgeTolerance
+	for (FTopologicalEdge* Edge : NewBorderEdges)
+	{
+		if (Edge->Length() < SmallEdgeTolerance)
+		{
+			Edge->GetStartVertex()->Link(*Edge->GetEndVertex());
+			Edge->SetAsDegenerated();
+		}
+	}
+
+	TArray<FTopologicalVertex*> BorderVertices;
+	BorderVertices.Reserve(NewBorderEdges.Num());
+	FindBorderVertex(NewBorderEdges, BorderVertices);
+
+	// MergeCoincidentVertices tolerance is increase of a factor ThinToleranceFactor * 2.1 because:
+	// to merge the vertices of thin border of the surface we need to set the tolerance to the size of the width of the surface.
+	// If two thin surface are connected, we need to merge the vertices of the both thin border i.e. 2 * max thin width
+	//
+	//            -------------------*                               -----------------
+	//                               |                                                \
+	//            -------------------*----------------------    =>   -------------------*----------------------
+	//                               |                                                    \
+	//                               *----------------------                               ---------------------
+	double TolValue = Tolerance;
+	SetTolerance(ThinTolerance * 2.1);
+	MergeCoincidentVertices(BorderVertices);
+	SetTolerance(TolValue);
+
+	MergeCoincidentEdges(NewBorderEdges);
+
+	StitchParallelEdges(BorderVertices);
+
+#ifdef CADKERNEL_DEV
+	Report.RemoveThinFacesDuration = FChrono::Elapse(StartTime);
+	FChrono::PrintClockElapse(EVerboseLevel::Log, TEXT(""), TEXT(".RemoveThinFaces"), Report.RemoveThinFacesDuration = FChrono::Elapse(StartTime));
+#endif
+}
+
+void FTopomaker::SetSelfConnectedEdgeDegenerated(TArray<FTopologicalVertex*>& Vertices)
+{
+	for (const FTopologicalVertex* Vertex : Vertices)
+	{
+		for (const FTopologicalVertex* TwinVertex : Vertex->GetTwinEntities())
+		{
+			for (FTopologicalEdge* Edge : TwinVertex->GetDirectConnectedEdges())
+			{
+				if (Edge->GetStartVertex()->IsLinkedTo(Edge->GetEndVertex()))
+				{
+					if (!Edge->IsDegenerated() && Edge->Length() < 2 * SewTolerance)
+					{
+						Edge->SetAsDegenerated();
+					}
+				}
+			}
+		}
+	}
+}
 
 void FTopomaker::CheckSelfConnectedEdge()
 {
@@ -1043,13 +1191,13 @@ void FTopomaker::SplitIntoConnectedShells()
 	{
 		if (FaceSubset.MainShell != nullptr)
 		{
-			FShell* Shell = (FShell*) FaceSubset.MainShell;
+			FShell* Shell = (FShell*)FaceSubset.MainShell;
 			Shell->Empty(FaceSubset.Faces.Num());
 			Shell->Add(FaceSubset.Faces);
 		}
 		else
 		{
-			FBody* Body = (FBody*) FaceSubset.MainBody;
+			FBody* Body = (FBody*)FaceSubset.MainBody;
 			if (Body == nullptr)
 			{
 				TSharedRef<FBody> SharedBody = FEntity::MakeShared<FBody>();
@@ -1094,7 +1242,7 @@ void FTopomaker::RemoveEmptyShells()
 	{
 		if (Shell->FaceCount() == 0)
 		{
-			FBody* Body = (FBody*) Shell->GetHost();
+			FBody* Body = (FBody*)Shell->GetHost();
 			if (Body != nullptr && !Body->HasMarker1())
 			{
 				Body->SetMarker1();
diff --git a/Engine/Source/Runtime/Datasmith/CADKernel/Private/CADKernel/UI/Display.cpp b/Engine/Source/Runtime/Datasmith/CADKernel/Private/CADKernel/UI/Display.cpp
--- a/Engine/Source/Runtime/Datasmith/CADKernel/Private/CADKernel/UI/Display.cpp
+++ b/Engine/Source/Runtime/Datasmith/CADKernel/Private/CADKernel/UI/Display.cpp
@@ -736,7 +736,7 @@ void Draw(const FTopologicalFace& Face)
 				switch (Edge.Entity->GetTwinEntityCount())
 				{
 				case 1:
-					Property = EVisuProperty::BorderEdge;
+					Property =  Edge.Entity->IsDegenerated() ? EVisuProperty::PurpleCurve : EVisuProperty::BorderEdge;
 					break;
 				case 2:
 					Property = EVisuProperty::BlueCurve;
@@ -921,14 +921,6 @@ void Display(const FModel& Model)
 		}
 	}
 
-	{
-		const TArray<TSharedPtr<FTopologicalFace>>& Surfaces = Model.GetFaces();
-		FProgress SurfaceProgress((int32)Surfaces.Num(), TEXT("Surfaces"));
-		for (const TSharedPtr<FTopologicalFace>& Surface : Surfaces)
-		{
-			Display(*Surface);
-		}
-	}
 #endif
 }
 
@@ -947,14 +939,6 @@ void DisplayProductTree(const FModel& Model)
 		}
 	}
 
-	{
-		const TArray<TSharedPtr<FTopologicalFace>>& Faces = Model.GetFaces();
-		FProgress SurfaceProgress((int32)Faces.Num(), TEXT("Faces"));
-		for (const TSharedPtr<FTopologicalFace>& Face : Faces)
-		{
-			Display(*Face);
-		}
-	}
 #endif
 }
 
diff --git a/Engine/Source/Runtime/Datasmith/CADKernel/Public/CADKernel/Geo/Curves/Curve.h b/Engine/Source/Runtime/Datasmith/CADKernel/Public/CADKernel/Geo/Curves/Curve.h
--- a/Engine/Source/Runtime/Datasmith/CADKernel/Public/CADKernel/Geo/Curves/Curve.h
+++ b/Engine/Source/Runtime/Datasmith/CADKernel/Public/CADKernel/Geo/Curves/Curve.h
@@ -104,6 +104,7 @@ public:
 	}
 
 	virtual TSharedPtr<FEntityGeom> ApplyMatrix(const FMatrixH& InMatrix) const override = 0;
+	virtual void Offset(const FPoint& OffsetDirection) = 0;
 
 	double GetLength(double Tolerance) const;
 
diff --git a/Engine/Source/Runtime/Datasmith/CADKernel/Public/CADKernel/Geo/Curves/RestrictionCurve.h b/Engine/Source/Runtime/Datasmith/CADKernel/Public/CADKernel/Geo/Curves/RestrictionCurve.h
--- a/Engine/Source/Runtime/Datasmith/CADKernel/Public/CADKernel/Geo/Curves/RestrictionCurve.h
+++ b/Engine/Source/Runtime/Datasmith/CADKernel/Public/CADKernel/Geo/Curves/RestrictionCurve.h
@@ -220,6 +220,15 @@ public:
 		Polyline.GetSubPolyline(InBoundary, Orientation, OutPoints);
 	}
 
+	/**
+	 * Get the sub polyline bounded by the input InBoundary in the orientation of the input InOrientation and append it to the output OutPoints
+	 */
+	template<class PointType>
+	void GetDiscretizationPoints(const FLinearBoundary& InBoundary, TArray<double>& OutCoordinates, TArray<PointType>& OutPoints) const
+	{
+		Polyline.GetSubPolyline(InBoundary, OutCoordinates, OutPoints);
+	}
+
 	/**
 	 * Samples the sub curve limited by the boundary respecting the input Desired segment length
 	 */
@@ -234,6 +243,22 @@ public:
 	}
 
 	void ExtendTo(const FPoint2D& Point) override;
+
+	bool IsIso(EIso Iso, double ErrorTolerance = DOUBLE_SMALL_NUMBER) const
+	{
+		return Polyline.IsIso(Iso, ErrorTolerance);
+	}
+
+	void Offset2D(const FPoint2D& OffsetDirection);
+
+	/**
+	 * must not be call
+	 */
+	virtual void Offset(const FPoint& OffsetDirection) override
+	{
+		ensure(false);
+	}
+
 };
 
 } // namespace UE::CADKernel
diff --git a/Engine/Source/Runtime/Datasmith/CADKernel/Public/CADKernel/Geo/Curves/SurfacicCurve.h b/Engine/Source/Runtime/Datasmith/CADKernel/Public/CADKernel/Geo/Curves/SurfacicCurve.h
--- a/Engine/Source/Runtime/Datasmith/CADKernel/Public/CADKernel/Geo/Curves/SurfacicCurve.h
+++ b/Engine/Source/Runtime/Datasmith/CADKernel/Public/CADKernel/Geo/Curves/SurfacicCurve.h
@@ -107,6 +107,14 @@ public:
 
 	virtual TSharedPtr<FEntityGeom> ApplyMatrix(const FMatrixH& InMatrix) const override;
 
+	/**
+	 * Must not be call
+	 */
+	virtual void Offset(const FPoint& OffsetDirection) override
+	{
+		ensure(false);
+	}
+
 	virtual void EvaluatePoint(double Coordinate, FCurvePoint& OutPoint, int32 DerivativeOrder = 0) const override;
 
 	virtual void Evaluate2DPoint(double Coordinate, FCurvePoint2D& OutPoint, int32 DerivativeOrder = 0) const override
diff --git a/Engine/Source/Runtime/Datasmith/CADKernel/Public/CADKernel/Geo/GeoEnum.h b/Engine/Source/Runtime/Datasmith/CADKernel/Public/CADKernel/Geo/GeoEnum.h
--- a/Engine/Source/Runtime/Datasmith/CADKernel/Public/CADKernel/Geo/GeoEnum.h
+++ b/Engine/Source/Runtime/Datasmith/CADKernel/Public/CADKernel/Geo/GeoEnum.h
@@ -84,6 +84,13 @@ enum EOrientation : uint8
 	Front = 1
 };
 
+enum ESituation : uint8
+{
+	Undefined = 0,
+	Inside,
+	Outside
+};
+
 #ifdef CADKERNEL_DEV
 extern const TCHAR* OrientationNames[];
 #endif
diff --git a/Engine/Source/Runtime/Datasmith/CADKernel/Public/CADKernel/Geo/Sampling/PolylineTools.h b/Engine/Source/Runtime/Datasmith/CADKernel/Public/CADKernel/Geo/Sampling/PolylineTools.h
--- a/Engine/Source/Runtime/Datasmith/CADKernel/Public/CADKernel/Geo/Sampling/PolylineTools.h
+++ b/Engine/Source/Runtime/Datasmith/CADKernel/Public/CADKernel/Geo/Sampling/PolylineTools.h
@@ -22,8 +22,8 @@ struct FPolylineBBox
 	FPoint Min;
 	FPoint MaxPoints[3];
 	FPoint MinPoints[3];
-	double CoordinateOfMaxPoint[3];
-	double CoordinateOfMinPoint[3];
+	double CoordinateOfMaxPoint[3] = { -HUGE_VALUE, -HUGE_VALUE, -HUGE_VALUE };
+	double CoordinateOfMinPoint[3] = { HUGE_VALUE, HUGE_VALUE, HUGE_VALUE };
 
 	FPolylineBBox()
 		: Max(-HUGE_VALUE, -HUGE_VALUE, -HUGE_VALUE)
@@ -169,7 +169,6 @@ protected:
 
 		double LastSegmentLength;
 		double Length = 0;
-		ensureCADKernel(BoundaryIndices[1] > BoundaryIndices[0]);
 
 		OutCurvilinearCoordinates.Add(0);
 		if (BoundaryIndices[1] > BoundaryIndices[0] + 1)
diff --git a/Engine/Source/Runtime/Datasmith/CADKernel/Public/CADKernel/Geo/Sampling/SurfacicPolyline.h b/Engine/Source/Runtime/Datasmith/CADKernel/Public/CADKernel/Geo/Sampling/SurfacicPolyline.h
--- a/Engine/Source/Runtime/Datasmith/CADKernel/Public/CADKernel/Geo/Sampling/SurfacicPolyline.h
+++ b/Engine/Source/Runtime/Datasmith/CADKernel/Public/CADKernel/Geo/Sampling/SurfacicPolyline.h
@@ -16,6 +16,7 @@
 #include "CADKernel/Utils/IndexOfCoordinateFinder.h"
 
 #include "Serialization/Archive.h"
+#include "Algo/AllOf.h"
 
 namespace UE::CADKernel
 {
@@ -320,6 +321,15 @@ public:
 		Approximator.GetSubPolyline(InBoundary, InOrientation, OutPoints);
 	}
 
+	/**
+	 * Get the sub polyline bounded by the input InBoundary in the orientation of the input InOrientation and append it to the output OutPoints
+	 */
+	void GetSubPolyline(const FLinearBoundary& InBoundary, TArray<double>& OutCoordinates, TArray<FPoint2D>& OutPoints) const
+	{
+		TPolylineApproximator<FPoint2D> Approximator(Coordinates, Points2D);
+		Approximator.GetSubPolyline(InBoundary, OutCoordinates, OutPoints);
+	}
+
 	/**
 	 * Get the sub polyline bounded by the input InBoundary in the orientation of the input InOrientation and append it to the output OutPoints
 	 */
@@ -329,6 +339,15 @@ public:
 		Approximator3D.GetSubPolyline(InBoundary, InOrientation, OutPoints);
 	}
 
+	/**
+	 * Get the sub polyline bounded by the input InBoundary in the orientation of the input InOrientation and append it to the output OutPoints
+	 */
+	void GetSubPolyline(const FLinearBoundary& InBoundary, TArray<double>& OutCoordinates, TArray<FPoint>& OutPoints) const
+	{
+		TPolylineApproximator<FPoint> Approximator3D(Coordinates, Points3D);
+		Approximator3D.GetSubPolyline(InBoundary, OutCoordinates, OutPoints);
+	}
+
 	/**
 	 * Reserves memory such that the polyline can contain at least Number elements.
 	 *
@@ -431,6 +450,14 @@ public:
 		return Approximator.ComputeLengthOfSubPolyline(InBoundary);
 	}
 
+	bool IsIso(EIso Iso, double ErrorTolerance = DOUBLE_SMALL_NUMBER) const
+	{
+		FPoint2D StartPoint = Points2D[0];
+		return Algo::AllOf(Points2D, [&](const FPoint2D& Point) { 
+			return FMath::IsNearlyEqual(Point[Iso], StartPoint[Iso], ErrorTolerance);
+			});
+	}
+
 };
 
 } // ns UE::CADKernel
diff --git a/Engine/Source/Runtime/Datasmith/CADKernel/Public/CADKernel/Math/Boundary.h b/Engine/Source/Runtime/Datasmith/CADKernel/Public/CADKernel/Math/Boundary.h
--- a/Engine/Source/Runtime/Datasmith/CADKernel/Public/CADKernel/Math/Boundary.h
+++ b/Engine/Source/Runtime/Datasmith/CADKernel/Public/CADKernel/Math/Boundary.h
@@ -309,6 +309,38 @@ public:
 		return UVBoundaries[EIso::IsoU].IsDegenerated() || UVBoundaries[EIso::IsoV].IsDegenerated();
 	}
 
+	ESituation IsInside(const FSurfacicBoundary& OtherBoundary, const FSurfacicTolerance& Tolerance2D)
+	{
+		int32 Inside = 0;
+		int32 Outside = 0;
+		TFunction<void(double, double, double)> CheckInside = [&](double LeftSide, double RigthSide, double Tolerance)
+		{
+			if (LeftSide + Tolerance < RigthSide)
+			{
+				Inside++;
+			}
+			else if (RigthSide + Tolerance < LeftSide)
+			{
+				Outside++;
+			}
+		};
+
+		CheckInside(OtherBoundary[EIso::IsoU].GetMin(), UVBoundaries[EIso::IsoU].GetMin(), Tolerance2D[EIso::IsoU]);
+		CheckInside(OtherBoundary[EIso::IsoV].GetMin(), UVBoundaries[EIso::IsoV].GetMin(), Tolerance2D[EIso::IsoV]);
+		CheckInside(UVBoundaries[EIso::IsoU].GetMax(), OtherBoundary[EIso::IsoU].GetMax(), Tolerance2D[EIso::IsoU]);
+		CheckInside(UVBoundaries[EIso::IsoV].GetMax(), OtherBoundary[EIso::IsoV].GetMax(), Tolerance2D[EIso::IsoV]);
+
+		if (Inside > 2)
+		{
+			return ESituation::Inside;
+		}
+		if (Outside > 2)
+		{
+			return ESituation::Outside;
+		}
+		return ESituation::Undefined;
+	}
+
 	/**
 	 * Uses to initiate a boundary computation with ExtendTo
 	 */
@@ -396,6 +428,7 @@ public:
 	{
 		return UVBoundaries[Iso];
 	}
+
 };
 }
 
diff --git a/Engine/Source/Runtime/Datasmith/CADKernel/Public/CADKernel/Math/Point.h b/Engine/Source/Runtime/Datasmith/CADKernel/Public/CADKernel/Math/Point.h
--- a/Engine/Source/Runtime/Datasmith/CADKernel/Public/CADKernel/Math/Point.h
+++ b/Engine/Source/Runtime/Datasmith/CADKernel/Public/CADKernel/Math/Point.h
@@ -326,6 +326,15 @@ public:
 
 	double SignedAngle(const FPoint& Vector2, const FPoint& Normal) const;
 
+	/**
+	 * Return the projection of the point on the diagonal axis (of vector (1,1,1))
+	 * i.e. return X + Y + Z
+	 */
+	double DiagonalAxisCoordinate() const
+	{
+		return X + Y + Z;
+	}
+
 	static double MixedTripleProduct(const FPoint& VectorA, const FPoint& VectorB, const FPoint& VectorC)
 	{
 		return VectorA * (VectorB ^ VectorC);
@@ -647,6 +656,15 @@ public:
 		return FPoint2D(-V, U);
 	}
 
+	/**
+	 * Return the projection of the point on the diagonal axis (of vector (1,1))
+	 * i.e. return U + V
+	 */
+	double DiagonalAxisCoordinate() const
+	{
+		return U + V;
+	}
+
 	friend FArchive& operator<<(FArchive& Ar, FPoint2D& Point)
 	{
 		Ar.Serialize(&Point, sizeof(FPoint2D));
diff --git a/Engine/Source/Runtime/Datasmith/CADKernel/Public/CADKernel/Topo/Body.h b/Engine/Source/Runtime/Datasmith/CADKernel/Public/CADKernel/Topo/Body.h
--- a/Engine/Source/Runtime/Datasmith/CADKernel/Public/CADKernel/Topo/Body.h
+++ b/Engine/Source/Runtime/Datasmith/CADKernel/Public/CADKernel/Topo/Body.h
@@ -76,6 +76,8 @@ public:
 
 	void RemoveEmptyShell(FModel& Model);
 
+	virtual void Remove(const FTopologicalShapeEntity* ShellToRemove) override;
+
 	void Empty()
 	{
 		Shells.Empty();
diff --git a/Engine/Source/Runtime/Datasmith/CADKernel/Public/CADKernel/Topo/Linkable.h b/Engine/Source/Runtime/Datasmith/CADKernel/Public/CADKernel/Topo/Linkable.h
--- a/Engine/Source/Runtime/Datasmith/CADKernel/Public/CADKernel/Topo/Linkable.h
+++ b/Engine/Source/Runtime/Datasmith/CADKernel/Public/CADKernel/Topo/Linkable.h
@@ -111,7 +111,7 @@ public:
 	{
 		ensureCADKernel(TopologicalLink.IsValid());
 		TopologicalLink->RemoveEntity((EntityType&)*this);
-		TopologicalLink.Reset();
+		ResetTopologicalLink();
 	}
 
 	/**
diff --git a/Engine/Source/Runtime/Datasmith/CADKernel/Public/CADKernel/Topo/Model.h b/Engine/Source/Runtime/Datasmith/CADKernel/Public/CADKernel/Topo/Model.h
--- a/Engine/Source/Runtime/Datasmith/CADKernel/Public/CADKernel/Topo/Model.h
+++ b/Engine/Source/Runtime/Datasmith/CADKernel/Public/CADKernel/Topo/Model.h
@@ -25,12 +25,10 @@ class CADKERNEL_API FModel : public FTopologicalShapeEntity
 protected:
 
 	TArray<TSharedPtr<FBody>> Bodies;
-	TArray<TSharedPtr<FTopologicalFace>> Faces;
 
 	FModel()
 	{
 		Bodies.Reserve(100);
-		Faces.Reserve(100);
 	}
 
 public:
@@ -39,7 +37,6 @@ public:
 	{
 		FTopologicalShapeEntity::Serialize(Ar);
 		SerializeIdents(Ar, (TArray<TSharedPtr<FEntity>>&) Bodies);
-		SerializeIdents(Ar, (TArray<TSharedPtr<FEntity>>&) Faces);
 
 		if (Ar.IsLoading())
 		{
@@ -59,14 +56,12 @@ public:
 		}
 
 		SpawnIdentOnEntities(Bodies, Database);
-		SpawnIdentOnEntities(Faces, Database);
 	}
 
 	virtual void ResetMarkersRecursively() override
 	{
 		ResetMarkers();
 		ResetMarkersRecursivelyOnEntities(Bodies);
-		ResetMarkersRecursivelyOnEntities(Faces);
 	}
 
 #ifdef CADKERNEL_DEV
@@ -80,16 +75,6 @@ public:
 
 	void AddEntity(TSharedRef<FTopologicalEntity> InEntity);
 
-	void Add(const TSharedPtr<FTopologicalFace>& InFace)
-	{
-		Faces.Add(InFace);
-	}
-
-	void Append(TArray<TSharedPtr<FTopologicalFace>>& InNewFaces)
-	{
-		Faces.Append(InNewFaces);
-	}
-
 	void Append(TArray<TSharedPtr<FBody>>& InNewBody)
 	{
 		Bodies.Append(InNewBody);
@@ -103,12 +88,6 @@ public:
 	void Empty()
 	{
 		Bodies.Empty();
-		Faces.Empty();
-	}
-
-	void RemoveFace(TSharedPtr<FTopologicalFace> InToplologicalFace)
-	{
-		Faces.Remove(InToplologicalFace);
 	}
 
 	void RemoveBody(TSharedPtr<FBody> InBody)
@@ -116,9 +95,9 @@ public:
 		Bodies.Remove(InBody);
 	}
 
-	void RemoveBody(const FBody* InBody)
+	virtual void Remove(const FTopologicalShapeEntity* InBody) override
 	{
-		int32 Index = Bodies.IndexOfByPredicate([&](const TSharedPtr<FBody>& Body) { return (InBody == Body.Get()); });
+		int32 Index = Bodies.IndexOfByPredicate([&](const TSharedPtr<FBody>& Body) { return (InBody == (FTopologicalShapeEntity*) Body.Get()); });
 		if (Index != INDEX_NONE)
 		{
 			Bodies.RemoveAt(Index);
@@ -137,7 +116,6 @@ public:
 	void Copy(const TSharedPtr<FModel>& OtherModel)
 	{
 		Bodies.Append(OtherModel->Bodies);
-		Faces.Append(OtherModel->Faces);
 	}
 
 	/**
@@ -146,7 +124,6 @@ public:
 	void Copy(const FModel& OtherModel)
 	{
 		Bodies.Append(OtherModel.Bodies);
-		Faces.Append(OtherModel.Faces);
 	}
 
 	/**
@@ -157,17 +134,11 @@ public:
 	{
 		Copy(OtherModel);
 		OtherModel.Bodies.Empty();
-		OtherModel.Faces.Empty();
 	}
 
 	int32 EntityCount() const
 	{
-		return Bodies.Num() + Faces.Num();
-	}
-
-	const TArray<TSharedPtr<FTopologicalFace>>& GetFaces() const
-	{
-		return Faces;
+		return Bodies.Num();
 	}
 
 	virtual void GetFaces(TArray<FTopologicalFace*>& OutFaces) override;
@@ -192,18 +163,6 @@ public:
 	virtual void FillTopologyReport(FTopologyReport& Report) const override;
 #endif
 
-	/**
-	 * Fore each body
-	 */
-	void FixModelTopology(double JoiningTolerance);
-
-	void MergeInto(TSharedPtr<FBody> Body, TArray<TSharedPtr<FTopologicalEntity>>& InEntities);
-
-	/**
-	 * Fore each shell of each body, try to stitch topological gap
-	 */
-	void HealModelTopology(double JoiningTolerance);
-
 	void Orient();
 };
 
diff --git a/Engine/Source/Runtime/Datasmith/CADKernel/Public/CADKernel/Topo/Shell.h b/Engine/Source/Runtime/Datasmith/CADKernel/Public/CADKernel/Topo/Shell.h
--- a/Engine/Source/Runtime/Datasmith/CADKernel/Public/CADKernel/Topo/Shell.h
+++ b/Engine/Source/Runtime/Datasmith/CADKernel/Public/CADKernel/Topo/Shell.h
@@ -174,6 +174,9 @@ public:
 	 * @return SwapFaceCount for report purpose
 	 */
 	int32 Orient();
+
+	virtual void Remove(const FTopologicalShapeEntity*) override;
+
 };
 
 }
diff --git a/Engine/Source/Runtime/Datasmith/CADKernel/Public/CADKernel/Topo/TopologicalEdge.h b/Engine/Source/Runtime/Datasmith/CADKernel/Public/CADKernel/Topo/TopologicalEdge.h
--- a/Engine/Source/Runtime/Datasmith/CADKernel/Public/CADKernel/Topo/TopologicalEdge.h
+++ b/Engine/Source/Runtime/Datasmith/CADKernel/Public/CADKernel/Topo/TopologicalEdge.h
@@ -388,9 +388,11 @@ public:
 		return EndVertex->GetCoordinates();
 	}
 
-
-
 	void GetTangentsAtExtremities(FPoint& StartTangent, FPoint& EndTangent, bool bForward) const;
+	void GetTangentsAtExtremities(FPoint& StartTangent, FPoint& EndTangent, EOrientation Orientation) const
+	{
+		GetTangentsAtExtremities(StartTangent, EndTangent, Orientation == EOrientation::Front);
+	}
 
 	// ======   Meshing Function   ======
 
@@ -626,6 +628,7 @@ public:
 		Curve->GetExtremities(Boundary, Extremities);
 	}
 
+	void Offset2D(const FPoint2D& OffsetDirection);
 
 	// ======   Geometrical Functions   ======
 
@@ -793,7 +796,7 @@ struct CADKERNEL_API FImposedCuttingPoint
 	/**
 	 * coordinate of the edge's mesh nodes
 	 */
-	double Coordinate;
+	double Coordinate = 0;
 	int32 OppositNodeIndex = -1;
 
 	FImposedCuttingPoint()
@@ -815,11 +818,11 @@ struct CADKERNEL_API FCuttingPoint
 	/**
 	 * coordinate of the edge's mesh nodes
 	 */
-	double Coordinate;
-	ECoordinateType Type;
+	double Coordinate = 0;
+	ECoordinateType Type = ECoordinateType::OtherCoordinate;
 	int32 OppositNodeIndex = -1;
 	int32 OppositNodeIndex2 = -1;
-	double IsoDeltaU;
+	double IsoDeltaU = 0;
 
 	FCuttingPoint()
 	{
diff --git a/Engine/Source/Runtime/Datasmith/CADKernel/Public/CADKernel/Topo/TopologicalFace.h b/Engine/Source/Runtime/Datasmith/CADKernel/Public/CADKernel/Topo/TopologicalFace.h
--- a/Engine/Source/Runtime/Datasmith/CADKernel/Public/CADKernel/Topo/TopologicalFace.h
+++ b/Engine/Source/Runtime/Datasmith/CADKernel/Public/CADKernel/Topo/TopologicalFace.h
@@ -40,12 +40,6 @@ enum class EQuadType : uint8
 	Other
 };
 
-class FBezierSurface;
-class FFaceMesh;
-class FGrid;
-class FSegmentCurve;
-class FThinZone;
-class FThinZoneFinder;
 struct FBBoxWithNormal;
 
 class CADKERNEL_API FTopologicalFace : public FTopologicalShapeEntity
@@ -202,6 +196,8 @@ public:
 		return Loops;
 	}
 
+	const TSharedPtr<FTopologicalLoop> GetExternalLoop() const;
+
 	/**
 	 * Get a sampling of each loop of the face
 	 * @param OutLoopSamplings an array of 2d points
@@ -246,6 +242,16 @@ public:
 		}
 	}
 
+	int32 EdgeCount() const
+	{
+		int32 EdgeNum = 0;
+		for (const TSharedPtr<FTopologicalLoop>& Loop : Loops)
+		{
+			EdgeNum += Loop->EdgeCount();
+		}
+		return EdgeNum;
+	}
+
 	// ======   Carrier Surface Functions   ======
 
 	TSharedRef<FSurface> GetCarrierSurface() const
@@ -284,9 +290,10 @@ public:
 	bool HasSameBoundariesAs(const TSharedPtr<FTopologicalFace>& OtherFace) const;
 
 	/**
-	 * Disconnects the face of its neighbors
+	 * Disconnects the face of its neighbors i.e. remove topological edge and vertex link with its neighbors
+	 * @param OutNewBorderEdges the neighbors edges
 	 */
-	void RemoveLinksWithNeighbours();
+	void Disjoin(TArray<FTopologicalEdge*>& OutNewBorderEdges);
 
 #ifdef CADKERNEL_DEV
 	virtual void FillTopologyReport(FTopologyReport& Report) const override;
@@ -417,6 +424,9 @@ public:
 		States &= ~EHaveStates::IsBackOriented;
 	}
 
+	virtual void Remove(const FTopologicalShapeEntity*) override
+	{
+	}
 
 	// =========================================================================================================================================================================================================
 	// =========================================================================================================================================================================================================
@@ -440,7 +450,7 @@ private:
 	TArray<int32> StartSideIndices;
 	TArray<FEdge2DProperties> SideProperties;
 	int32 NumOfMeshedSide = 0;
-	double LoopLength;
+	double LoopLength = -1.;
 	double LengthOfMeshedSide = 0;
 	double QuadCriteria = 0;
 	FSurfaceCurvature Curvatures;
@@ -581,8 +591,8 @@ struct FBBoxWithNormal
 	FPoint2D MinCoordinates[3];
 	FVector MaxPointNormals[3];
 	FVector MinPointNormals[3];
-	bool MaxNormalNeedUpdate[3];
-	bool MinNormalNeedUpdate[3];
+	bool MaxNormalNeedUpdate[3] = {true, true, true};
+	bool MinNormalNeedUpdate[3] = {true, true, true};
 
 	FBBoxWithNormal()
 		: Max(-HUGE_VALUE, -HUGE_VALUE, -HUGE_VALUE)
diff --git a/Engine/Source/Runtime/Datasmith/CADKernel/Public/CADKernel/Topo/TopologicalLoop.h b/Engine/Source/Runtime/Datasmith/CADKernel/Public/CADKernel/Topo/TopologicalLoop.h
--- a/Engine/Source/Runtime/Datasmith/CADKernel/Public/CADKernel/Topo/TopologicalLoop.h
+++ b/Engine/Source/Runtime/Datasmith/CADKernel/Public/CADKernel/Topo/TopologicalLoop.h
@@ -170,6 +170,16 @@ public:
 		return bIsExternal;
 	}
 
+	void SetExternal()
+	{
+		bIsExternal = true;
+	}
+
+	void SetInternal()
+	{
+		bIsExternal = false;
+	}
+
 	/*
 	 * @return false if the orientation is doubtful
 	 */
@@ -214,7 +224,7 @@ public:
 		return -1;
 	}
 
-	void Get2DSampling(TArray<FPoint2D>& LoopSampling);
+	void Get2DSampling(TArray<FPoint2D>& LoopSampling) const;
 
 	/**
 	 * The idea is to remove degenerated edges of the loop i.e. where the surface is degenerated
@@ -223,7 +233,7 @@ public:
 	 * - so where the sampling could be in self-intersecting
 	 * @return false if the loop is degenerated
 	 */
-	bool Get2DSamplingWithoutDegeneratedEdges(TArray<FPoint2D>& LoopSampling);
+	bool Get2DSamplingWithoutDegeneratedEdges(TArray<FPoint2D>& LoopSampling) const;
 
 	void FindSurfaceCorners(TArray<TSharedPtr<FTopologicalVertex>>& OutCorners, TArray<int32>& OutStartSideIndex) const;
 	void FindBreaks(TArray<TSharedPtr<FTopologicalVertex>>& Ruptures, TArray<int32>& OutStartSideIndex, TArray<double>& RuptureValues) const;
@@ -235,6 +245,9 @@ public:
 	void CheckEdgesOrientation();
 	void CheckLoopWithTwoEdgesOrientation();
 	void RemoveDegeneratedEdges();
+
+	bool IsInside(const FTopologicalLoop& Other) const;
+
 };
 
 }
diff --git a/Engine/Source/Runtime/Datasmith/CADKernel/Public/CADKernel/Topo/TopologicalShapeEntity.h b/Engine/Source/Runtime/Datasmith/CADKernel/Public/CADKernel/Topo/TopologicalShapeEntity.h
--- a/Engine/Source/Runtime/Datasmith/CADKernel/Public/CADKernel/Topo/TopologicalShapeEntity.h
+++ b/Engine/Source/Runtime/Datasmith/CADKernel/Public/CADKernel/Topo/TopologicalShapeEntity.h
@@ -125,6 +125,13 @@ public:
 		return Dictionary.GetPatchId();
 	}
 
+	void RemoveOfHost()
+	{
+		GetHost()->Remove(this);
+	}
+
+	virtual void Remove(const FTopologicalShapeEntity*) = 0;
+
 #ifdef CADKERNEL_DEV
 	virtual FInfoEntity& GetInfo(FInfoEntity&) const override;
 #endif
diff --git a/Engine/Source/Runtime/Datasmith/CADKernel/Public/CADKernel/Topo/TopologicalVertex.h b/Engine/Source/Runtime/Datasmith/CADKernel/Public/CADKernel/Topo/TopologicalVertex.h
--- a/Engine/Source/Runtime/Datasmith/CADKernel/Public/CADKernel/Topo/TopologicalVertex.h
+++ b/Engine/Source/Runtime/Datasmith/CADKernel/Public/CADKernel/Topo/TopologicalVertex.h
@@ -146,7 +146,7 @@ public:
 		{
 			TopologicalLink->RemoveEntity(*this);
 			TopologicalLink->ComputeBarycenter();
-			TopologicalLink.Reset();
+			ResetTopologicalLink();
 		}
 	}
 
diff --git a/Engine/Source/Runtime/Datasmith/CADKernel/Public/CADKernel/Topo/Topomaker.h b/Engine/Source/Runtime/Datasmith/CADKernel/Public/CADKernel/Topo/Topomaker.h
--- a/Engine/Source/Runtime/Datasmith/CADKernel/Public/CADKernel/Topo/Topomaker.h
+++ b/Engine/Source/Runtime/Datasmith/CADKernel/Public/CADKernel/Topo/Topomaker.h
@@ -19,8 +19,6 @@ class FTopologicalVertex;
 
 class CADKERNEL_API FTopomaker
 {
-	friend class FThinZone;
-	friend class FThinZoneFinder;
 
 protected:
 
@@ -29,8 +27,7 @@ protected:
 	TArray<FShell*> Shells;
 	TArray<TSharedPtr<FTopologicalFace>> Faces;
 
-	const double SewTolerance;
-	const double SewToleranceSquare;
+	double Tolerance;
 
 #ifdef CADKERNEL_DEV
 	FTopomakerReport Report;
@@ -42,6 +39,13 @@ public:
 	FTopomaker(FSession& InSession, const TArray<TSharedPtr<FShell>>& Shells, double Tolerance);
 	FTopomaker(FSession& InSession, const TArray<TSharedPtr<FTopologicalFace>>& Surfaces, double Tolerance);
 
+	void SetTolerance(double NewValue)
+	{
+		Tolerance = NewValue;
+		SewTolerance = NewValue * UE_DOUBLE_SQRT_2;
+		SewToleranceSquare = FMath::Square(SewTolerance);
+	}
+
 	void Sew();
 
 	/**
@@ -49,10 +53,6 @@ public:
 	 */
 	void CheckTopology();
 
-	//void MergeInto(TSharedPtr<FBody> Body, TArray<TSharedPtr<FTopologicalEntity>>& InEntities);
-	//void SortByShell(TSharedPtr<FBody> Body, TArray<TSharedPtr<FBody>>& OutNewBody);
-	//void Join(TArray<TSharedPtr<FBody>> Bodies, double Tolerance);
-
 	/**
 	 * Split into connected shell and put each shell into the appropriate body
 	 */
@@ -65,10 +65,17 @@ public:
 	 */
 	void UnlinkNonManifoldVertex();
 
+	void RemoveThinFaces();
+
 #ifdef CADKERNEL_DEV
-	void PrintReport()
+	void PrintSewReport()
 	{
-		Report.Print();
+		Report.PrintSewReport();
+	}
+
+	void PrintRemoveThinFacesReport()
+	{
+		Report.PrintRemoveThinFacesReport();
 	}
 #endif
 
@@ -112,10 +119,14 @@ private:
 	void CheckSelfConnectedEdge();
 	void RemoveIsolatedEdges(); // Useful ???
 
+	void SetSelfConnectedEdgeDegenerated(TArray<FTopologicalVertex*>& Vertices);
+
 	/**
 	 * First step, trivial edge merge i.e. couple of edges with same extremity vertex
 	 */
 	void MergeCoincidentEdges(TArray<FTopologicalVertex*>& Vertices);
+	void MergeCoincidentEdges(TArray<FTopologicalEdge*>& EdgesToProcess);
+	void MergeCoincidentEdges(FTopologicalEdge* Edge);
 
 	/**
 	 * Second step, parallel edges but with different length. the longest must be split
@@ -143,6 +154,9 @@ private:
 	 */
 	void MergeUnconnectedAdjacentEdges();
 
+	double SewTolerance;
+	double SewToleranceSquare;
+
 };
 
 } // namespace UE::CADKernel
