Fix Memory leak in "Datasmith Retessellate"
The memory leak came from Chaos::FTriangleMeshImplicitObject that are not release at each retessellation.
A workaround in the wait of a better solution from physics team is to flush the physicsScene at the end of retesselation step.
MeshDescriptions are also better released.

diff --git a/Engine/Plugins/Enterprise/DataprepEditor/Source/DataprepCore/Private/DataprepCorePrivateUtils.cpp b/Engine/Plugins/Enterprise/DataprepEditor/Source/DataprepCore/Private/DataprepCorePrivateUtils.cpp
--- a/Engine/Plugins/Enterprise/DataprepEditor/Source/DataprepCore/Private/DataprepCorePrivateUtils.cpp
+++ b/Engine/Plugins/Enterprise/DataprepEditor/Source/DataprepCore/Private/DataprepCorePrivateUtils.cpp
@@ -134,6 +134,10 @@ void DataprepCorePrivateUtils::BuildStaticMeshes(TSet<UStaticMesh*>& StaticMeshe
 					SourceModel.BuildSettings.DistanceFieldResolutionScale = 0;
 					//SourceModel.BuildSettings.bBuildReversedIndexBuffer = false;
 				}
+				
+				// As soon as StaticMeshes are built, the mesh description can be released
+				// This generate a significant freeing of memory
+				StaticMesh->ClearMeshDescription(Index);
 			}
 
 			StaticMeshesSettings.Add(MoveTemp(BuildSettings));				
diff --git a/Engine/Plugins/Enterprise/DatasmithCADImporter/Source/ParametricSurfaceExtension/ParametricSurfaceExtension.Build.cs b/Engine/Plugins/Enterprise/DatasmithCADImporter/Source/ParametricSurfaceExtension/ParametricSurfaceExtension.Build.cs
--- a/Engine/Plugins/Enterprise/DatasmithCADImporter/Source/ParametricSurfaceExtension/ParametricSurfaceExtension.Build.cs
+++ b/Engine/Plugins/Enterprise/DatasmithCADImporter/Source/ParametricSurfaceExtension/ParametricSurfaceExtension.Build.cs
@@ -24,6 +24,7 @@ namespace UnrealBuildTool.Rules
 					"DataprepCore",
 					"DatasmithContent",
 					"ParametricSurface",
+					"PhysicsCore",
 					"DatasmithImporter",
 					"Engine",
 					"StaticMeshEditor",
diff --git a/Engine/Plugins/Enterprise/DatasmithCADImporter/Source/ParametricSurfaceExtension/Private/ParametricRetessellateAction.cpp b/Engine/Plugins/Enterprise/DatasmithCADImporter/Source/ParametricSurfaceExtension/Private/ParametricRetessellateAction.cpp
--- a/Engine/Plugins/Enterprise/DatasmithCADImporter/Source/ParametricSurfaceExtension/Private/ParametricRetessellateAction.cpp
+++ b/Engine/Plugins/Enterprise/DatasmithCADImporter/Source/ParametricSurfaceExtension/Private/ParametricRetessellateAction.cpp
@@ -12,18 +12,20 @@
 #include "MeshDescriptionHelper.h"
 
 #include "Algo/AnyOf.h"
+#include "Algo/Transform.h"
 #include "AssetRegistry/AssetData.h"
 #include "Async/ParallelFor.h"
+#include "Chaos/ChaosScene.h"
 #include "Engine/StaticMesh.h"
 #include "HAL/PlatformFileManager.h"
 #include "IStaticMeshEditor.h"
-#include "StaticMeshAttributes.h"
 #include "Misc/FileHelper.h"
 #include "Misc/ScopedSlowTask.h"
+#include "StaticMeshAttributes.h"
 #include "Toolkits/IToolkit.h"
 #include "Toolkits/ToolkitManager.h"
 #include "UObject/StrongObjectPtr.h"
-#include "Algo/Transform.h"
+
 
 #define LOCTEXT_NAMESPACE "ParametricRetessellateAction"
 
@@ -52,7 +54,7 @@ void UParametricRetessellateAction::ApplyOnAssets(const TArray<FAssetData>& Sele
 
 bool FParametricRetessellateAction_Impl::CanApplyOnAssets(const TArray<FAssetData>& SelectedAssets)
 {
-	return Algo::AnyOf(SelectedAssets, [](const FAssetData& Asset){ return Datasmith::GetAdditionalData<UParametricSurfaceData>(Asset) != nullptr; });
+	return Algo::AnyOf(SelectedAssets, [](const FAssetData& Asset) { return Datasmith::GetAdditionalData<UParametricSurfaceData>(Asset) != nullptr; });
 }
 
 void FParametricRetessellateAction_Impl::ApplyOnAssets(const TArray<FAssetData>& SelectedAssets)
@@ -60,9 +62,13 @@ void FParametricRetessellateAction_Impl::ApplyOnAssets(const TArray<FAssetData>&
 	TFunction<void(UStaticMesh*)> FinalizeChanges = [](UStaticMesh* StaticMesh) -> void
 	{
 		StaticMesh->PostEditChange();
-		StaticMesh->MarkPackageDirty();
+		UStaticMesh::FCommitMeshDescriptionParams Params;
+		Params.bMarkPackageDirty = true;
+		Params.bUseHashAsGuid = true;
+		StaticMesh->CommitMeshDescription(0, Params);
+		StaticMesh->ClearMeshDescription(0);
 
-		// Refresh associated editor
+		// Refresh associated editor 
 		TSharedPtr<IToolkit> EditingToolkit = FToolkitManager::Get().FindEditorForAsset(StaticMesh);
 		if (IStaticMeshEditor* StaticMeshEditorInUse = StaticCastSharedPtr<IStaticMeshEditor>(EditingToolkit).Get())
 		{
@@ -77,7 +83,7 @@ void FParametricRetessellateAction_Impl::ApplyOnAssets(const TArray<FAssetData>&
 	bool bAskForSameOption = NumAssetsToProcess > 1;
 
 	TArray<UStaticMesh*> TessellatedMeshes;
-	TessellatedMeshes.Reserve( SelectedAssets.Num() );
+	TessellatedMeshes.Reserve(SelectedAssets.Num());
 
 	TUniquePtr<FScopedSlowTask> Progress;
 	int32 AssetIndex = -1;
@@ -134,17 +140,17 @@ void FParametricRetessellateAction_Impl::ApplyOnAssets(const TArray<FAssetData>&
 						Progress->EnterProgressFrame(1, Text);
 					}
 
-					if( StaticMesh->GetMeshDescription(0) == nullptr)
+					if (StaticMesh->GetMeshDescription(0) == nullptr)
 					{
-						StaticMesh->CreateMeshDescription( 0 );
+						StaticMesh->CreateMeshDescription(0);
 					}
 
-					if(StaticMesh->GetMeshDescription(0) != nullptr)
+					if (StaticMesh->GetMeshDescription(0) != nullptr)
 					{
 						StaticMesh->Modify();
-						StaticMesh->PreEditChange( nullptr );
+						StaticMesh->PreEditChange(nullptr);
 
-						if(ParametricSurfaceData != nullptr && ParametricSurfaceData->Tessellate(*StaticMesh, RetessellateOptions->Options))
+						if (ParametricSurfaceData != nullptr && ParametricSurfaceData->Tessellate(*StaticMesh, RetessellateOptions->Options))
 						{
 							TessellatedMeshes.Add(StaticMesh);
 						}
@@ -159,24 +165,28 @@ void FParametricRetessellateAction_Impl::ApplyOnAssets(const TArray<FAssetData>&
 	}
 
 	// Make sure lightmap settings are valid
-	if(TessellatedMeshes.Num() > 1)
+	if (TessellatedMeshes.Num() > 1)
 	{
 		ParallelFor(TessellatedMeshes.Num(), [&](int32 Index)
-		{
-			FDatasmithStaticMeshImporter::PreBuildStaticMesh(TessellatedMeshes[Index]); 
-		});
+			{
+				FDatasmithStaticMeshImporter::PreBuildStaticMesh(TessellatedMeshes[Index]);
+			});
 	}
-	else if(TessellatedMeshes.Num() > 0)
+	else if (TessellatedMeshes.Num() > 0)
 	{
-		FDatasmithStaticMeshImporter::PreBuildStaticMesh( TessellatedMeshes[0] );
+		FDatasmithStaticMeshImporter::PreBuildStaticMesh(TessellatedMeshes[0]);
 	}
 
-	FDatasmithStaticMeshImporter::BuildStaticMeshes( TessellatedMeshes );
+	FDatasmithStaticMeshImporter::BuildStaticMeshes(TessellatedMeshes);
 
-	for(UStaticMesh* StaticMesh : TessellatedMeshes)
+	for (UStaticMesh* StaticMesh : TessellatedMeshes)
 	{
 		FinalizeChanges(StaticMesh);
 	}
+
+	// Workaround to force to clean all references to the physics data and so to release old physics data.
+	// https://jira.it.epicgames.com/browse/UE-166555
+	GWorld->GetPhysicsScene()->Flush();
 }
 
 TSet<UStaticMesh*> GetReferencedStaticMeshes(const TArray<AActor*>& SelectedActors)
@@ -203,13 +213,13 @@ TSet<UStaticMesh*> GetReferencedStaticMeshes(const TArray<AActor*>& SelectedActo
 bool UParametricRetessellateAction::CanApplyOnActors(const TArray<AActor*>& SelectedActors)
 {
 	const TSet<UStaticMesh*> ReferencedStaticMeshes = GetReferencedStaticMeshes(SelectedActors);
-	return Algo::AnyOf(ReferencedStaticMeshes, [](const UStaticMesh* Mesh){ return Datasmith::GetAdditionalData<UParametricSurfaceData>(FAssetData(Mesh)); });
+	return Algo::AnyOf(ReferencedStaticMeshes, [](const UStaticMesh* Mesh) { return Datasmith::GetAdditionalData<UParametricSurfaceData>(FAssetData(Mesh)); });
 }
 
 void UParametricRetessellateAction::ApplyOnActors(const TArray<AActor*>& SelectedActors)
 {
 	TArray<FAssetData> AssetData;
-	Algo::Transform(GetReferencedStaticMeshes(SelectedActors), AssetData, [](UStaticMesh* Mesh){ return FAssetData(Mesh);});
+	Algo::Transform(GetReferencedStaticMeshes(SelectedActors), AssetData, [](UStaticMesh* Mesh) { return FAssetData(Mesh); });
 	return ApplyOnAssets(AssetData);
 }
 
diff --git a/Engine/Plugins/Enterprise/DatasmithCADImporter/Source/ParametricSurfaceExtension/Private/ParametricSurfaceBlueprintLibrary.cpp b/Engine/Plugins/Enterprise/DatasmithCADImporter/Source/ParametricSurfaceExtension/Private/ParametricSurfaceBlueprintLibrary.cpp
--- a/Engine/Plugins/Enterprise/DatasmithCADImporter/Source/ParametricSurfaceExtension/Private/ParametricSurfaceBlueprintLibrary.cpp
+++ b/Engine/Plugins/Enterprise/DatasmithCADImporter/Source/ParametricSurfaceExtension/Private/ParametricSurfaceBlueprintLibrary.cpp
@@ -8,8 +8,11 @@
 #include "DatasmithUtils.h"
 
 #include "AssetRegistry/AssetData.h"
+#include "Chaos/ChaosScene.h"
 #include "Engine/StaticMesh.h"
 #include "IStaticMeshEditor.h"
+#include "Physics/PhysScene.h"
+#include "PhysicsEngine/BodySetup.h"
 #include "StaticMeshAttributes.h"
 #include "Toolkits/ToolkitManager.h"
 
@@ -62,7 +65,12 @@ bool UParametricSurfaceBlueprintLibrary::RetessellateStaticMeshWithNotification(
 					FDatasmithStaticMeshImporter::BuildStaticMesh(StaticMesh);
 
 					StaticMesh->PostEditChange();
-					StaticMesh->MarkPackageDirty();
+
+					UStaticMesh::FCommitMeshDescriptionParams Params;
+					Params.bMarkPackageDirty = true;
+					Params.bUseHashAsGuid = true;
+					StaticMesh->CommitMeshDescription(0, Params);
+					StaticMesh->ClearMeshDescription(0);
 
 					// Refresh associated editor
 					TSharedPtr<IToolkit> EditingToolkit = FToolkitManager::Get().FindEditorForAsset(StaticMesh);
@@ -74,10 +82,25 @@ bool UParametricSurfaceBlueprintLibrary::RetessellateStaticMeshWithNotification(
 				// No posting required, just make sure the new tessellation is committed
 				else
 				{
+					// Invalidate physics data as the mesh is rebuild.
+					StaticMesh->GetBodySetup()->InvalidatePhysicsData();
+
 					UStaticMesh::FCommitMeshDescriptionParams Params;
 					Params.bMarkPackageDirty = false;
 					Params.bUseHashAsGuid = true;
 					StaticMesh->CommitMeshDescription( LODIndex, Params );
+					StaticMesh->ClearMeshDescription(0); // important
+
+					// Workaround to force to clean all references to the physics data and so to release old physics data.
+					// https://jira.it.epicgames.com/browse/UE-166555
+					for (int32 idx = 0; idx < GEditor->GetWorldContexts().Num(); ++idx)
+					{
+						const FWorldContext& Context = GEditor->GetWorldContexts()[idx];
+						if (FPhysScene* PhysScene = Context.World()->GetPhysicsScene())
+						{
+							PhysScene->Flush();
+						}
+					}
 				}
 
 				// Save last tessellation settings
