CAD Import should limit the number of material slots
The CVar ds.CADTranslator.MaxMaterialCountPerMesh define the limit.
If the limit is reached, the last material is used to replace the others.

diff --git a/Engine/Config/ConsoleVariables.ini b/Engine/Config/ConsoleVariables.ini
--- a/Engine/Config/ConsoleVariables.ini
+++ b/Engine/Config/ConsoleVariables.ini
@@ -119,7 +119,7 @@ fx.SkipVectorVMBackendOptimizations=1
 
 ; Uncomment to overwrite any existing cache associated with the file being imported. 
 ; Default is false
-;ds.CADTranslator.OverwriteCache=true
+; ds.CADTranslator.OverwriteCache=true
 
 ; Uncomment to enable/disable the timer that kill the worker if the import time is unusually long. With this time control, the load of the corrupted file 
 ; is canceled but the rest of the scene is imported. Default is enable
@@ -155,6 +155,11 @@ fx.SkipVectorVMBackendOptimizations=1
 ; Default value is 1.
 ;ds.CADTranslator.MeshingParameterFactor=1
 
+; The main UE5 rendering systems do not support more than 256 materials per mesh. 
+; This CVar allow to define the max material count per mesh.
+; Default value is 256
+;ds.CADTranslator.MaxMaterialCountPerMesh=256
+
 ; Alias Import Option
 ; Enable Sew action merges BReps according to their material i.e. only BReps associated with same material can be merged together.
 ; Default is false
diff --git a/Engine/Plugins/Enterprise/DatasmithCADImporter/Source/CADLibrary/Private/MeshDescriptionHelper.cpp b/Engine/Plugins/Enterprise/DatasmithCADImporter/Source/CADLibrary/Private/MeshDescriptionHelper.cpp
--- a/Engine/Plugins/Enterprise/DatasmithCADImporter/Source/CADLibrary/Private/MeshDescriptionHelper.cpp
+++ b/Engine/Plugins/Enterprise/DatasmithCADImporter/Source/CADLibrary/Private/MeshDescriptionHelper.cpp
@@ -16,6 +16,8 @@
 
 typedef uint32 TriangleIndex[3];
 
+DEFINE_LOG_CATEGORY_STATIC(LogCADLibrary, Log, All);
+
 namespace CADLibrary
 {
 
@@ -290,14 +292,23 @@ bool FillMesh(const FMeshParameters& MeshParameters, const FImportParameters& Im
 	}
 
 	// Add to the mesh, a polygon groups per material
+	int32 PolyGroupIndex = 0;
+	FPolygonGroupID PolyGroupID = 0;
 	for (auto& Material : MaterialToPolygonGroupMapping)
 	{
-		uint32 MaterialHash = Material.Key;
-		FName ImportedSlotName = *LexToString<uint32>(MaterialHash);
-
-		FPolygonGroupID PolyGroupID = MeshDescription.CreatePolygonGroup();
-		PolygonGroupImportedMaterialSlotNames[PolyGroupID] = ImportedSlotName;
+		if (PolyGroupIndex < FImportParameters::GMaxMaterialCountPerMesh)
+		{
+			uint32 MaterialHash = Material.Key;
+			FName ImportedSlotName = *LexToString<uint32>(MaterialHash);
+			PolyGroupID = MeshDescription.CreatePolygonGroup();
+			PolygonGroupImportedMaterialSlotNames[PolyGroupID] = ImportedSlotName;
+		}
+		else if (PolyGroupIndex == FImportParameters::GMaxMaterialCountPerMesh)
+		{
+			UE_LOG(LogCADLibrary, Warning, TEXT("The main UE5 rendering systems do not support more than 256 materials per mesh and the limit has been defined to %d materials. Only the first %d materials are kept. The others are replaced by the last one"), FImportParameters::GMaxMaterialCountPerMesh, FImportParameters::GMaxMaterialCountPerMesh);
+		}
 		Material.Value = PolyGroupID;
+		PolyGroupIndex++;
 	}
 
 	if (Algo::AnyOf(BodyTessellation.Faces, [](const FTessellationData& FaceTessellation) { return !FaceTessellation.TexCoordArray.IsEmpty(); }))
diff --git a/Engine/Plugins/Enterprise/DatasmithCADImporter/Source/CADTools/Private/CADOption.cpp b/Engine/Plugins/Enterprise/DatasmithCADImporter/Source/CADTools/Private/CADOption.cpp
--- a/Engine/Plugins/Enterprise/DatasmithCADImporter/Source/CADTools/Private/CADOption.cpp
+++ b/Engine/Plugins/Enterprise/DatasmithCADImporter/Source/CADTools/Private/CADOption.cpp
@@ -99,6 +99,15 @@ The used value of the meshing parameter is value * MeshingParameterFactor \n\
 Default value is 1.\n"),
 ECVF_Default);
 
+int32 FImportParameters::GMaxMaterialCountPerMesh = 256;
+FAutoConsoleVariableRef GCADTranslatorMaxMaterialCountPerMesh(
+	TEXT("ds.CADTranslator.MaxMaterialCountPerMesh"),
+	FImportParameters::GMaxMaterialCountPerMesh,
+	TEXT("The main UE5 rendering systems do not support more than 256 materials per mesh. This CVar allow to define the max material count per mesh.\n\
+Default value is 256\n"),
+ECVF_Default);
+
+
 uint32 GetTypeHash(const FImportParameters& ImportParameters)
 {
 	uint32 ParametersHash = ::GetTypeHash(ImportParameters.bGDisableCADKernelTessellation);
diff --git a/Engine/Plugins/Enterprise/DatasmithCADImporter/Source/CADTools/Public/CADOptions.h b/Engine/Plugins/Enterprise/DatasmithCADImporter/Source/CADTools/Public/CADOptions.h
--- a/Engine/Plugins/Enterprise/DatasmithCADImporter/Source/CADTools/Public/CADOptions.h
+++ b/Engine/Plugins/Enterprise/DatasmithCADImporter/Source/CADTools/Public/CADOptions.h
@@ -39,6 +39,7 @@ namespace CADLibrary
 		CADTOOLS_API static float GStitchingTolerance;
 		CADTOOLS_API static float GUnitScale;
 		CADTOOLS_API static float GMeshingParameterFactor;
+		CADTOOLS_API static int32 GMaxMaterialCountPerMesh;
 
 	public:
 		FImportParameters(FDatasmithUtils::EModelCoordSystem NewCoordinateSystem = FDatasmithUtils::EModelCoordSystem::ZUp_RightHanded)
@@ -89,6 +90,7 @@ namespace CADLibrary
 			Ar << ImportParameters.bGEnableCADCache;
 			Ar << ImportParameters.bGPreferJtFileEmbeddedTessellation;
 			Ar << ImportParameters.GStitchingTolerance;
+			Ar << ImportParameters.GMaxMaterialCountPerMesh;
 			return Ar;
 		}
 
diff --git a/Engine/Plugins/Enterprise/DatasmithCADImporter/Source/DatasmithWireTranslator/Private/DatasmithWireTranslator.cpp b/Engine/Plugins/Enterprise/DatasmithCADImporter/Source/DatasmithWireTranslator/Private/DatasmithWireTranslator.cpp
--- a/Engine/Plugins/Enterprise/DatasmithCADImporter/Source/DatasmithWireTranslator/Private/DatasmithWireTranslator.cpp
+++ b/Engine/Plugins/Enterprise/DatasmithCADImporter/Source/DatasmithWireTranslator/Private/DatasmithWireTranslator.cpp
@@ -195,9 +195,25 @@ public:
 		Label = InLabel;
 	}
 
+	FColor LastColor;
+	bool bWarningNotLogYet = true;
 	void Add(const TSharedPtr<AlDagNode>& Node, FColor Color, const FString& ShaderName)
 	{
-		ShaderNames.Add(ShaderName);
+		using namespace CADLibrary;
+		if (ShaderNames.Num() == FImportParameters::GMaxMaterialCountPerMesh)
+		{
+			if (!ShaderNames.Contains(ShaderName))
+			{
+			UE_LOG(LogDatasmithWireTranslator, Warning, TEXT("The main UE5 rendering systems do not support more than 256 materials per mesh and the limit has been defined to %d materials. Only the first %d materials are kept. The others are replaced by the last one"), FImportParameters::GMaxMaterialCountPerMesh, FImportParameters::GMaxMaterialCountPerMesh);
+				bWarningNotLogYet = false;
+				Color = LastColor;
+			}
+		}
+		else
+		{
+			ShaderNames.Add(ShaderName);
+			LastColor = Color;
+		}
 		Shells.Emplace(Node, Color);
 	}
 
