Fix for CIS error

diff --git a/Engine/Plugins/Enterprise/DatasmithCADImporter/Source/CADLibrary/Private/CADKernelTools.cpp b/Engine/Plugins/Enterprise/DatasmithCADImporter/Source/CADLibrary/Private/CADKernelTools.cpp
--- a/Engine/Plugins/Enterprise/DatasmithCADImporter/Source/CADLibrary/Private/CADKernelTools.cpp
+++ b/Engine/Plugins/Enterprise/DatasmithCADImporter/Source/CADLibrary/Private/CADKernelTools.cpp
@@ -286,7 +286,7 @@ namespace CADLibrary
 		TSharedRef<FModelMesh> CADKernelModelMesh = FEntity::MakeShared<FModelMesh>();
 		FParametricMesher Mesher(*CADKernelModelMesh);
 
-	DefineMeshCriteria(*CADKernelModelMesh, TessellationContext.ImportParameters, GeometricTolerance);
+		DefineMeshCriteria(*CADKernelModelMesh, ImportParameters, 0.01);
 		Mesher.MeshEntity(CADTopologicalEntity);
 
 		FMeshConversionContext Context(ImportParameters, MeshParameters);
@@ -406,13 +406,13 @@ namespace CADLibrary
 
 		if (ImportParameters.GetMaxEdgeLength() > SMALL_NUMBER)
 		{
-		TSharedPtr<UE::CADKernel::FCriterion> MaxSizeCriterion = UE::CADKernel::FCriterion::CreateCriterion(UE::CADKernel::ECriterion::MaxSize, ImportParameters.GetMaxEdgeLength());
+			TSharedPtr<UE::CADKernel::FCriterion> MaxSizeCriterion = UE::CADKernel::FCriterion::CreateCriterion(UE::CADKernel::ECriterion::MaxSize, ImportParameters.GetMaxEdgeLength());
 			MeshModel.AddCriterion(MaxSizeCriterion);
 		}
 
 		if (ImportParameters.GetChordTolerance() > SMALL_NUMBER)
 		{
-		TSharedPtr<UE::CADKernel::FCriterion> ChordCriterion = UE::CADKernel::FCriterion::CreateCriterion(UE::CADKernel::ECriterion::Sag, ImportParameters.GetChordTolerance());
+			TSharedPtr<UE::CADKernel::FCriterion> ChordCriterion = UE::CADKernel::FCriterion::CreateCriterion(UE::CADKernel::ECriterion::Sag, ImportParameters.GetChordTolerance());
 			MeshModel.AddCriterion(ChordCriterion);
 		}
 
