Fix the update of material slot names of MeshDescription. The cash was due to a different size of the MeshDescription before and after a conversion into a DynamicMesh
Furthermore, the initial process build also a unwanted set of PolygoneGroup du to the call of CADLibrary::CopyPatchGroups

CADLibrary::CopyPatchGroups was a wrong way to save the patch groups to retesselate with the RetessellateOptions.RetessellationRule "EDatasmithCADRetessellationRule::SkipDeletedSurfaces"
The retessellation process is improved by saving the the patch groups to retesselate into FMeshConversionContext

This fix also takes care to:
- not create unused MaterialSlot.
- not create empty "BodyMesh::Face"

diff --git a/Engine/Plugins/Enterprise/DatasmithCADImporter/Source/CADInterfaces/Private/TechSoftFileParser.cpp b/Engine/Plugins/Enterprise/DatasmithCADImporter/Source/CADInterfaces/Private/TechSoftFileParser.cpp
--- a/Engine/Plugins/Enterprise/DatasmithCADImporter/Source/CADInterfaces/Private/TechSoftFileParser.cpp
+++ b/Engine/Plugins/Enterprise/DatasmithCADImporter/Source/CADInterfaces/Private/TechSoftFileParser.cpp
@@ -162,7 +162,7 @@ void TraverseAttribute(const A3DMiscAttributeData& AttributeData, TMap<FString,
 			{
 				AttributeName += TEXT(" ") + AttributeTitle;
 			}
-			else if (Index > 0)
+			else if(Index > 0)
 			{
 				AttributeName += TEXT(" ") + FString::FromInt((int32)Index);
 			}
@@ -331,7 +331,7 @@ ECADParsingResult FTechSoftFileParser::Process()
 		return ECADParsingResult::FileNotFound;
 	}
 
-	A3DImport Import(TCHAR_TO_UTF8(*File.GetPathOfFileToLoad()));
+	A3DImport Import(TCHAR_TO_UTF8(*File.GetPathOfFileToLoad())); 
 
 	TechSoftFileParserImpl::SetIOOption(Import);
 
@@ -439,7 +439,7 @@ void FTechSoftFileParser::SewModel()
 {
 	CADLibrary::TUniqueTSObj<A3DSewOptionsData> SewData;
 	SewData->m_bComputePreferredOpenShellOrientation = false;
-
+	
 	TechSoftInterface::SewModel(ModelFile.Get(), CADLibrary::FImportParameters::GStitchingTolerance, SewData.GetPtr());
 }
 
@@ -457,7 +457,6 @@ void FTechSoftFileParser::GenerateBodyMeshes()
 void FTechSoftFileParser::GenerateBodyMesh(A3DRiRepresentationItem* Representation, FArchiveBody& Body)
 {
 	FBodyMesh& BodyMesh = CADFileData.AddBodyMesh(Body.Id, Body);
-	BodyMesh.AddGraphicPropertiesFrom(Body);
 
 	uint32 NewBRepCount = 0;
 	A3DRiBrepModel** NewBReps = nullptr;
@@ -480,7 +479,7 @@ void FTechSoftFileParser::GenerateBodyMesh(A3DRiRepresentationItem* Representati
 		{
 			TechSoftUtils::FillBodyMesh(NewBReps[Index], CADFileData.GetImportParameters(), Body.Unit, BodyMesh);
 		}
-	}
+	} 
 	else
 	{
 		TechSoftUtils::FillBodyMesh(Representation, CADFileData.GetImportParameters(), Body.Unit, BodyMesh);
@@ -495,6 +494,12 @@ void FTechSoftFileParser::GenerateBodyMesh(A3DRiRepresentationItem* Representati
 	// Convert material
 	for (FTessellationData& Tessellation : BodyMesh.Faces)
 	{
+		// Don't add material of empty face (i.e. face with empty tessellation
+		if (Tessellation.VertexIndices.Num() == 0)
+		{
+			continue;
+		}
+
 		// Extract proper color or material based on style index
 		FMaterialUId CachedStyleIndex = Tessellation.MaterialUId;
 		Tessellation.MaterialUId = 0;
@@ -536,11 +541,11 @@ void FTechSoftFileParser::GenerateBodyMesh(A3DRiRepresentationItem* Representati
 			// This will be used when the file is reloaded
 			JsonObject->SetNumberField(JSON_ENTRY_BODY_UNIT, Body.Unit);
 
-			if (Body.ColorUId)
+			if(Body.ColorUId)
 			{
 				JsonObject->SetNumberField(JSON_ENTRY_COLOR_NAME, Body.ColorUId);
 			}
-			if (Body.MaterialUId)
+			if(Body.MaterialUId)
 			{
 				JsonObject->SetNumberField(JSON_ENTRY_MATERIAL_NAME, Body.MaterialUId);
 			}
@@ -605,7 +610,7 @@ ECADParsingResult FTechSoftFileParser::TraverseModel()
 	ExtractSpecificMetaData(ModelFile.Get(), Reference);
 	Reference.Unit = FileUnit;
 
-	if (ModelFileData->m_uiPOccurrencesSize == 0)
+	if(ModelFileData->m_uiPOccurrencesSize == 0)
 	{
 		CADFileData.AddWarningMessages(FString::Printf(TEXT("File %s is empty."), *CADFileData.GetCADFileDescription().GetFileName()));
 		return ECADParsingResult::ProcessFailed;
@@ -785,7 +790,7 @@ void FTechSoftFileParser::TraverseReference(const A3DAsmProductOccurrence* A3DRe
 	FMatrix ReferenceMatrix = FMatrix::Identity;
 	A3DMiscTransformation* Transform = ReferenceData->m_pLocation;
 	ExtractTransformation(Transform, Reference);
-
+	
 	Reference.TransformMatrix = Reference.TransformMatrix * ReferenceMatrix;
 
 	FArchiveInstance EmptyInstance;
@@ -848,11 +853,11 @@ void FTechSoftFileParser::TraverseOccurrence(const A3DAsmProductOccurrence* Occu
 	BuildInstanceName(Instance, ParentReference);
 
 	A3DMiscTransformation* Transform = OccurrenceData->m_pLocation;
-
+	
 	A3DAsmProductOccurrence* ReferencePtr = ProductOccurrence::GetReference(*OccurrenceData);
 
 	// Is the Reference already processed ?
-	if (ReferencePtr)
+	if(ReferencePtr)
 	{
 		FCadId* ReferenceId = ReferenceCache.Find(ReferencePtr);
 		if (ReferenceId != nullptr)
@@ -921,7 +926,7 @@ void FTechSoftFileParser::TraverseOccurrence(const A3DAsmProductOccurrence* Occu
 		ProcessReference(OccurrencePtr, Instance, NewReference);
 	}
 
-	if (ReferencePtr)
+	if(ReferencePtr)
 	{
 		ReferenceCache.Add(ReferencePtr, Instance.ReferenceNodeId);
 	}
@@ -944,7 +949,7 @@ void FTechSoftFileParser::ProcessReference(const A3DAsmProductOccurrence* Occurr
 		OccurrenceData.FillFrom(CachedOccurrencePtr);
 	}
 
-	if (ProductOccurrence::HasNoPartNoChild(*OccurrenceData))
+	if(ProductOccurrence::HasNoPartNoChild(*OccurrenceData))
 	{
 		return;
 	}
@@ -985,7 +990,7 @@ void FTechSoftFileParser::CountUnderOccurrence(const A3DAsmProductOccurrence* Oc
 		A3DAsmProductOccurrence* ReferencePtr = ProductOccurrence::GetReference(*OccurrenceData);
 
 		// Is the Reference already processed ?
-		if (ReferencePtr)
+		if(ReferencePtr)
 		{
 			FCadId* ReferenceId = ReferenceCache.Find(ReferencePtr);
 			if (ReferenceId != nullptr)
@@ -1069,8 +1074,8 @@ void FTechSoftFileParser::ProcessPrototype(const A3DAsmProductOccurrence* InProt
 				}
 			}
 		}
-
-		if (ProductOccurrence::HasPartOrChild(*PrototypeData))
+		
+		if (ProductOccurrence::HasPartOrChild(*PrototypeData) )
 		{
 			OutReference.bIsUnloaded = false;
 			PrototypePtr = nullptr;
@@ -1100,7 +1105,7 @@ void FTechSoftFileParser::ProcessPrototype(const A3DAsmProductOccurrence* InProt
 
 	if (OutReference.bIsUnloaded)
 	{
-		if (OutReference.Label.IsEmpty())
+		if(OutReference.Label.IsEmpty())
 		{
 			OutReference.Label = OutReference.ExternalFile.GetFileName();
 		}
@@ -1303,7 +1308,7 @@ void FTechSoftFileParser::ExtractMetaData(const A3DEntity* Entity, FArchiveCADOb
 		if (OutObject.Label.IsEmpty() && MetaData->m_pcName && MetaData->m_pcName[0] != '\0')
 		{
 			FString Name = UTF8_TO_TCHAR(MetaData->m_pcName);
-			if (Name != TEXT("unnamed"))  // "unnamed" is create by Techsoft. This name is ignored 
+			if(Name != TEXT("unnamed"))  // "unnamed" is create by Techsoft. This name is ignored 
 			{
 				Name = TechSoftUtils::CleanLabel(Name);
 				OutObject.Label = Name;
@@ -1388,7 +1393,7 @@ void FTechSoftFileParser::BuildInstanceName(FArchiveInstance& InstanceData, cons
 		case ECADFormat::CATIA:
 			InstanceData.Label = TechSoftUtils::CleanCatiaInstanceLabel(InstanceData.Label);
 			break;
-
+		
 		case ECADFormat::CATIA_3DXML:
 			InstanceData.Label = TechSoftUtils::Clean3dxmlInstanceLabel(InstanceData.Label);
 			break;
@@ -1424,14 +1429,14 @@ void FTechSoftFileParser::BuildBodyName(FArchiveBody& Body, const FArchiveRefere
 		Body.Label = TechSoftUtils::CleanCreoLabel(Body.Label);
 	}
 
-	if (Body.IsNameDefined())
+	if(Body.IsNameDefined())
 	{
 		return;
 	}
 
 	if (Format == ECADFormat::CATIA)
 	{
-		if (Body.SetNameWithAttributeValue(TEXT("BodyID")))
+		if(Body.SetNameWithAttributeValue(TEXT("BodyID")))
 		{
 			return;
 		}
@@ -1543,7 +1548,7 @@ void FTechSoftFileParser::ExtractSpecificMetaData(const A3DAsmProductOccurrence*
 				A3DElementsByRefsetUg Refset = UnigraphicsSpecificData->m_asChildrenByRefsets[Index];
 			}
 
-			if (UnigraphicsSpecificData->m_uiSolidsByRefsetsSize)
+			if(UnigraphicsSpecificData->m_uiSolidsByRefsetsSize)
 			{
 				for (uint32 Index = 0; Index < UnigraphicsSpecificData->m_uiSolidsByRefsetsSize; ++Index)
 				{
@@ -1601,7 +1606,7 @@ FArchiveColor& FTechSoftFileParser::FindOrAddColor(uint32 ColorIndex, uint8 Alph
 	FArchiveColor& NewColor = CADFileData.AddColor(ColorHId);
 	NewColor.Color = TechSoftUtils::GetColorAt(ColorIndex);
 	NewColor.Color.A = Alpha;
-
+	
 	NewColor.UEMaterialUId = BuildColorUId(NewColor.Color);
 	return NewColor;
 }
@@ -1612,15 +1617,15 @@ FArchiveMaterial& FTechSoftFileParser::AddMaterialAt(uint32 MaterialIndexToSave,
 	FCADMaterial& Material = NewMaterial.Material;
 
 	TUniqueTSObjFromIndex<A3DGraphMaterialData> MaterialData(MaterialIndexToSave);
-	if (MaterialData.IsValid())
+	if(MaterialData.IsValid())
 	{
 		Material.Diffuse = TechSoftUtils::GetColorAt(MaterialData->m_uiDiffuse);
 		Material.Ambient = TechSoftUtils::GetColorAt(MaterialData->m_uiAmbient);
 		Material.Specular = TechSoftUtils::GetColorAt(MaterialData->m_uiSpecular);
 		Material.Shininess = MaterialData->m_dShininess;
-		if (GraphStyleData.m_bIsTransparencyDefined)
+		if(GraphStyleData.m_bIsTransparencyDefined)
 		{
-			Material.Transparency = 1. - GraphStyleData.m_ucTransparency / 255.;
+			Material.Transparency = 1. - GraphStyleData.m_ucTransparency/255.;
 		}
 		// todo: find how to convert Emissive color into ? reflexion coef...
 		// Material.Emissive = GetColor(MaterialData->m_uiEmissive);
@@ -1647,7 +1652,7 @@ FArchiveMaterial& FTechSoftFileParser::FindOrAddMaterial(FMaterialUId MaterialIn
 		if (TextureData.IsValid())
 		{
 			return AddMaterialAt(TextureData->m_uiMaterialIndex, MaterialIndex, GraphStyleData);
-
+			
 #ifdef NOTYETDEFINE
 			TUniqueTSObj<A3DGraphTextureDefinitionData> TextureDefinitionData(TextureData->m_uiTextureDefinitionIndex);
 			if (TextureDefinitionData.IsValid())
@@ -1788,17 +1793,17 @@ void FTechSoftFileParser::ExtractGeneralTransformation(const A3DMiscTransformati
 		}
 		else
 		{
-			FVector3d Translation = Transform.GetTranslation();
-			Translation *= FImportParameters::GUnitScale;
+		FVector3d Translation = Transform.GetTranslation();
+		Translation *= FImportParameters::GUnitScale;
 
-			double UniformScale = TechSoftFileParserImpl::ExtractUniformScale(Scale);
+		double UniformScale = TechSoftFileParserImpl::ExtractUniformScale(Scale);
 			Component.Unit *= UniformScale;
 
-			FQuat4d Rotation = Transform.GetRotation();
+		FQuat4d Rotation = Transform.GetRotation();
 
-			FTransform3d NewTransform;
-			NewTransform.SetScale3D(Scale);
-			NewTransform.SetRotation(Rotation);
+		FTransform3d NewTransform;
+		NewTransform.SetScale3D(Scale);
+		NewTransform.SetRotation(Rotation);
 
 			Component.TransformMatrix = NewTransform.ToMatrixWithScale();
 			Component.TransformMatrix.SetOrigin(Translation);
diff --git a/Engine/Plugins/Enterprise/DatasmithCADImporter/Source/CADInterfaces/Private/TechSoftInterfaceUtils.inl b/Engine/Plugins/Enterprise/DatasmithCADImporter/Source/CADInterfaces/Private/TechSoftInterfaceUtils.inl
--- a/Engine/Plugins/Enterprise/DatasmithCADImporter/Source/CADInterfaces/Private/TechSoftInterfaceUtils.inl
+++ b/Engine/Plugins/Enterprise/DatasmithCADImporter/Source/CADInterfaces/Private/TechSoftInterfaceUtils.inl
@@ -165,6 +165,15 @@ private:
 		for (unsigned int Index = 0; Index < A3DTessellationData->m_uiFaceTessSize; ++Index)
 		{
 			const A3DTessFaceData& FaceTessData = A3DTessellationData->m_psFaceTessData[Index];
+
+			uint32 FaceSetIndex = 0;
+			bool bMustProcess = FaceTessData.m_uiSizesTriangulatedSize > FaceSetIndex;
+			// don't create empty TessellationData
+			if (!bMustProcess)
+			{
+				continue;
+			}
+
 			FTessellationData& Tessellation = Faces.Emplace_GetRef();
 
 			// there is a bijection between A3DTess3DData->m_psFaceTessData and A3DTopoShellData->m_ppFaces
@@ -192,9 +201,6 @@ private:
 			uint32 LastTrianguleIndex = FaceTessData.m_uiStartTriangulated;
 			LastVertexIndex = 0;
 
-			uint32 FaceSetIndex = 0;
-			bool bMustProcess = FaceTessData.m_uiSizesTriangulatedSize > FaceSetIndex;
-
 			if (bMustProcess && UsedEntitiesFlags & kA3DTessFaceDataTriangle)
 			{
 				AddFaceTriangle(Tessellation, FaceTessData.m_puiSizesTriangulated[FaceSetIndex++], LastTrianguleIndex, LastVertexIndex);
diff --git a/Engine/Plugins/Enterprise/DatasmithCADImporter/Source/CADKernelSurface/Private/CADKernelSurfaceExtension.cpp b/Engine/Plugins/Enterprise/DatasmithCADImporter/Source/CADKernelSurface/Private/CADKernelSurfaceExtension.cpp
--- a/Engine/Plugins/Enterprise/DatasmithCADImporter/Source/CADKernelSurface/Private/CADKernelSurfaceExtension.cpp
+++ b/Engine/Plugins/Enterprise/DatasmithCADImporter/Source/CADKernelSurface/Private/CADKernelSurfaceExtension.cpp
@@ -46,6 +46,8 @@ bool UCADKernelParametricSurfaceData::Tessellate(UStaticMesh& StaticMesh, const
 	CadMeshParameters.SymmetricNormal = (FVector3f) MeshParameters.SymmetricNormal;
 	CadMeshParameters.SymmetricOrigin = (FVector3f) MeshParameters.SymmetricOrigin;
 
+	CADLibrary::FMeshConversionContext MeshConversionContext(ImportParameters, CadMeshParameters);
+
 	// Previous MeshDescription is get to be able to create a new one with the same order of PolygonGroup (the matching of color and partition is currently based on their order)
 	if (FMeshDescription* DestinationMeshDescription = StaticMesh.GetMeshDescription(0))
 	{
@@ -55,7 +57,7 @@ bool UCADKernelParametricSurfaceData::Tessellate(UStaticMesh& StaticMesh, const
 
 		if (RetessellateOptions.RetessellationRule == EDatasmithCADRetessellationRule::SkipDeletedSurfaces)
 		{
-			CADLibrary::CopyPatchGroups(*DestinationMeshDescription, MeshDescription);
+			CADLibrary::GetExistingPatches(*DestinationMeshDescription, MeshConversionContext.PatchesToMesh);
 		}
 
 		const double GeometricTolerance = 0.01; // mm
@@ -69,7 +71,7 @@ bool UCADKernelParametricSurfaceData::Tessellate(UStaticMesh& StaticMesh, const
 			return bSuccessfulTessellation;
 		}
 
-		if(CADLibrary::FCADKernelTools::Tessellate(CADKernelModel, ImportParameters, CadMeshParameters, MeshDescription))
+		if(CADLibrary::FCADKernelTools::Tessellate(CADKernelModel, MeshConversionContext, MeshDescription))
 		{
 			// To update the SectionInfoMap 
 			TPolygonGroupAttributesConstRef<FName> MaterialSlotNames = MeshDescriptionAttributes.GetPolygonGroupMaterialSlotNames();
diff --git a/Engine/Plugins/Enterprise/DatasmithCADImporter/Source/CADKernelSurface/Public/CADModelToCADKernelConverterBase.h b/Engine/Plugins/Enterprise/DatasmithCADImporter/Source/CADKernelSurface/Public/CADModelToCADKernelConverterBase.h
--- a/Engine/Plugins/Enterprise/DatasmithCADImporter/Source/CADKernelSurface/Public/CADModelToCADKernelConverterBase.h
+++ b/Engine/Plugins/Enterprise/DatasmithCADImporter/Source/CADKernelSurface/Public/CADModelToCADKernelConverterBase.h
@@ -9,6 +9,7 @@
 #include "CADModelConverter.h"
 #include "CADOptions.h"
 #include "IDatasmithSceneElements.h"
+#include "MeshDescriptionHelper.h"
 
 #include "CADKernel/Core/Session.h"
 #include "CADKernel/Topo/Model.h"
@@ -64,7 +65,10 @@ public:
 	virtual bool Tessellate(const CADLibrary::FMeshParameters& InMeshParameters, FMeshDescription& OutMeshDescription) override
 	{
 		UE::CADKernel::FModel& Model = CADKernelSession.GetModel();
-		return CADLibrary::FCADKernelTools::Tessellate(Model, ImportParameters, InMeshParameters, OutMeshDescription);
+		
+		CADLibrary::FMeshConversionContext Context(ImportParameters, InMeshParameters);
+
+		return CADLibrary::FCADKernelTools::Tessellate(Model, Context, OutMeshDescription);
 	}
 
 	virtual void SetImportParameters(double ChordTolerance, double MaxEdgeLength, double NormalTolerance, CADLibrary::EStitchingTechnique StitchingTechnique) override
diff --git a/Engine/Plugins/Enterprise/DatasmithCADImporter/Source/CADLibrary/Private/MeshDescriptionHelper.cpp b/Engine/Plugins/Enterprise/DatasmithCADImporter/Source/CADLibrary/Private/MeshDescriptionHelper.cpp
--- a/Engine/Plugins/Enterprise/DatasmithCADImporter/Source/CADLibrary/Private/MeshDescriptionHelper.cpp
+++ b/Engine/Plugins/Enterprise/DatasmithCADImporter/Source/CADLibrary/Private/MeshDescriptionHelper.cpp
@@ -106,7 +106,7 @@ void MergeCoincidents(const TArray<VectorType>& DataArray, const VectorType& One
 	}
 }
 
-void FillVertexPosition(const FImportParameters& ImportParams, const FMeshParameters& MeshParameters, FBodyMesh& Body, FMeshDescription& MeshDescription)
+void FillVertexPosition(const FMeshConversionContext& MeshConversionContext, FBodyMesh& Body, FMeshDescription& MeshDescription)
 {
 	TRACE_CPUPROFILER_EVENT_SCOPE(CADLibrary::FillVertexPosition);
 
@@ -121,7 +121,7 @@ void FillVertexPosition(const FImportParameters& ImportParams, const FMeshParame
 	VertexIdSet.SetNumZeroed(VertexCount);
 
 	// Make MeshDescription.VertexPositions and VertexID
-	MeshDescription.ReserveNewVertices(MeshParameters.bIsSymmetric ? VertexCount * 2 : VertexCount);
+	MeshDescription.ReserveNewVertices(MeshConversionContext.MeshParameters.bIsSymmetric ? VertexCount * 2 : VertexCount);
 
 	int32 VertexIndex = -1;
 	for (const FVector3f& Vertex : VertexArray)
@@ -135,7 +135,7 @@ void FillVertexPosition(const FImportParameters& ImportParams, const FMeshParame
 		}
 
 		FVertexID VertexID = MeshDescription.CreateVertex();
-		VertexPositions[VertexID] = (FVector3f)FDatasmithUtils::ConvertVector((FDatasmithUtils::EModelCoordSystem)ImportParams.GetModelCoordSys(), Vertex);
+		VertexPositions[VertexID] = (FVector3f)FDatasmithUtils::ConvertVector((FDatasmithUtils::EModelCoordSystem)MeshConversionContext.ImportParameters.GetModelCoordSys(), Vertex);
 		VertexIdSet[VertexIndex] = VertexID;
 	}
 
@@ -143,9 +143,9 @@ void FillVertexPosition(const FImportParameters& ImportParams, const FMeshParame
 	MergeCoincidents(VertexArray, FVector3f::OneVector, GeometricTolerance, VertexIdSet);
 
 	// if Symmetric mesh, the symmetric side of the mesh have to be generated
-	if (MeshParameters.bIsSymmetric)
+	if (MeshConversionContext.MeshParameters.bIsSymmetric)
 	{
-		FMatrix44f SymmetricMatrix = (FMatrix44f) FDatasmithUtils::GetSymmetricMatrix(MeshParameters.SymmetricOrigin, MeshParameters.SymmetricNormal);
+		FMatrix44f SymmetricMatrix = (FMatrix44f) FDatasmithUtils::GetSymmetricMatrix(MeshConversionContext.MeshParameters.SymmetricOrigin, MeshConversionContext.MeshParameters.SymmetricNormal);
 
 		TArray<int32>& SymmetricVertexIds = Body.SymmetricVertexIds;
 		SymmetricVertexIds.SetNum(VertexArray.Num());
@@ -160,7 +160,7 @@ void FillVertexPosition(const FImportParameters& ImportParams, const FMeshParame
 			}
 
 			FVertexID VertexID = MeshDescription.CreateVertex();
-			const FVector3f VertexPosition = FDatasmithUtils::ConvertVector((FDatasmithUtils::EModelCoordSystem)ImportParams.GetModelCoordSys(), Vertex);
+			const FVector3f VertexPosition = FDatasmithUtils::ConvertVector((FDatasmithUtils::EModelCoordSystem)MeshConversionContext.ImportParameters.GetModelCoordSys(), Vertex);
 			const FVector4f SymmetricPosition = FVector4f(SymmetricMatrix.TransformPosition(VertexPosition));
 			VertexPositions[VertexID] = FVector3f(SymmetricPosition);
 			SymmetricVertexIds[VertexIndex++] = VertexID;
@@ -184,75 +184,80 @@ TPolygonAttributesRef<int32> EnableCADPatchGroups(FMeshDescription& OutMeshDescr
 	return PatchGroups;
 }
 
-/**
- * Polygon group is an attribute of polygons.As long as the mesh description is empty(no polygon), polygon group cannot be defined.
- * The work around used is to create PolygonGroups and to set them a PolygonGroupAttributes.
- * To get the existing polygon groups, for each created polygonGroups the PolygonGroupId is got(see GetExistingPatches)
- * Warning: CopyPatchGroups is call in FCoreTechRetessellate_Impl::ApplyOnOneAsset(CoreTechRetessellateAction.cpp) only if the option
- * RetessellateOptions.RetessellationRule equal EDatasmithCADRetessellationRule::SkipDeletedSurfaces
- */
-void CopyPatchGroups(FMeshDescription& MeshSource, FMeshDescription& MeshDestination)
+void GetExistingPatches(FMeshDescription& MeshSource, TSet<int32>& OutExistingPatchIds)
 {
-	TPolygonGroupAttributesRef<int32> PatchGroups = MeshDestination.PolygonGroupAttributes().GetAttributesRef<int32>(PolyTriGroups);
-	if (PatchGroups.IsValid() == false)
-	{
-		MeshDestination.PolygonGroupAttributes().RegisterAttribute<int32>(PolyTriGroups, 1, 0, EMeshAttributeFlags::AutoGenerated);
-		PatchGroups = MeshDestination.PolygonGroupAttributes().GetAttributesRef<int32>(PolyTriGroups);
-		check(PatchGroups.IsValid());
-	}
-
-	TSet<int32> PatchIdSet;
 	TPolygonAttributesRef<int32> ElementToGroups = MeshSource.PolygonAttributes().GetAttributesRef<int32>(PolyTriGroups);
+	int32 LastPatchId = -1;
 	for (const FPolygonID TriangleID : MeshSource.Polygons().GetElementIDs())
 	{
 		int32 PatchId = ElementToGroups[TriangleID];
-
-		bool bIsAlreadyInSet = false;
-		PatchIdSet.Add(PatchId, &bIsAlreadyInSet);
-		if (!bIsAlreadyInSet)
+		if (PatchId != LastPatchId)
 		{
-			const FPolygonGroupID PolygonID = MeshDestination.CreatePolygonGroup();
-			PatchGroups[PolygonID] = PatchId;
+			OutExistingPatchIds.Add(PatchId);
+			LastPatchId = PatchId;
 		}
 	}
 }
 
-/**
- * See CopyPatchGroups
- */
-void GetExistingPatches(FMeshDescription& MeshDestination, TSet<int32>& OutPatchIdSet)
+FMeshDescriptionDataCache::FMeshDescriptionDataCache(FMeshDescription& MeshSource)
 {
-	TPolygonGroupAttributesRef<int32> PatchGroups = MeshDestination.PolygonGroupAttributes().GetAttributesRef<int32>(PolyTriGroups);
-	if (PatchGroups.IsValid() == false)
-	{
-		return;
-	}
+	FStaticMeshAttributes MeshSourceAttributes(MeshSource);
+
+	TPolygonAttributesRef<FPolygonGroupID> PolygoneToPolygoneGroupId = MeshSourceAttributes.GetPolygonPolygonGroupIndices();
+	TPolygonAttributesRef<int32> PatchGroup = EnableCADPatchGroups(MeshSource);
 
-	for (const FPolygonGroupID PolygonGroupID : MeshDestination.PolygonGroups().GetElementIDs())
+	int32 LastPatchGroupId = -1;
+	for (int32 PolygoneID = 0; PolygoneID < PolygoneToPolygoneGroupId.GetNumElements(); ++PolygoneID)
 	{
-		int32 PatchId = PatchGroups[PolygonGroupID];
-		if (PatchId > 0)
+		const FPolygonGroupID GroupID = PolygoneToPolygoneGroupId[PolygoneID];
+		int32 PatchGroupID = PatchGroup[PolygoneID];
+
+		if (PatchGroupID != LastPatchGroupId)
 		{
-			OutPatchIdSet.Add(PatchId);
+			PatchGroupToPolygonGroup.Add(PatchGroupID, GroupID);
+			LastPatchGroupId = PatchGroupID;
 		}
 	}
+
+	TPolygonGroupAttributesRef<FName> MeshSourcePolygonGroupImportedMaterialSlotNames = MeshSourceAttributes.GetPolygonGroupMaterialSlotNames();
+	MaterialSlotNames = MeshSourcePolygonGroupImportedMaterialSlotNames.GetRawArray(0);
 }
 
-void CopyMaterialSlotNames(FMeshDescription& MeshSource, FMeshDescription& MeshDestination)
+void FMeshDescriptionDataCache::RestoreMaterialSlotNames(FMeshDescription& Mesh) const
 {
-	FStaticMeshAttributes MeshAttributes(MeshSource);
-	TPolygonGroupAttributesRef<FName> PolygonGroupImportedMaterialSlotNames = MeshAttributes.GetPolygonGroupMaterialSlotNames();
+	TPolygonAttributesRef<int32> PatchGroup = EnableCADPatchGroups(Mesh);
+	FStaticMeshAttributes MeshAttributes(Mesh);
+	TPolygonAttributesRef<FPolygonGroupID> PolygoneToPolygoneGroupId = MeshAttributes.GetPolygonPolygonGroupIndices();
+	TPolygonGroupAttributesRef<FName> PolygonGroupMaterialSlotNames = MeshAttributes.GetPolygonGroupMaterialSlotNames();
+
+	TSet<FPolygonGroupID> GroupIdDone;
+	int32 LastPatchGroupId = -1;
+	for (int32 PolygoneID = 0; PolygoneID < PolygoneToPolygoneGroupId.GetNumElements(); ++PolygoneID)
+	{
+		const int32 PatchGroupID = PatchGroup[PolygoneID];
+		if (PatchGroupID != LastPatchGroupId)
+		{
+			const FPolygonGroupID GroupID = PolygoneToPolygoneGroupId[PolygoneID];
+			if (!GroupIdDone.Contains(GroupID))
+			{
+				GroupIdDone.Add(GroupID);
 
-	FStaticMeshAttributes NewMeshDescriptionAttributes(MeshDestination);
-	TPolygonGroupAttributesRef<FName> NewPolygonGroupImportedMaterialSlotNames = NewMeshDescriptionAttributes.GetPolygonGroupMaterialSlotNames();
+				const FPolygonGroupID* GroupIDMeshSource = Find(PatchGroupID);
+				if (GroupIDMeshSource)
+				{
+					const FName& SlotName = GetSlotName(*GroupIDMeshSource);
+					ensure(GroupID < PolygonGroupMaterialSlotNames.GetNumElements());
+					PolygonGroupMaterialSlotNames[GroupID] = SlotName;
+				}
+				// else PatchGroupeID doesn't exist in MeshSource
+			}
 
-	for (int32 Index = 0; Index < PolygonGroupImportedMaterialSlotNames.GetNumElements(); ++Index)
-	{
-		NewPolygonGroupImportedMaterialSlotNames[Index] = PolygonGroupImportedMaterialSlotNames[Index];
+			LastPatchGroupId = PatchGroupID;
+		}
 	}
 }
 
-bool FillMesh(const FMeshParameters& MeshParameters, const FImportParameters& ImportParams, FBodyMesh& BodyTessellation, FMeshDescription& MeshDescription)
+bool FillMesh(const FMeshConversionContext& MeshConversionContext, FBodyMesh& BodyTessellation, FMeshDescription& MeshDescription)
 {
 	TRACE_CPUPROFILER_EVENT_SCOPE(CADLibrary::FillMesh);
 
@@ -283,10 +288,24 @@ bool FillMesh(const FMeshParameters& MeshParameters, const FImportParameters& Im
 	// To avoid duplicated triangle
 	TSet<TTuple<FVertexInstanceID, FVertexInstanceID, FVertexInstanceID>> VertexIdsToTriangle;
 
+	const TSet<int32>& ExistingPatches = MeshConversionContext.PatchesToMesh;
+	bool bImportOnlyAlreadyPresent = (bool)ExistingPatches.Num();
+
 	// Find all the materials used
 	TMap<uint32, FPolygonGroupID> MaterialToPolygonGroupMapping;
 	for (const FTessellationData& FaceTessellation : BodyTessellation.Faces)
 	{
+		// don't add material of empty face
+		if (FaceTessellation.VertexIndices.Num() == 0)
+		{
+			continue;
+		}
+
+		if (bImportOnlyAlreadyPresent && !ExistingPatches.Contains(FaceTessellation.PatchId))
+		{
+			continue;
+		}
+
 		// material is preferred over color
 		MaterialToPolygonGroupMapping.Add(FaceTessellation.MaterialUId ? FaceTessellation.MaterialUId : FaceTessellation.ColorUId, INDEX_NONE);
 	}
@@ -317,26 +336,22 @@ bool FillMesh(const FMeshParameters& MeshParameters, const FImportParameters& Im
 	}
 
 	int32 NbStep = 1;
-	if (MeshParameters.bIsSymmetric)
+	if (MeshConversionContext.MeshParameters.bIsSymmetric)
 	{
 		NbStep = 2;
 	}
 
-	TSet<int32> PatchIdSet;
-	GetExistingPatches(MeshDescription, PatchIdSet);
-	bool bImportOnlyAlreadyPresent = (bool)PatchIdSet.Num();
-
 	TPolygonAttributesRef<int32> PatchGroups = EnableCADPatchGroups(MeshDescription);
 	int32 PatchIndex = 0;
 	for (int32 Step = 0; Step < NbStep; ++Step)
 	{
 		// Swap mesh if needed
-		const TriangleIndex& Orientation = (!MeshParameters.bNeedSwapOrientation == (bool)Step) ? CounterClockwise : Clockwise;
+		const TriangleIndex& Orientation = (!MeshConversionContext.MeshParameters.bNeedSwapOrientation == (bool)Step) ? CounterClockwise : Clockwise;
 		TArray<int32>& VertexIdSet = (Step == 0) ? BodyTessellation.VertexIds : BodyTessellation.SymmetricVertexIds;
 
 		for (FTessellationData& Tessellation : BodyTessellation.Faces)
 		{
-			if (bImportOnlyAlreadyPresent && !PatchIdSet.Contains(Tessellation.PatchId))
+			if (bImportOnlyAlreadyPresent && !ExistingPatches.Contains(Tessellation.PatchId))
 			{
 				continue;
 			}
@@ -351,7 +366,7 @@ bool FillMesh(const FMeshParameters& MeshParameters, const FImportParameters& Im
 
 			if (!Step)
 			{
-				FDatasmithUtils::ConvertVectorArray(ImportParams.GetModelCoordSys(), Tessellation.NormalArray);
+				FDatasmithUtils::ConvertVectorArray(MeshConversionContext.ImportParameters.GetModelCoordSys(), Tessellation.NormalArray);
 				for (FVector3f& Normal : Tessellation.NormalArray)
 				{
 					Normal = Normal.GetSafeNormal();
@@ -507,14 +522,14 @@ bool FillMesh(const FMeshParameters& MeshParameters, const FImportParameters& Im
 			if (Step)
 			{
 				// compute normals of Symmetric vertex
-				FMatrix44f SymmetricMatrix = FDatasmithUtils::GetSymmetricMatrix(MeshParameters.SymmetricOrigin, MeshParameters.SymmetricNormal);
+				FMatrix44f SymmetricMatrix = FDatasmithUtils::GetSymmetricMatrix(MeshConversionContext.MeshParameters.SymmetricOrigin, MeshConversionContext.MeshParameters.SymmetricNormal);
 				for (const FVertexInstanceID& VertexInstanceID : MeshVertexInstanceIDs)
 				{
 					VertexInstanceNormals[VertexInstanceID] = SymmetricMatrix.TransformVector(VertexInstanceNormals[VertexInstanceID]);
 				}
 			}
 
-			if (MeshParameters.bNeedSwapOrientation)
+			if (MeshConversionContext.MeshParameters.bNeedSwapOrientation)
 			{
 				for (int32 Index = 0; Index < MeshVertexInstanceIDs.Num(); Index++)
 				{
@@ -526,7 +541,7 @@ bool FillMesh(const FMeshParameters& MeshParameters, const FImportParameters& Im
 	return true;
 }
 
-bool ConvertBodyMeshToMeshDescription(const FImportParameters& ImportParams, const FMeshParameters& MeshParameters, FBodyMesh& Body, FMeshDescription& MeshDescription)
+bool ConvertBodyMeshToMeshDescription(const FMeshConversionContext& MeshConversionContext, FBodyMesh& Body, FMeshDescription& MeshDescription)
 {
 	TRACE_CPUPROFILER_EVENT_SCOPE(CADLibrary::ConvertBodyMeshToMeshDescription);
 
@@ -535,9 +550,9 @@ bool ConvertBodyMeshToMeshDescription(const FImportParameters& ImportParams, con
 	MeshDescription.ReserveNewPolygons(Body.TriangleCount);
 	MeshDescription.ReserveNewEdges(Body.TriangleCount * 3);
 
-	FillVertexPosition(ImportParams, MeshParameters, Body, MeshDescription);
+	FillVertexPosition(MeshConversionContext, Body, MeshDescription);
 
-	if (!FillMesh(MeshParameters, ImportParams, Body, MeshDescription))
+	if (!FillMesh(MeshConversionContext, Body, MeshDescription))
 	{
 		return false;
 	}
diff --git a/Engine/Plugins/Enterprise/DatasmithCADImporter/Source/CADLibrary/Private/MeshOperator.cpp b/Engine/Plugins/Enterprise/DatasmithCADImporter/Source/CADLibrary/Private/MeshOperator.cpp
--- a/Engine/Plugins/Enterprise/DatasmithCADImporter/Source/CADLibrary/Private/MeshOperator.cpp
+++ b/Engine/Plugins/Enterprise/DatasmithCADImporter/Source/CADLibrary/Private/MeshOperator.cpp
@@ -366,8 +366,7 @@ void MeshOperator::ResolveTJunctions(FMeshDescription& MeshDescription, double T
 		BoundaryEdgeCount = MergeCoincidentEdges.FinalNumBoundaryEdges;
 	}
 
-	// TODO: Should not be needed as soon as CopyMaterialSlotNames, and CopyPatchGroups are no more called
-	FMeshDescription NewMeshDescription(MeshDescription);
+	CADLibrary::FMeshDescriptionDataCache MeshDescriptionDataCache(MeshDescription);
 
 	{
 		FConversionToMeshDescriptionOptions ConversionOptions;
@@ -380,12 +379,8 @@ void MeshOperator::ResolveTJunctions(FMeshDescription& MeshDescription, double T
 		ConversionOptions.bTransformVtxColorsSRGBToLinear = false;
 
 		FDynamicMeshToMeshDescription ConverterToMeshDescription(ConversionOptions);
-		ConverterToMeshDescription.Convert(&DynamicMesh, NewMeshDescription);
+		ConverterToMeshDescription.Convert(&DynamicMesh, MeshDescription);
 	}
 
-	// Transfer material slot names and patch groups as these datum are not managed by FDynamicMesh3
-	CADLibrary::CopyMaterialSlotNames(MeshDescription, NewMeshDescription);
-	CADLibrary::CopyPatchGroups(MeshDescription, NewMeshDescription);
-
-	MeshDescription = MoveTemp(NewMeshDescription);
+	MeshDescriptionDataCache.RestoreMaterialSlotNames(MeshDescription);
 }
\ No newline at end of file
diff --git a/Engine/Plugins/Enterprise/DatasmithCADImporter/Source/CADLibrary/Public/CADKernelTools.h b/Engine/Plugins/Enterprise/DatasmithCADImporter/Source/CADLibrary/Public/CADKernelTools.h
--- a/Engine/Plugins/Enterprise/DatasmithCADImporter/Source/CADLibrary/Public/CADKernelTools.h
+++ b/Engine/Plugins/Enterprise/DatasmithCADImporter/Source/CADLibrary/Public/CADKernelTools.h
@@ -8,32 +8,34 @@
 
 namespace UE::CADKernel
 {
-	class FBody;
-	class FBodyMesh;
-	class FFaceMesh;
-	class FTopologicalShapeEntity;
-	class FModelMesh;
+class FBody;
+class FBodyMesh;
+class FFaceMesh;
+class FTopologicalShapeEntity;
+class FModelMesh;
 }
 
 struct FMeshDescription;
 
 namespace CADLibrary
 {
-	class FImportParameters;
-	struct FMeshParameters;
+class FImportParameters;
+struct FMeshParameters;
 
-	class CADLIBRARY_API FCADKernelTools
-	{
-	public:
-		static void DefineMeshCriteria(UE::CADKernel::FModelMesh& MeshModel, const FImportParameters& ImportParameters, double GeometricTolerance);
+struct FMeshConversionContext;
 
-		static void GetBodyTessellation(const UE::CADKernel::FModelMesh& ModelMesh, const UE::CADKernel::FBody& Body, FBodyMesh& OutBodyMesh);
+class CADLIBRARY_API FCADKernelTools
+{
+public:
+	static void DefineMeshCriteria(UE::CADKernel::FModelMesh& MeshModel, const FImportParameters& ImportParameters, double GeometricTolerance);
+
+	static void GetBodyTessellation(const UE::CADKernel::FModelMesh& ModelMesh, const UE::CADKernel::FBody& Body, FBodyMesh& OutBodyMesh);
 
-		/**
-		 * Tessellate a CADKernel entity and update the MeshDescription
-		 */
-		static bool Tessellate(UE::CADKernel::FTopologicalShapeEntity& InCADKernelEntity, const FImportParameters& ImportParameters, const FMeshParameters& MeshParameters, FMeshDescription& OutMeshDescription);
+	/**
+	 * Tessellate a CADKernel entity and update the MeshDescription
+	 */
+	static bool Tessellate(UE::CADKernel::FTopologicalShapeEntity& InCADKernelEntity, FMeshConversionContext& Context, FMeshDescription& OutMeshDescription);
 
-		static uint32 GetFaceTessellation(UE::CADKernel::FFaceMesh& FaceMesh, FBodyMesh& OutBodyMesh, FObjectDisplayDataId FaceMaterial);
-	};
+	static uint32 GetFaceTessellation(UE::CADKernel::FFaceMesh& FaceMesh, FBodyMesh& OutBodyMesh, FObjectDisplayDataId FaceMaterial);
+};
 }
diff --git a/Engine/Plugins/Enterprise/DatasmithCADImporter/Source/CADLibrary/Public/MeshDescriptionHelper.h b/Engine/Plugins/Enterprise/DatasmithCADImporter/Source/CADLibrary/Public/MeshDescriptionHelper.h
--- a/Engine/Plugins/Enterprise/DatasmithCADImporter/Source/CADLibrary/Public/MeshDescriptionHelper.h
+++ b/Engine/Plugins/Enterprise/DatasmithCADImporter/Source/CADLibrary/Public/MeshDescriptionHelper.h
@@ -21,13 +21,69 @@ namespace CADLibrary
 
 using TColorMap = TMap<uint32, uint32>;
 
+struct FMeshConversionContext
+{
+	const FImportParameters& ImportParameters;
+	const FMeshParameters& MeshParameters;
+	TArray<int32> VertexIds;
+	TArray<int32> SymmetricVertexIds;
+	TSet<int32> PatchesToMesh;
+
+	FMeshConversionContext(const FImportParameters& InImportParams, const FMeshParameters& InMeshParameters)
+		: ImportParameters(InImportParams)
+		, MeshParameters(InMeshParameters)
+	{
+	}
+};
+
+/**
+ * Macro data of MeshDescription.
+ * Needed to updata a MeshDesction after a convertion of a DynamicMesh To a MeshDescription.
+ * Indeed, PolygonGroupId can be modify during the process of DynamicMesh and material slot names are not saved in a DynamicMesh.
+ * The only persistent data is the patchGroupId
+ * 
+ * This class save the PatchGroupToPolygonGroup map and the MaterialSlotNames of the MeshDescription before the conversion into a DynamicMesh 
+ * This information allows to restor the material slot name used by each PolygonGroup
+ *
+ * The call of UpdateMaterialSlotNames restor the information in the new or update MeshDescription
+ */
+class CADLIBRARY_API FMeshDescriptionDataCache
+{
+private:
+	/** the mapping between PatchId and PolygonGroupId */
+	TMap<int32, FPolygonGroupID> PatchGroupToPolygonGroup;
+
+	/** the mapping between PolygonGroupId and SlotNames */
+	TArray<FName> MaterialSlotNames;
+
+	FName EmptyName;
+
+	const FPolygonGroupID* Find(const int32 PatchGroupId) const 
+	{
+		return PatchGroupToPolygonGroup.Find(PatchGroupId);
+	}
+
+	const FName& GetSlotName(const FPolygonGroupID& GroupID) const
+	{
+		if(GroupID < MaterialSlotNames.Num())
+		{
+			return MaterialSlotNames[GroupID];
+		}
+		return EmptyName;
+	}
+
+public:
+
+	FMeshDescriptionDataCache(FMeshDescription& Mesh);
+
+	void RestoreMaterialSlotNames(FMeshDescription& MeshDestination) const;
+};
+
 CADLIBRARY_API TSharedPtr<IDatasmithUEPbrMaterialElement> CreateDefaultUEPbrMaterial();
 CADLIBRARY_API TSharedPtr<IDatasmithUEPbrMaterialElement> CreateUEPbrMaterialFromColor(const FColor& InColor);
 CADLIBRARY_API TSharedPtr<IDatasmithUEPbrMaterialElement> CreateUEPbrMaterialFromMaterial(FCADMaterial& InMaterial, TSharedRef<IDatasmithScene> Scene);
 
-CADLIBRARY_API bool ConvertBodyMeshToMeshDescription(const FImportParameters& ImportParams, const FMeshParameters& MeshParameters, FBodyMesh& Body, FMeshDescription& MeshDescription);
-CADLIBRARY_API void CopyPatchGroups(FMeshDescription& MeshSource, FMeshDescription& MeshDestination);
-CADLIBRARY_API void CopyMaterialSlotNames(FMeshDescription& MeshSource, FMeshDescription& MeshDestination);
+CADLIBRARY_API bool ConvertBodyMeshToMeshDescription(const FMeshConversionContext& MeshConversionContext, FBodyMesh& Body, FMeshDescription& MeshDescription);
 
 /** 
  * Enable per-triangle integer attribute named PolyTriGroups 
@@ -35,6 +91,6 @@ CADLIBRARY_API void CopyMaterialSlotNames(FMeshDescription& MeshSource, FMeshDes
  * In case of mesh coming from a CAD file, a PolyTriGroup is associated to a CAD topological face
  */
 CADLIBRARY_API TPolygonAttributesRef<int32> EnableCADPatchGroups(FMeshDescription& OutMeshDescription);
-CADLIBRARY_API void GetExistingPatches(FMeshDescription& MeshDestination, TSet<int32>& OutPatchIdSet);
+CADLIBRARY_API void GetExistingPatches(FMeshDescription& Mesh, TSet<int32>& OutPatchIdSet);
 
 } // namespace CADLibrary
diff --git a/Engine/Plugins/Enterprise/DatasmithCADImporter/Source/DatasmithCADTranslator/Private/DatasmithMeshBuilder.cpp b/Engine/Plugins/Enterprise/DatasmithCADImporter/Source/DatasmithCADTranslator/Private/DatasmithMeshBuilder.cpp
--- a/Engine/Plugins/Enterprise/DatasmithCADImporter/Source/DatasmithCADTranslator/Private/DatasmithMeshBuilder.cpp
+++ b/Engine/Plugins/Enterprise/DatasmithCADImporter/Source/DatasmithCADTranslator/Private/DatasmithMeshBuilder.cpp
@@ -73,7 +73,9 @@ TOptional<FMeshDescription> FDatasmithMeshBuilder::GetMeshDescription(TSharedRef
 	FMeshDescription MeshDescription;
 	DatasmithMeshHelper::PrepareAttributeForStaticMesh(MeshDescription);
 
-	if (ConvertBodyMeshToMeshDescription(ImportParameters, OutMeshParameters, Body, MeshDescription))
+	CADLibrary::FMeshConversionContext ConversionContext(ImportParameters, OutMeshParameters);
+
+	if (ConvertBodyMeshToMeshDescription(ConversionContext, Body, MeshDescription))
 	{
 		return MoveTemp(MeshDescription);
 	}
diff --git a/Engine/Plugins/Enterprise/DatasmithCADImporter/Source/ParametricSurface/Private/TechSoft/CADModelToTechSoftConverterBase.cpp b/Engine/Plugins/Enterprise/DatasmithCADImporter/Source/ParametricSurface/Private/TechSoft/CADModelToTechSoftConverterBase.cpp
--- a/Engine/Plugins/Enterprise/DatasmithCADImporter/Source/ParametricSurface/Private/TechSoft/CADModelToTechSoftConverterBase.cpp
+++ b/Engine/Plugins/Enterprise/DatasmithCADImporter/Source/ParametricSurface/Private/TechSoft/CADModelToTechSoftConverterBase.cpp
@@ -89,7 +89,8 @@ bool FCADModelToTechSoftConverterBase::Tessellate(const CADLibrary::FMeshParamet
 		return false;
 	}
 
-	return CADLibrary::ConvertBodyMeshToMeshDescription(ImportParameters, InMeshParameters, BodyMesh, OutMeshDescription);
+	CADLibrary::FMeshConversionContext MeshConversionContext(ImportParameters, InMeshParameters);
+	return CADLibrary::ConvertBodyMeshToMeshDescription(MeshConversionContext, BodyMesh, OutMeshDescription);
 }
 
 void FCADModelToTechSoftConverterBase::AddSurfaceDataForMesh(const TCHAR* InFilePath, const CADLibrary::FMeshParameters& InMeshParameters, const FDatasmithTessellationOptions& InTessellationOptions, FDatasmithMeshElementPayload& OutMeshPayload) const
diff --git a/Engine/Plugins/Enterprise/DatasmithCADImporter/Source/ParametricSurface/Private/TechSoft/TechSoftParametricSurface.cpp b/Engine/Plugins/Enterprise/DatasmithCADImporter/Source/ParametricSurface/Private/TechSoft/TechSoftParametricSurface.cpp
--- a/Engine/Plugins/Enterprise/DatasmithCADImporter/Source/ParametricSurface/Private/TechSoft/TechSoftParametricSurface.cpp
+++ b/Engine/Plugins/Enterprise/DatasmithCADImporter/Source/ParametricSurface/Private/TechSoft/TechSoftParametricSurface.cpp
@@ -35,33 +35,33 @@ bool UTechSoftParametricSurfaceData::Tessellate(UStaticMesh& StaticMesh, const F
 	// Previous MeshDescription is get to be able to create a new one with the same order of PolygonGroup (the matching of color and partition is currently based on their order)
 	if (FMeshDescription* DestinationMeshDescription = StaticMesh.GetMeshDescription(0))
 	{
-		FMeshDescription MeshDescription;
-		FStaticMeshAttributes MeshDescriptionAttributes(MeshDescription);
-		MeshDescriptionAttributes.Register();
-
-		if (RetessellateOptions.RetessellationRule == EDatasmithCADRetessellationRule::SkipDeletedSurfaces)
-		{
-			CADLibrary::CopyPatchGroups(*DestinationMeshDescription, MeshDescription);
-		}
-
 		CADLibrary::FImportParameters ImportParameters((FDatasmithUtils::EModelCoordSystem)SceneParameters.ModelCoordSys);
 		ImportParameters.SetTesselationParameters(RetessellateOptions.ChordTolerance, RetessellateOptions.MaxEdgeLength, RetessellateOptions.NormalTolerance, (CADLibrary::EStitchingTechnique)RetessellateOptions.StitchingTechnique);
 
-		CADLibrary::FTechSoftInterface& TechSoftInterface = CADLibrary::FTechSoftInterface::Get();
+		FMeshConversionContext Context(ImportParameters, MeshParameters);
 
+		CADLibrary::FTechSoftInterface& TechSoftInterface = CADLibrary::FTechSoftInterface::Get();
 		bSuccessfulTessellation = TechSoftInterface.InitializeKernel(*FPaths::EnginePluginsDir());
 		if (bSuccessfulTessellation)
 		{
 			CADLibrary::FBodyMesh BodyMesh;
-
 			bSuccessfulTessellation = CADLibrary::TechSoftUtils::GetBodyFromPcrFile(ResourceFile, ImportParameters, BodyMesh);
 
 			if (bSuccessfulTessellation)
 			{
-				bSuccessfulTessellation = CADLibrary::ConvertBodyMeshToMeshDescription(ImportParameters, MeshParameters, BodyMesh, MeshDescription);
+				if (RetessellateOptions.RetessellationRule == EDatasmithCADRetessellationRule::SkipDeletedSurfaces)
+				{
+					CADLibrary::GetExistingPatches(*DestinationMeshDescription, Context.PatchesToMesh);
+				}
+
+				FMeshDescription MeshDescription;
+				FStaticMeshAttributes MeshDescriptionAttributes(MeshDescription);
+				MeshDescriptionAttributes.Register();
+				bSuccessfulTessellation = CADLibrary::ConvertBodyMeshToMeshDescription(Context, BodyMesh, MeshDescription);
 
 				if (bSuccessfulTessellation)
 				{
+	
 					TPolygonGroupAttributesConstRef<FName> MaterialSlotNames = MeshDescriptionAttributes.GetPolygonGroupMaterialSlotNames();
 					FMeshSectionInfoMap& SectionInfoMap = StaticMesh.GetSectionInfoMap();
 
diff --git a/Engine/Plugins/Enterprise/DatasmithImporter/Source/DatasmithImporter/Private/DatasmithActorImporter.cpp b/Engine/Plugins/Enterprise/DatasmithImporter/Source/DatasmithImporter/Private/DatasmithActorImporter.cpp
--- a/Engine/Plugins/Enterprise/DatasmithImporter/Source/DatasmithImporter/Private/DatasmithActorImporter.cpp
+++ b/Engine/Plugins/Enterprise/DatasmithImporter/Source/DatasmithImporter/Private/DatasmithActorImporter.cpp
@@ -354,16 +354,20 @@ void FDatasmithActorImporter::SetupStaticMeshComponent( FDatasmithImportContext&
 
 	if ( !StaticMesh )
 	{
-		FString OwnerName;
+		FString OwnerLabel;
 
 		if ( StaticMeshComponent->GetOwner() )
 		{
-			OwnerName = StaticMeshComponent->GetOwner()->GetName();
+			OwnerLabel = StaticMeshComponent->GetOwner()->GetActorLabel();
+			if (OwnerLabel.IsEmpty())
+			{
+				OwnerLabel = StaticMeshComponent->GetOwner()->GetName();
+			}
 		}
 
 		ImportContext.LogError( FText::Format( LOCTEXT( "FindStaticMesh", "Cannot find Static Mesh {0} for Static Mesh Actor {1}."),
-			FText::FromString( MeshActorElement->GetStaticMeshPathName() ),
-			FText::FromString( OwnerName + TEXT(".") + StaticMeshComponent->GetName() ) ) );
+			FText::FromString(MeshActorElement->GetLabel()),
+			FText::FromString(OwnerLabel + TEXT(".") + StaticMeshComponent->GetName())));
 
 		return;
 	}
