Empty or Deleted CADKernel Bodies are no more convert into StaticMeshActor

robomerge[FN_EngineMerge] -UE5-RES

diff --git a/Engine/Plugins/Enterprise/DatasmithCADImporter/Source/CADInterfaces/Private/TechSoftFileParserCADKernelTessellator.cpp b/Engine/Plugins/Enterprise/DatasmithCADImporter/Source/CADInterfaces/Private/TechSoftFileParserCADKernelTessellator.cpp
--- a/Engine/Plugins/Enterprise/DatasmithCADImporter/Source/CADInterfaces/Private/TechSoftFileParserCADKernelTessellator.cpp
+++ b/Engine/Plugins/Enterprise/DatasmithCADImporter/Source/CADInterfaces/Private/TechSoftFileParserCADKernelTessellator.cpp
@@ -73,7 +73,6 @@ void FTechSoftFileParserCADKernelTessellator::SewAndGenerateBodyMeshes()
 	{
 		SewAndMesh(Representations.Value);
 	}
-
 }
 
 void FTechSoftFileParserCADKernelTessellator::SewAndMesh(TArray<A3DRiRepresentationItem*>& Representations)
@@ -194,7 +193,15 @@ void FTechSoftFileParserCADKernelTessellator::SewAndMesh(TArray<A3DRiRepresentat
 			FCadId ArchiveBodyId = ArchiveBodyIdOfDeletedRepresentation[Index++];
 			FArchiveBody& ArchiveBody = SceneGraph.GetBody(ArchiveBodyId);
 			ArchiveBody.Delete();
+		}
+	}
 
+	// Delete unused ArchiveBody
+	for (FArchiveBody& ArchiveBody : SceneGraph.Bodies)
+	{
+		if (ArchiveBody.MeshActorUId == 0)
+		{
+			ArchiveBody.Delete();
 		}
 	}
 }
diff --git a/Engine/Source/Runtime/Datasmith/CADKernel/Private/CADKernel/Topo/Topomaker.cpp b/Engine/Source/Runtime/Datasmith/CADKernel/Private/CADKernel/Topo/Topomaker.cpp
--- a/Engine/Source/Runtime/Datasmith/CADKernel/Private/CADKernel/Topo/Topomaker.cpp
+++ b/Engine/Source/Runtime/Datasmith/CADKernel/Private/CADKernel/Topo/Topomaker.cpp
@@ -1026,6 +1026,15 @@ void FTopomaker::SplitIntoConnectedShells()
 		}
 	};
 
+	for (const TSharedPtr<FTopologicalFace>& Face : Faces)
+	{
+		if (Face->IsDeleted() || Face->IsDegenerated())
+		{
+			Face->SetMarker1();
+			continue;
+		}
+	}
+		
 	for (const TSharedPtr<FTopologicalFace>& Face : Faces)
 	{
 		if (Face->HasMarker1())
