Add lib 2023.1 to insure support of Alias 2023.1

diff --git a/Engine/Plugins/Enterprise/DatasmithCADImporter/DatasmithCADImporter.uplugin b/Engine/Plugins/Enterprise/DatasmithCADImporter/DatasmithCADImporter.uplugin
--- a/Engine/Plugins/Enterprise/DatasmithCADImporter/DatasmithCADImporter.uplugin
+++ b/Engine/Plugins/Enterprise/DatasmithCADImporter/DatasmithCADImporter.uplugin
@@ -117,6 +117,14 @@
 				"Win64"
 			]
 		},
+		{
+			"Name": "DatasmithWireTranslator2023_1",
+			"Type": "Editor",
+			"LoadingPhase": "Default",
+			"PlatformAllowList": [
+				"Win64"
+			]
+		},
 		{
 			"Name": "ParametricSurfaceExtension",
 			"Type": "Editor",
diff --git a/Engine/Plugins/Enterprise/DatasmithCADImporter/Source/DatasmithWireTranslator/DatasmithWireTranslator2023_1.Build.cs b/Engine/Plugins/Enterprise/DatasmithCADImporter/Source/DatasmithWireTranslator/DatasmithWireTranslator2023_1.Build.cs
new file mode 100644
--- /dev/null
+++ b/Engine/Plugins/Enterprise/DatasmithCADImporter/Source/DatasmithWireTranslator/DatasmithWireTranslator2023_1.Build.cs
@@ -0,0 +1,22 @@
+// Copyright Epic Games, Inc. All Rights Reserved.
+
+using System.IO;
+using UnrealBuildTool;
+
+public class DatasmithWireTranslator2023_1 : DatasmithWireTranslatorBase
+{
+	public DatasmithWireTranslator2023_1(ReadOnlyTargetRules Target) 
+		: base(Target)
+	{
+	}
+
+	public override string GetAliasVersion()
+	{
+		return "OpenModel2023_1";
+	}
+
+	public override string GetAliasDefinition()
+	{
+		return "OPEN_MODEL_2023_1";
+	}
+}
\ No newline at end of file
diff --git a/Engine/Plugins/Enterprise/DatasmithCADImporter/Source/DatasmithWireTranslator/Private/DatasmithWireTranslator.cpp b/Engine/Plugins/Enterprise/DatasmithCADImporter/Source/DatasmithWireTranslator/Private/DatasmithWireTranslator.cpp
--- a/Engine/Plugins/Enterprise/DatasmithCADImporter/Source/DatasmithWireTranslator/Private/DatasmithWireTranslator.cpp
+++ b/Engine/Plugins/Enterprise/DatasmithCADImporter/Source/DatasmithWireTranslator/Private/DatasmithWireTranslator.cpp
@@ -98,6 +98,7 @@ static const FColor DefaultColor = FColor(200, 200, 200);
 #ifdef USE_OPENMODEL
 
 const uint64 LibAliasNext_Version = 0xffffffffffffffffull;
+const uint64 LibAlias2023_1_0_Version = 8162778619576619;
 const uint64 LibAlias2023_0_0_Version = 8162774324609149;
 const uint64 LibAlias2022_2_0_Version = 7881307937833405;
 const uint64 LibAlias2022_1_0_Version = 7881303642865885;
@@ -131,8 +132,12 @@ const uint64 LibAliasVersionMax = LibAlias2023_0_0_Version;
 const FString AliasVersionChar = TEXT("AliasStudio 2022.2, Model files");
 #elif defined(OPEN_MODEL_2023_0)
 const uint64 LibAliasVersionMin = LibAlias2023_0_0_Version;
-const uint64 LibAliasVersionMax = LibAliasNext_Version;
+const uint64 LibAliasVersionMax = LibAlias2023_1_0_Version;
 const FString AliasVersionChar = TEXT("AliasStudio 2023.0, Model files");
+#elif defined(OPEN_MODEL_2023_1)
+const uint64 LibAliasVersionMin = LibAlias2023_1_0_Version;
+const uint64 LibAliasVersionMax = LibAliasNext_Version;
+const FString AliasVersionChar = TEXT("AliasStudio 2023.1, Model files");
 #endif
 
 // Alias material management (to allow sew of BReps of different materials):
@@ -2111,8 +2116,8 @@ TOptional<FMeshDescription> FWireTranslatorImpl::MeshDagNodeWithExternalMesher(A
 	bool bRet = CADModelConverter->Tessellate(MeshParameters, MeshDescription);
 	if (!bRet)
 	{
-		FString StaticMeshLable = MeshElement->GetLabel();
-		FString StaticMeshName = MeshElement->GetName();
+		const TCHAR* StaticMeshLable = MeshElement->GetLabel();
+		const TCHAR* StaticMeshName = MeshElement->GetName();
 		UE_LOG(LogDatasmithWireTranslator, Warning, TEXT("Failed to generate the mesh of \"%s\" (%s) StaticMesh."), *StaticMeshLable, *StaticMeshName);
 	}
 
diff --git a/Engine/Plugins/Enterprise/DatasmithCADImporter/Source/DatasmithWireTranslator/Private/OpenModelUtils.cpp b/Engine/Plugins/Enterprise/DatasmithCADImporter/Source/DatasmithWireTranslator/Private/OpenModelUtils.cpp
--- a/Engine/Plugins/Enterprise/DatasmithCADImporter/Source/DatasmithWireTranslator/Private/OpenModelUtils.cpp
+++ b/Engine/Plugins/Enterprise/DatasmithCADImporter/Source/DatasmithWireTranslator/Private/OpenModelUtils.cpp
@@ -128,7 +128,6 @@ bool OpenModelUtils::TransferAlMeshToMeshDescription(const AlMesh& AliasMesh, co
 	MeshDescription.ReserveNewEdges(VertexInstanceCount * NbStep);
 	MeshDescription.ReserveNewPolygons(TriangleCount * NbStep);
 
-	// Assume one material per mesh, no partitioning
 	MeshDescription.ReserveNewPolygonGroups(1);
 	FPolygonGroupID PolyGroupId = MeshDescription.CreatePolygonGroup();
 	FName ImportedSlotName = SlotMaterialId;
