Fix for CAD Orientation bug
The bug is due to the optimisation of the vertex instance count in the MeshDescription. Now, a Vertex instance is shared with multiple triangles. So we cannot invert the normal at each connected triangles inversion. A flag is used to invert the normal once.

diff --git a/Engine/Plugins/Enterprise/DatasmithCADImporter/Source/CADLibrary/Private/MeshEditingWrapper.cpp b/Engine/Plugins/Enterprise/DatasmithCADImporter/Source/CADLibrary/Private/MeshEditingWrapper.cpp
--- a/Engine/Plugins/Enterprise/DatasmithCADImporter/Source/CADLibrary/Private/MeshEditingWrapper.cpp
+++ b/Engine/Plugins/Enterprise/DatasmithCADImporter/Source/CADLibrary/Private/MeshEditingWrapper.cpp
@@ -31,6 +31,8 @@ void FMeshEditingWrapper::UpdateMeshWrapper()
 	{
 		DefineVertexTopologyApproximation(VertexID);
 	}
+
+	VertexInstanceMarker.Init(false, MeshDescription.VertexInstances().GetArraySize());
 }
 
 
@@ -252,7 +254,12 @@ void FMeshEditingWrapper::SwapTriangleOrientation(FTriangleID Triangle)
 	TArrayView<const FVertexInstanceID> TriVertexInstances = MeshDescription.GetTriangleVertexInstances(Triangle);
 	for (int32 IVertex = 0; IVertex < 3; IVertex++)
 	{
-		SwapVertexNormal(TriVertexInstances[IVertex]);
+		const FVertexInstanceID& InstanceID = TriVertexInstances[IVertex];
+		if(!VertexInstanceMarker[InstanceID])
+		{
+			SwapVertexNormal(InstanceID);
+			VertexInstanceMarker[InstanceID] = true;
+		}
 	}
 }
 
diff --git a/Engine/Plugins/Enterprise/DatasmithCADImporter/Source/CADLibrary/Private/MeshEditingWrapper.h b/Engine/Plugins/Enterprise/DatasmithCADImporter/Source/CADLibrary/Private/MeshEditingWrapper.h
--- a/Engine/Plugins/Enterprise/DatasmithCADImporter/Source/CADLibrary/Private/MeshEditingWrapper.h
+++ b/Engine/Plugins/Enterprise/DatasmithCADImporter/Source/CADLibrary/Private/MeshEditingWrapper.h
@@ -102,6 +102,9 @@ private:
 	TArray<FElementMetaData> EdgeMetaData;
 	TArray<FElementMetaData> TriangleMetaData;
 
+	// for orientation purpose
+	TBitArray<> VertexInstanceMarker;
+
 private:
 
 
