Fix editor crash after retesselation

The crash was due to an invalid MeshDescription of the StaticMesh.
CL 22470234 adds a second call to FStaticMeshSourceModel::CommitMeshDescription. But two successive calls of CommitMeshDescription empty every thing... That invalidates the MeshDescription.

diff --git a/Engine/Plugins/Enterprise/DatasmithCADImporter/Source/ParametricSurfaceExtension/Private/ParametricRetessellateAction.cpp b/Engine/Plugins/Enterprise/DatasmithCADImporter/Source/ParametricSurfaceExtension/Private/ParametricRetessellateAction.cpp
--- a/Engine/Plugins/Enterprise/DatasmithCADImporter/Source/ParametricSurfaceExtension/Private/ParametricRetessellateAction.cpp
+++ b/Engine/Plugins/Enterprise/DatasmithCADImporter/Source/ParametricSurfaceExtension/Private/ParametricRetessellateAction.cpp
@@ -65,11 +65,7 @@ void FParametricRetessellateAction_Impl::ApplyOnAssets(const TArray<FAssetData>&
 	TFunction<void(UStaticMesh*)> FinalizeChanges = [](UStaticMesh* StaticMesh) -> void
 	{
 		StaticMesh->PostEditChange();
-		UStaticMesh::FCommitMeshDescriptionParams Params;
-		Params.bMarkPackageDirty = true;
-		Params.bUseHashAsGuid = true;
-		StaticMesh->CommitMeshDescription(0, Params);
-		StaticMesh->ClearMeshDescription(0);
+		StaticMesh->MarkPackageDirty();
 
 		// Refresh associated editor 
 		TSharedPtr<IToolkit> EditingToolkit = FToolkitManager::Get().FindEditorForAsset(StaticMesh);
