Fixed PlmXml file crashes UE.
The crash is due to the call of operator[] in the ProductRevisionViewNodes map with a key that is not contained.
To fix it, a test is done to check if the map contains the key
otherwise, like ParseProductDef does, if ProductRevisionViewNodes map has the key, the ProductRevisionView is traversed.

diff --git a/Engine/Plugins/Enterprise/DatasmithCADImporter/Source/DatasmithPLMXMLTranslator/Private/DatasmithPlmXmlImporter.cpp b/Engine/Plugins/Enterprise/DatasmithCADImporter/Source/DatasmithPLMXMLTranslator/Private/DatasmithPlmXmlImporter.cpp
--- a/Engine/Plugins/Enterprise/DatasmithCADImporter/Source/DatasmithPLMXMLTranslator/Private/DatasmithPlmXmlImporter.cpp
+++ b/Engine/Plugins/Enterprise/DatasmithCADImporter/Source/DatasmithPLMXMLTranslator/Private/DatasmithPlmXmlImporter.cpp
@@ -764,7 +764,16 @@ namespace PlmXml
 
 		for (const FIdRef& InstanceRef : InstanceRefs)
 		{
-			TraverseProductInstance(MakeTraverseProductInstanceContext(Context, Context.ImportContext.GetProductInstance(InstanceRef)));
+			if (Context.ImportContext.HasProductInstance(InstanceRef))
+			{
+				TraverseProductInstance(MakeTraverseProductInstanceContext(Context, Context.ImportContext.GetProductInstance(InstanceRef)));
+			}
+			else if (Context.ImportContext.HasProductRevisionView(InstanceRef)) // in case InstanceRef is not a product instance
+			{
+				FTransform RootTransform;
+				RootTransform.SetIdentity();
+				TraverseProductRevisionView(MakeTraverseContext(Context, RootTransform, Context.ImportContext.GetProductRevisionView(InstanceRef)));
+			}
 		}
 	}
 
