Error when Importing STP file with TechSoft + CADKernel Tessellator

diff --git a/Engine/Plugins/Enterprise/DatasmithCADImporter/Source/CADInterfaces/Private/TechSoftFileParserCADKernelTessellator.cpp b/Engine/Plugins/Enterprise/DatasmithCADImporter/Source/CADInterfaces/Private/TechSoftFileParserCADKernelTessellator.cpp
--- a/Engine/Plugins/Enterprise/DatasmithCADImporter/Source/CADInterfaces/Private/TechSoftFileParserCADKernelTessellator.cpp
+++ b/Engine/Plugins/Enterprise/DatasmithCADImporter/Source/CADInterfaces/Private/TechSoftFileParserCADKernelTessellator.cpp
@@ -220,6 +220,7 @@ void FTechSoftFileParserCADKernelTessellator::GenerateBodyMesh(A3DRiRepresentati
 	FBody* CADKernelBody = TechSoftBridge.AddBody(Representation, ArchiveBody.MetaData, ArchiveBody.Unit);
 	if (CADKernelBody == nullptr)
 	{
+		ArchiveBody.Delete();
 		return;
 	}
 
@@ -261,6 +262,10 @@ void FTechSoftFileParserCADKernelTessellator::MeshAndGetTessellation(UE::CADKern
 	Mesher.MeshEntity(CADKernelBody);
 
 	FCADKernelTools::GetBodyTessellation(*CADKernelModelMesh, CADKernelBody, BodyMesh);
+	if (BodyMesh.Faces.Num() == 0)
+	{
+		ArchiveBody.Delete();
+	}
 
 	ArchiveBody.ColorFaceSet = BodyMesh.ColorSet;
 	ArchiveBody.MaterialFaceSet = BodyMesh.MaterialSet;
