CAD import crash due to high level of worker compared to available memory

The number of worker is defined by available RAM size / RecommandedRamPerWorker (6 GB);
If the user set the CVar MaxImportThreads, the number of worker is the min between  number of cores and the value defined by the user

During the CAD files parsing, the mesh caches files size is computed.
Based on this size, the memory used by the static mesh is estimated.
If this estimation is bigger than the free RAM size, a pop-up proposes to stop the import and only processed files are imported or to continue with the risk of running out of memory.

diff --git a/Engine/Plugins/Enterprise/DatasmithCADImporter/Source/DatasmithCADTranslator/Private/DatasmithCADTranslator.cpp b/Engine/Plugins/Enterprise/DatasmithCADImporter/Source/DatasmithCADTranslator/Private/DatasmithCADTranslator.cpp
--- a/Engine/Plugins/Enterprise/DatasmithCADImporter/Source/DatasmithCADTranslator/Private/DatasmithCADTranslator.cpp
+++ b/Engine/Plugins/Enterprise/DatasmithCADImporter/Source/DatasmithCADTranslator/Private/DatasmithCADTranslator.cpp
@@ -13,7 +13,9 @@
 #include "DatasmithSceneGraphBuilder.h"
 #include "DatasmithTranslator.h"
 #include "DatasmithUtils.h"
+#include "HAL/FileManager.h"
 #include "HAL/IConsoleManager.h"
+#include "HAL/PlatformProcess.h"
 #include "IDatasmithSceneElements.h"
 
 
@@ -226,19 +228,29 @@ bool FDatasmithCADTranslator::LoadScene(TSharedRef<IDatasmithScene> DatasmithSce
 	{
 		TMap<uint32, FString> CADFileToUEFileMap;
 		{
-			int32 NumCores = FPlatformMisc::NumberOfCores();
+			constexpr double RecommandedRamPerWorkers = 6.;
+			constexpr double OneGigaByte = 1024. * 1024. * 1024.;
+			const double AvailableRamGB = (double) (FPlatformMemory::GetStats().AvailablePhysical / OneGigaByte);
+
+			const int32 MaxNumberOfWorkers = FPlatformMisc::NumberOfCores();
+			const int32 RecommandedNumberOfWorkers = (int32) (AvailableRamGB / RecommandedRamPerWorkers + 0.5);
+
+			// UE recommendation 
+			int32 NumberOfWorkers = FMath::Min(MaxNumberOfWorkers, RecommandedNumberOfWorkers);
+
+			// User choice but limited by the number of cores. More brings nothing
 			if (GMaxImportThreads > 1)
 			{
-				if (FileDescriptor.CanReferenceOtherFiles())
-				{
-					NumCores = FMath::Min(GMaxImportThreads, NumCores);
-				}
-				else 
-				{
-					NumCores = 1;
-				}
+				NumberOfWorkers = FMath::Min(GMaxImportThreads, MaxNumberOfWorkers);
 			}
-			DatasmithDispatcher::FDatasmithDispatcher Dispatcher(ImportParameters, CachePath, NumCores, CADFileToUEFileMap, CADFileToUEGeomMap);
+
+			// If the file cannot have reference. One worker is enough.
+			if (GMaxImportThreads != 1 && !FileDescriptor.CanReferenceOtherFiles())
+			{
+				NumberOfWorkers = 1;
+			}
+
+			DatasmithDispatcher::FDatasmithDispatcher Dispatcher(ImportParameters, CachePath, NumberOfWorkers, CADFileToUEFileMap, CADFileToUEGeomMap);
 			Dispatcher.AddTask(FileDescriptor);
 
 			Dispatcher.Process(GMaxImportThreads != 1);
diff --git a/Engine/Plugins/Enterprise/DatasmithCADImporter/Source/DatasmithDispatcher/Private/DatasmithDispatcher.cpp b/Engine/Plugins/Enterprise/DatasmithCADImporter/Source/DatasmithDispatcher/Private/DatasmithDispatcher.cpp
--- a/Engine/Plugins/Enterprise/DatasmithCADImporter/Source/DatasmithDispatcher/Private/DatasmithDispatcher.cpp
+++ b/Engine/Plugins/Enterprise/DatasmithCADImporter/Source/DatasmithDispatcher/Private/DatasmithDispatcher.cpp
@@ -10,13 +10,46 @@
 
 #include "Algo/Count.h"
 #include "HAL/FileManager.h"
+#include "Misc/MessageDialog.h"
 #include "Misc/Paths.h"
 #include "Misc/ScopeLock.h"
 #include "HAL/IConsoleManager.h"
 
+#define LOCTEXT_NAMESPACE "DatasmithDispatcher"
+
 namespace DatasmithDispatcher
 {
 
+static uint64 EstimationOfMemoryUsedByStaticMeshes = 0;
+static uint64 AvailableMemory = 0;
+static bool bDisplayWorkerWarning = true;
+
+const EAppReturnType::Type OpenMessageDialog(uint64 AvailableRam, uint64 EstimationOfMemory, bool bWorkerWarning)
+{
+	const FText StaticMeshSizeTxt = FText::AsMemory(EstimationOfMemoryUsedByStaticMeshes, EMemoryUnitStandard::IEC);
+	const FText RamSizeTxt = FText::AsMemory(AvailableRam, EMemoryUnitStandard::IEC);
+
+	const FText WarningMessageTxt = LOCTEXT("WarningNotEnoughRamForLocalImport", "Warning: DatasmithCADWorker are currently not used.");
+	const FText MainMessageTxt = LOCTEXT("NotEnoughRamForLocalImport",
+		"Estimated memory used by static meshes is greater than available RAM:\n"
+		"\t{0} available vs {1} estimated.\n"
+		"Click OK to continue importing or CANCEL to get only what have already been processed.\n"
+		"If you want to continue, be sure to allocate enough virtual memory to avoid running out of memory.\n"
+		"If you cancel, check the output log for more information.\n"
+		"{2}");
+
+	const FText Message = FText::Format(MainMessageTxt, RamSizeTxt, StaticMeshSizeTxt, bWorkerWarning ? WarningMessageTxt : FText());
+	const EAppReturnType::Type Choice = FMessageDialog::Open(EAppMsgType::OkCancel, Message);
+
+	if (Choice == EAppReturnType::Cancel)
+	{
+		UE_LOG(LogDatasmithDispatcher, Error, TEXT("Estimated memory used by static meshes is greater than available RAM (%s available vs %s estimated). Import process is interrupted."), *RamSizeTxt.ToString(), *StaticMeshSizeTxt.ToString());
+		UE_LOG(LogDatasmithDispatcher, Error, TEXT("Import of the files below have been aborted:"));
+	}
+	return Choice;
+}
+
+
 FDatasmithDispatcher::FDatasmithDispatcher(const CADLibrary::FImportParameters& InImportParameters, const FString& InCacheDir, int32 InNumberOfWorkers, TMap<uint32, FString>& OutCADFileToUnrealFileMap, TMap<uint32, FString>& OutCADFileToUnrealGeomMap)
 	: NextTaskIndex(0)
 	, CompletedTaskCount(0)
@@ -129,6 +162,45 @@ void FDatasmithDispatcher::SetTaskState(int32 TaskIndex, ETaskState TaskState)
 
 void FDatasmithDispatcher::Process(bool bWithProcessor)
 {
+	EstimationOfMemoryUsedByStaticMeshes = 0;
+	bool bCheckMemory = true;
+
+	// This function is added in 5.1.1 but is cloned with the lines 370 to 387
+	TFunction<bool()> CanImportContinue = [&]() -> bool
+	{
+		AvailableMemory = FPlatformMemory::GetStats().AvailablePhysical;
+		if (AvailableMemory < EstimationOfMemoryUsedByStaticMeshes)
+		{
+			const EAppReturnType::Type Choice = OpenMessageDialog(AvailableMemory, EstimationOfMemoryUsedByStaticMeshes, bDisplayWorkerWarning);
+
+			if (Choice == EAppReturnType::Cancel)
+			{
+				while (TOptional<FTask> Task = GetNextTask())
+				{
+					UE_LOG(LogDatasmithDispatcher, Error, TEXT("   - %s"), *Task->FileDescription.GetFileName());
+				}
+
+				CloseHandlers();
+
+				// Log last running tasks as failed
+				{
+					FScopeLock Lock(&TaskPoolCriticalSection);
+					for (FTask& Task : TaskPool)
+					{
+						if (Task.State == ETaskState::Running)
+						{
+							UE_LOG(LogDatasmithDispatcher, Error, TEXT("   - %s"), *Task.FileDescription.GetFileName());
+						}
+					}
+					CompletedTaskCount = TaskPool.Num();
+				}
+				return false;
+			}
+			bCheckMemory = false;
+		}
+		return true;
+	};
+
 	// Temporary code to validate that DatasmithCADWorker.exe does exist before triggering the multi-processing
 	// A static method of FDatasmithWorkerHandler, i.e. CanStart, should provide this information.
 	{
@@ -144,6 +216,8 @@ void FDatasmithDispatcher::Process(bool bWithProcessor)
 		bWithProcessor &= FPaths::FileExists(WorkerPath);
 	}
 
+	bDisplayWorkerWarning = !bWithProcessor;
+
 	if (bWithProcessor)
 	{
 		SpawnHandlers();
@@ -151,6 +225,11 @@ void FDatasmithDispatcher::Process(bool bWithProcessor)
 		bool bLogRestartError = true;
 		while (!IsOver())
 		{
+			if (bCheckMemory && !CanImportContinue())
+			{
+				break;
+			}
+
 			bool bHasAliveWorker = false;
 			for (FDatasmithWorkerHandler& Handler : WorkerHandlers)
 			{
@@ -219,8 +298,13 @@ void FDatasmithDispatcher::LinkCTFileToUnrealCacheFile(const CADLibrary::FFileDe
 	}
 	if (!UnrealMeshFile.IsEmpty())
 	{
+		constexpr int64 RatioCacheSizeMeshSize = 11; // 11 is based on tests done on several huge DMUs
 		CADFileToUnrealGeomMap.Add(FileHash, UnrealMeshFile);
+		FString FilePath = FPaths::Combine(ProcessCacheFolder, TEXT("mesh"), UnrealMeshFile + TEXT(".gm"));
+		FFileStatData FileStatData = IFileManager::Get().GetStatData(*FilePath);
+		EstimationOfMemoryUsedByStaticMeshes += (uint64)(FileStatData.FileSize * RatioCacheSizeMeshSize);
 	}
+
 }
 
 int32 FDatasmithDispatcher::GetNextWorkerId()
@@ -253,6 +337,7 @@ void FDatasmithDispatcher::CloseHandlers()
 
 void FDatasmithDispatcher::ProcessLocal()
 {
+	bool bCheckMemory = true;
 	while (TOptional<FTask> Task = GetNextTask())
 	{
 		CADLibrary::FFileDescriptor& FileDescription = Task->FileDescription;
@@ -274,9 +359,32 @@ void FDatasmithDispatcher::ProcessLocal()
 					AddTask(ExternalFile);
 				}
 			}
+
+			// update the available memory to get into account possible memory leak
+			AvailableMemory = FPlatformMemory::GetStats().AvailablePhysical;
+
 			LinkCTFileToUnrealCacheFile(CADFileData.GetCADFileDescription(), CADFileData.GetSceneGraphFileName(), CADFileData.GetMeshFileName());
+
+			// This block is added in 5.1.1 but is cloned with function of the lines 166 to 204
+			// This will be refactor in 5.2 (jira UE-172253)  
+			if (bCheckMemory && AvailableMemory < EstimationOfMemoryUsedByStaticMeshes)
+			{
+				const EAppReturnType::Type Choice = OpenMessageDialog(AvailableMemory, EstimationOfMemoryUsedByStaticMeshes, bDisplayWorkerWarning);
+
+				if (Choice == EAppReturnType::Cancel)
+				{
+					while ((Task = GetNextTask()))
+					{
+						UE_LOG(LogDatasmithDispatcher, Error, TEXT("   - %s"), *Task->FileDescription.GetFileName());
+					}
+					break;
+				}
+				bCheckMemory = false;
+			}
 		}
 	}
 }
 
-} // ns DatasmithDispatcher
\ No newline at end of file
+} // ns DatasmithDispatcher
+
+#undef LOCTEXT_NAMESPACE
\ No newline at end of file
