Fixed typo in variable name bIsPlat*e*formSupported
#9442

diff --git a/Engine/Plugins/Enterprise/DatasmithCADImporter/Source/CADLibrary/CADLibrary.Build.cs b/Engine/Plugins/Enterprise/DatasmithCADImporter/Source/CADLibrary/CADLibrary.Build.cs
--- a/Engine/Plugins/Enterprise/DatasmithCADImporter/Source/CADLibrary/CADLibrary.Build.cs
+++ b/Engine/Plugins/Enterprise/DatasmithCADImporter/Source/CADLibrary/CADLibrary.Build.cs
@@ -24,8 +24,5 @@ public class CADLibrary : ModuleRules
 				"StaticMeshDescription"
 			}
 		);
-
-		// Support for Windows only
-		bool bIsPlateformSupported = Target.Platform == UnrealTargetPlatform.Win64;
 	}
 }
