Set verbose level of messages in FCADFileData
This allows to set the verbose level of non critical message to "Log" instead of "Warning"

diff --git a/Engine/Plugins/Enterprise/DatasmithCADImporter/Source/CADInterfaces/Private/CADFileReader.cpp b/Engine/Plugins/Enterprise/DatasmithCADImporter/Source/CADInterfaces/Private/CADFileReader.cpp
--- a/Engine/Plugins/Enterprise/DatasmithCADImporter/Source/CADInterfaces/Private/CADFileReader.cpp
+++ b/Engine/Plugins/Enterprise/DatasmithCADImporter/Source/CADInterfaces/Private/CADFileReader.cpp
@@ -87,7 +87,7 @@ namespace CADLibrary
 			return true;
 		}
 
-		CADFileData.AddWarningMessages(FString::Printf(TEXT("File %s cannot be found."), *File.GetFileName()));
+		CADFileData.LogWarning(FString::Printf(TEXT("File %s cannot be found."), *File.GetFileName()));
 		return false;
 	}
 
diff --git a/Engine/Plugins/Enterprise/DatasmithCADImporter/Source/CADInterfaces/Private/TechSoftFileParser.cpp b/Engine/Plugins/Enterprise/DatasmithCADImporter/Source/CADInterfaces/Private/TechSoftFileParser.cpp
--- a/Engine/Plugins/Enterprise/DatasmithCADImporter/Source/CADInterfaces/Private/TechSoftFileParser.cpp
+++ b/Engine/Plugins/Enterprise/DatasmithCADImporter/Source/CADInterfaces/Private/TechSoftFileParser.cpp
@@ -348,31 +348,31 @@ ECADParsingResult FTechSoftFileParser::Process()
 		{
 		case A3DStatus::A3D_LOAD_FILE_TOO_OLD:
 		{
-			CADFileData.AddWarningMessages(FString::Printf(TEXT("File %s hasn't been loaded because the version is less than the oldest supported version."), *File.GetFileName()));
+			CADFileData.LogWarning(FString::Printf(TEXT("File %s hasn't been loaded because the version is less than the oldest supported version."), *File.GetFileName()));
 			break;
 		}
 
 		case A3DStatus::A3D_LOAD_FILE_TOO_RECENT:
 		{
-			CADFileData.AddWarningMessages(FString::Printf(TEXT("File %s hasn't been loaded because the version is more recent than supported version."), *File.GetFileName()));
+			CADFileData.LogWarning(FString::Printf(TEXT("File %s hasn't been loaded because the version is more recent than supported version."), *File.GetFileName()));
 			break;
 		}
 
 		case A3DStatus::A3D_LOAD_CANNOT_ACCESS_CADFILE:
 		{
-			CADFileData.AddWarningMessages(FString::Printf(TEXT("File %s hasn't been loaded because the input path cannot be opened by the running process for reading."), *File.GetFileName()));
+			CADFileData.LogWarning(FString::Printf(TEXT("File %s hasn't been loaded because the input path cannot be opened by the running process for reading."), *File.GetFileName()));
 			break;
 		}
 
 		case A3DStatus::A3D_LOAD_INVALID_FILE_FORMAT:
 		{
-			CADFileData.AddWarningMessages(FString::Printf(TEXT("File %s hasn't been loaded because the format is not supported."), *File.GetFileName()));
+			CADFileData.LogWarning(FString::Printf(TEXT("File %s hasn't been loaded because the format is not supported."), *File.GetFileName()));
 			break;
 		}
 
 		default:
 		{
-			CADFileData.AddWarningMessages(FString::Printf(TEXT("File %s hasn't been loaded because an error occured while reading the file."), *File.GetFileName()));
+			CADFileData.LogWarning(FString::Printf(TEXT("File %s hasn't been loaded because an error occured while reading the file."), *File.GetFileName()));
 			break;
 		}
 		};
@@ -469,7 +469,7 @@ void FTechSoftFileParser::GenerateBodyMesh(A3DRiRepresentationItem* Representati
 		A3DStatus Status = TechSoftInterface::SewBReps(&Representation, BRepCount, CADLibrary::FImportParameters::GStitchingTolerance, FileUnit, SewData.GetPtr(), &NewBReps, NewBRepCount);
 		if (Status != A3DStatus::A3D_SUCCESS)
 		{
-			CADFileData.AddWarningMessages(TEXT("A body healing failed. A body could be missing."));
+			CADFileData.LogWarning(TEXT("A body healing failed. A body could be missing."));
 		}
 	}
 
@@ -612,13 +612,13 @@ ECADParsingResult FTechSoftFileParser::TraverseModel()
 
 	if(ModelFileData->m_uiPOccurrencesSize == 0)
 	{
-		CADFileData.AddWarningMessages(FString::Printf(TEXT("File %s is empty."), *CADFileData.GetCADFileDescription().GetFileName()));
+		CADFileData.LogWarning(FString::Printf(TEXT("File %s is empty."), *CADFileData.GetCADFileDescription().GetFileName()));
 		return ECADParsingResult::ProcessFailed;
 	}
 
 	if (ModelFileData->m_uiPOccurrencesSize > 1)
 	{
-		CADFileData.AddWarningMessages(FString::Printf(TEXT("File %s has many root components, only the first is loaded."), *CADFileData.GetCADFileDescription().GetFileName()));
+		CADFileData.LogWarning(FString::Printf(TEXT("File %s has many root components, only the first is loaded."), *CADFileData.GetCADFileDescription().GetFileName()));
 	}
 
 	if (IsConfigurationSet(ModelFileData->m_ppPOccurrences[0]))
diff --git a/Engine/Plugins/Enterprise/DatasmithCADImporter/Source/CADInterfaces/Private/TechSoftFileParserCADKernelTessellator.cpp b/Engine/Plugins/Enterprise/DatasmithCADImporter/Source/CADInterfaces/Private/TechSoftFileParserCADKernelTessellator.cpp
--- a/Engine/Plugins/Enterprise/DatasmithCADImporter/Source/CADInterfaces/Private/TechSoftFileParserCADKernelTessellator.cpp
+++ b/Engine/Plugins/Enterprise/DatasmithCADImporter/Source/CADInterfaces/Private/TechSoftFileParserCADKernelTessellator.cpp
@@ -339,12 +339,12 @@ A3DStatus FTechSoftFileParserCADKernelTessellator::AdaptBRepModel()
 	A3DStatus Ret = TechSoftInterface::AdaptBRepInModelFile(ModelFile.Get(), *CopyAndAdaptBrepModelData, ErrorCount, &Errors);
 	if ((Ret == A3D_SUCCESS || Ret == A3D_TOOLS_CONTINUE_ON_ERROR) && ErrorCount > 0)
 	{
-		// Add warning about error during the adaptation
-		CADFileData.AddWarningMessages(FString::Printf(TEXT("File %s had %d error during BRep adaptation step."), *CADFileData.GetCADFileDescription().GetFileName(), ErrorCount));
+		// Add message about non-critical errors during the adaptation
+		CADFileData.LogMessage(FString::Printf(TEXT("File %s had %d non-critical error(s) during BRep adaptation step."), *CADFileData.GetCADFileDescription().GetFileName(), ErrorCount));
 	}
 	else if (Ret != A3D_SUCCESS)
 	{
-		CADFileData.AddWarningMessages(FString::Printf(TEXT("File %s failed during BRep adaptation step."), *CADFileData.GetCADFileDescription().GetFileName(), ErrorCount));
+		CADFileData.LogWarning(FString::Printf(TEXT("File %s failed during BRep adaptation step."), *CADFileData.GetCADFileDescription().GetFileName(), ErrorCount));
 		return A3D_ERROR;
 	}
 	return A3D_SUCCESS;
diff --git a/Engine/Plugins/Enterprise/DatasmithCADImporter/Source/CADInterfaces/Public/CADFileData.h b/Engine/Plugins/Enterprise/DatasmithCADImporter/Source/CADInterfaces/Public/CADFileData.h
--- a/Engine/Plugins/Enterprise/DatasmithCADImporter/Source/CADInterfaces/Public/CADFileData.h
+++ b/Engine/Plugins/Enterprise/DatasmithCADImporter/Source/CADInterfaces/Public/CADFileData.h
@@ -103,9 +103,31 @@ public:
 		return bIsCacheDefined;
 	}
 
-	void AddWarningMessages(const FString& Message)
+	void LogError(const FString& Message)
 	{
-		WarningMessages.Add(Message);
+		Messages.Emplace(ELogVerbosity::Error, Message);
+	}
+
+	void LogWarning(const FString& Message)
+	{
+		Messages.Emplace(ELogVerbosity::Warning, Message);
+	}
+
+	/** Logs a message to console (and log file) */
+	void LogDisplayMessage(const FString& Message)
+	{
+		Messages.Emplace(ELogVerbosity::Display, Message);
+	}
+
+	/** logs a message to a log file (does not print to console) */
+	void LogMessage(const FString& Message)
+	{
+		Messages.Emplace(ELogVerbosity::Log, Message);
+	}
+
+	void LogVerboseMessage(const FString& Message)
+	{
+		Messages.Emplace(ELogVerbosity::Verbose, Message);
 	}
 
 	void LoadSceneGraphArchive()
@@ -171,9 +193,9 @@ public:
 		return MeshArchiveFile;
 	}
 
-	const TArray<FString>& GetWarningMessages() const
+	const TArray<TPair<uint8, FString>>& GetMessages() const
 	{
-		return WarningMessages;
+		return Messages;
 	}
 
 	const FArchiveSceneGraph& GetSceneGraphArchive() const
@@ -258,7 +280,7 @@ private:
 	FArchiveSceneGraph SceneGraphArchive;
 	TArray<FBodyMesh> BodyMeshes;
 
-	TArray<FString> WarningMessages;
+	TArray<TPair<uint8, FString>> Messages;
 
 	mutable uint32 SceneFileHash = 0;
 	mutable uint32 GeomFileHash = 0;
diff --git a/Engine/Plugins/Enterprise/DatasmithCADImporter/Source/DatasmithDispatcher/Private/DatasmithCommands.cpp b/Engine/Plugins/Enterprise/DatasmithCADImporter/Source/DatasmithDispatcher/Private/DatasmithCommands.cpp
--- a/Engine/Plugins/Enterprise/DatasmithCADImporter/Source/DatasmithDispatcher/Private/DatasmithCommands.cpp
+++ b/Engine/Plugins/Enterprise/DatasmithCADImporter/Source/DatasmithDispatcher/Private/DatasmithCommands.cpp
@@ -56,7 +56,7 @@ namespace DatasmithDispatcher
 		Ar << ProcessResult;
 		Ar << SceneGraphFileName;
 		Ar << GeomFileName;
-		Ar << WarningMessages;
+		Ar << Messages;
 	}
 
 	void FImportParametersCommand::SerializeImpl(FArchive& Ar)
diff --git a/Engine/Plugins/Enterprise/DatasmithCADImporter/Source/DatasmithDispatcher/Private/DatasmithDispatcher.cpp b/Engine/Plugins/Enterprise/DatasmithCADImporter/Source/DatasmithDispatcher/Private/DatasmithDispatcher.cpp
--- a/Engine/Plugins/Enterprise/DatasmithCADImporter/Source/DatasmithDispatcher/Private/DatasmithDispatcher.cpp
+++ b/Engine/Plugins/Enterprise/DatasmithCADImporter/Source/DatasmithDispatcher/Private/DatasmithDispatcher.cpp
@@ -85,11 +85,35 @@ void FDatasmithDispatcher::AddTask(const CADLibrary::FFileDescriptor& InFileDesc
 	TaskPool[TaskIndex].Index = TaskIndex;
 }
 
-void FDatasmithDispatcher::LogWarningMessages(const TArray<FString>& WarningMessages) const
+void FDatasmithDispatcher::LogMessages(const TArray<TPair<uint8, FString>>& Messages) const
 {
-	for (const FString& WarningMessage : WarningMessages)
+	for (const TPair<uint8, FString>& Message : Messages)
 	{
-		UE_LOG(LogDatasmithDispatcher, Warning, TEXT("%s"), *WarningMessage);
+		switch (Message.Key)
+		{
+		case ELogVerbosity::Fatal:
+			UE_LOG(LogDatasmithDispatcher, Fatal, TEXT("%s"), *Message.Value);
+			break;
+		case ELogVerbosity::Error:
+			UE_LOG(LogDatasmithDispatcher, Error, TEXT("%s"), *Message.Value);
+			break;
+		case ELogVerbosity::Warning:
+			UE_LOG(LogDatasmithDispatcher, Warning, TEXT("%s"), *Message.Value);
+			break;
+		case ELogVerbosity::Display:
+			UE_LOG(LogDatasmithDispatcher, Display, TEXT("%s"), *Message.Value);
+			break;
+		case ELogVerbosity::Log:
+			UE_LOG(LogDatasmithDispatcher, Log, TEXT("%s"), *Message.Value);
+			break;
+		case ELogVerbosity::Verbose:
+			UE_LOG(LogDatasmithDispatcher, Verbose, TEXT("%s"), *Message.Value);
+			break;
+		case ELogVerbosity::VeryVerbose:
+		default:
+			UE_LOG(LogDatasmithDispatcher, VeryVerbose, TEXT("%s"), *Message.Value);
+			break;
+		}
 	}
 }
 
diff --git a/Engine/Plugins/Enterprise/DatasmithCADImporter/Source/DatasmithDispatcher/Private/DatasmithWorkerHandler.cpp b/Engine/Plugins/Enterprise/DatasmithCADImporter/Source/DatasmithDispatcher/Private/DatasmithWorkerHandler.cpp
--- a/Engine/Plugins/Enterprise/DatasmithCADImporter/Source/DatasmithDispatcher/Private/DatasmithWorkerHandler.cpp
+++ b/Engine/Plugins/Enterprise/DatasmithCADImporter/Source/DatasmithDispatcher/Private/DatasmithWorkerHandler.cpp
@@ -171,6 +171,7 @@ void FDatasmithWorkerHandler::RunInternal()
 
 				if (CurrentTask.IsSet())
 				{
+					UE_LOG(LogDatasmithDispatcher, Log, TEXT("   - Next task: %s"), *CurrentTask->FileDescription.GetFileName());
 					FRunTaskCommand NewTask(CurrentTask.GetValue());
 
 					if (CommandIO.SendCommand(NewTask, Config::SendCommandTimeout_s))
@@ -339,8 +340,11 @@ void FDatasmithWorkerHandler::ProcessCommand(FCompletedTaskCommand& CompletedTas
 	}
 
 	Dispatcher.SetTaskState(CurrentTask->Index, CompletedTaskCommand.ProcessResult);
+
+	UE_LOG(LogDatasmithDispatcher, Log, TEXT("   - Task completed: %s"), *CurrentTask->FileDescription.GetFileName());
+
 	Dispatcher.LinkCTFileToUnrealCacheFile(CurrentTask->FileDescription, CompletedTaskCommand.SceneGraphFileName, CompletedTaskCommand.GeomFileName);
-	Dispatcher.LogWarningMessages(CompletedTaskCommand.WarningMessages);
+	Dispatcher.LogMessages(CompletedTaskCommand.Messages);
 	CurrentTask.Reset();
 
 	if (CompletedTaskCommand.ProcessResult == ETaskState::Unknown)
diff --git a/Engine/Plugins/Enterprise/DatasmithCADImporter/Source/DatasmithDispatcher/Public/DatasmithCommands.h b/Engine/Plugins/Enterprise/DatasmithCADImporter/Source/DatasmithDispatcher/Public/DatasmithCommands.h
--- a/Engine/Plugins/Enterprise/DatasmithCADImporter/Source/DatasmithDispatcher/Public/DatasmithCommands.h
+++ b/Engine/Plugins/Enterprise/DatasmithCADImporter/Source/DatasmithDispatcher/Public/DatasmithCommands.h
@@ -103,7 +103,7 @@ public:
 	FString SceneGraphFileName;
 	FString GeomFileName;
 	ETaskState ProcessResult = ETaskState::Unknown;
-	TArray<FString> WarningMessages;
+	TArray<TPair<uint8, FString>> Messages;
 };
 
 class FImportParametersCommand : public ICommand
diff --git a/Engine/Plugins/Enterprise/DatasmithCADImporter/Source/DatasmithDispatcher/Public/DatasmithDispatcher.h b/Engine/Plugins/Enterprise/DatasmithCADImporter/Source/DatasmithDispatcher/Public/DatasmithDispatcher.h
--- a/Engine/Plugins/Enterprise/DatasmithCADImporter/Source/DatasmithDispatcher/Public/DatasmithDispatcher.h
+++ b/Engine/Plugins/Enterprise/DatasmithCADImporter/Source/DatasmithDispatcher/Public/DatasmithDispatcher.h
@@ -27,7 +27,7 @@ public:
 
 	void LinkCTFileToUnrealCacheFile(const CADLibrary::FFileDescriptor& CTFileDescription, const FString& UnrealSceneGraphFile, const FString& UnrealGeomFile);
 
-	void LogWarningMessages(const TArray<FString>& Warnings) const;
+	void LogMessages(const TArray<TPair<uint8, FString>>& Warnings) const;
 
 private:
 	void SpawnHandlers();
diff --git a/Engine/Source/Programs/Enterprise/Datasmith/DatasmithCADWorker/Private/DatasmithCADWorkerImpl.cpp b/Engine/Source/Programs/Enterprise/Datasmith/DatasmithCADWorker/Private/DatasmithCADWorkerImpl.cpp
--- a/Engine/Source/Programs/Enterprise/Datasmith/DatasmithCADWorker/Private/DatasmithCADWorkerImpl.cpp
+++ b/Engine/Source/Programs/Enterprise/Datasmith/DatasmithCADWorker/Private/DatasmithCADWorkerImpl.cpp
@@ -182,7 +182,7 @@ void FDatasmithCADWorkerImpl::ProcessCommand(const FRunTaskCommand& RunTaskComma
 		CompletedTask.ExternalReferences = CADFileData.GetExternalRefSet();
 		CompletedTask.SceneGraphFileName = CADFileData.GetSceneGraphFileName();
 		CompletedTask.GeomFileName = CADFileData.GetMeshFileName();
-		CompletedTask.WarningMessages = CADFileData.GetWarningMessages();
+		CompletedTask.Messages = CADFileData.GetMessages();
 
 		UE_LOG(LogDatasmithCADWorker, Verbose, TEXT("=> Process %s %s saved into %s%s and %s%s."), *FileToProcess.GetFileName(), *FileToProcess.GetConfiguration(), *CompletedTask.SceneGraphFileName, TEXT(".sg"), *CompletedTask.GeomFileName, TEXT(".gm"));
 		UE_LOG(LogDatasmithCADWorker, Verbose, TEXT("     It generates %d bodies"), CADFileData.GetBodyMeshes().Num());
