Fix "Delete unused ArchiveBody": it was done too early so some bodies were deleted while they were valid

diff --git a/Engine/Plugins/Enterprise/DatasmithCADImporter/Source/CADInterfaces/Private/TechSoftFileParserCADKernelTessellator.cpp b/Engine/Plugins/Enterprise/DatasmithCADImporter/Source/CADInterfaces/Private/TechSoftFileParserCADKernelTessellator.cpp
--- a/Engine/Plugins/Enterprise/DatasmithCADImporter/Source/CADInterfaces/Private/TechSoftFileParserCADKernelTessellator.cpp
+++ b/Engine/Plugins/Enterprise/DatasmithCADImporter/Source/CADInterfaces/Private/TechSoftFileParserCADKernelTessellator.cpp
@@ -73,6 +73,15 @@ void FTechSoftFileParserCADKernelTessellator::SewAndGenerateBodyMeshes()
 	{
 		SewAndMesh(Representations.Value);
 	}
+
+	// Delete unused ArchiveBody
+	for (FArchiveBody& ArchiveBody : SceneGraph.Bodies)
+	{
+		if (ArchiveBody.MeshActorUId == 0 && !ArchiveBody.IsDeleted())
+		{
+			ArchiveBody.Delete();
+		}
+	}
 }
 
 void FTechSoftFileParserCADKernelTessellator::SewAndMesh(TArray<A3DRiRepresentationItem*>& Representations)
@@ -195,15 +204,6 @@ void FTechSoftFileParserCADKernelTessellator::SewAndMesh(TArray<A3DRiRepresentat
 			ArchiveBody.Delete();
 		}
 	}
-
-	// Delete unused ArchiveBody
-	for (FArchiveBody& ArchiveBody : SceneGraph.Bodies)
-	{
-		if (ArchiveBody.MeshActorUId == 0)
-		{
-			ArchiveBody.Delete();
-		}
-	}
 }
 
 void FTechSoftFileParserCADKernelTessellator::GenerateBodyMesh(A3DRiRepresentationItem* Representation, FArchiveBody& ArchiveBody)
