Fix Alias Analytics log

diff --git a/Engine/Plugins/Enterprise/DatasmithCADImporter/Source/DatasmithWireTranslator/Private/DatasmithWireTranslator.cpp b/Engine/Plugins/Enterprise/DatasmithCADImporter/Source/DatasmithWireTranslator/Private/DatasmithWireTranslator.cpp
--- a/Engine/Plugins/Enterprise/DatasmithCADImporter/Source/DatasmithWireTranslator/Private/DatasmithWireTranslator.cpp
+++ b/Engine/Plugins/Enterprise/DatasmithCADImporter/Source/DatasmithWireTranslator/Private/DatasmithWireTranslator.cpp
@@ -113,31 +113,31 @@ const uint64 LibAlias2019_Version = 5000000000000000;
 #if defined(OPEN_MODEL_2020)
 const uint64 LibAliasVersionMin = LibAlias2019_Version;
 const uint64 LibAliasVersionMax = LibAlias2021_3_0_Version;
-const FString AliasVersionChar = TEXT("AliasStudio 2020, Model files");
+const FString AliasSdkVersion   = TEXT("2020");
 #elif defined(OPEN_MODEL_2021_3)
 const uint64 LibAliasVersionMin = LibAlias2021_3_0_Version;
 const uint64 LibAliasVersionMax = LibAlias2022_0_1_Version;
-const FString  AliasVersionChar = TEXT("AliasStudio 2021.3, Model files");
+const FString AliasSdkVersion   = TEXT("2021.3");
 #elif defined(OPEN_MODEL_2022)
 const uint64 LibAliasVersionMin = LibAlias2022_0_1_Version;
 const uint64 LibAliasVersionMax = LibAlias2022_1_0_Version;
-const FString AliasVersionChar = TEXT("AliasStudio 2022, Model files");
+const FString AliasSdkVersion   = TEXT("2022");
 #elif defined(OPEN_MODEL_2022_1)
 const uint64 LibAliasVersionMin = LibAlias2022_1_0_Version;
 const uint64 LibAliasVersionMax = LibAlias2022_2_0_Version;
-const FString AliasVersionChar = TEXT("AliasStudio 2022.1, Model files");
+const FString AliasSdkVersion   = TEXT("2022.1");
 #elif defined(OPEN_MODEL_2022_2)
 const uint64 LibAliasVersionMin = LibAlias2022_2_0_Version;
 const uint64 LibAliasVersionMax = LibAlias2023_0_0_Version;
-const FString AliasVersionChar = TEXT("AliasStudio 2022.2, Model files");
+const FString AliasSdkVersion   = TEXT("2022.2");
 #elif defined(OPEN_MODEL_2023_0)
 const uint64 LibAliasVersionMin = LibAlias2023_0_0_Version;
 const uint64 LibAliasVersionMax = LibAlias2023_1_0_Version;
-const FString AliasVersionChar = TEXT("AliasStudio 2023.0, Model files");
-#elif defined(OPEN_MODEL_2023_1)
+const FString AliasSdkVersion   = TEXT("2023.0");
+#elif defined(OPEN_MODEL_2023_1)      
 const uint64 LibAliasVersionMin = LibAlias2023_1_0_Version;
 const uint64 LibAliasVersionMax = LibAliasNext_Version;
-const FString AliasVersionChar = TEXT("AliasStudio 2023.1, Model files");
+const FString AliasSdkVersion   = TEXT("2023.1");
 #endif
 
 // Alias material management (to allow sew of BReps of different materials):
@@ -291,58 +291,13 @@ public:
 		, SceneFileHash(0)
 		, AlRootNode(nullptr)
 	{
-		// Set ProductName, ProductVersion in DatasmithScene for Analytics purpose
-		uint64 AliasFileVersion = FPlatformMisc::GetFileVersion(TEXT("libalias_api.dll"));
 
+		const FString AliasProductVersion = FString::Printf(TEXT("Alias %s"), *AliasSdkVersion);
 		DatasmithScene->SetHost(TEXT("Alias"));
 		DatasmithScene->SetVendor(TEXT("Autodesk"));
 		DatasmithScene->SetProductName(TEXT("Alias Tools"));
-
-		if (AliasFileVersion < LibAlias2020_Version)
-		{
-			DatasmithScene->SetExporterSDKVersion(TEXT("2019"));
-			DatasmithScene->SetProductVersion(TEXT("Alias 2019"));
-		}
-		else if (AliasFileVersion < LibAlias2021_Version)
-		{
-			DatasmithScene->SetExporterSDKVersion(TEXT("2020"));
-			DatasmithScene->SetProductVersion(TEXT("Alias 2020"));
-		}
-		else if (AliasFileVersion < LibAlias2021_3_0_Version)
-		{
-			DatasmithScene->SetExporterSDKVersion(TEXT("2021.0"));
-			DatasmithScene->SetProductVersion(TEXT("Alias 2021.0"));
-		}
-		else if (AliasFileVersion < LibAlias2021_3_1_Version)
-		{
-			DatasmithScene->SetExporterSDKVersion(TEXT("2021.3.0"));
-			DatasmithScene->SetProductVersion(TEXT("Alias 2021.3.0"));
-		}
-		else if (AliasFileVersion < LibAlias2021_3_2_Version)
-		{
-			DatasmithScene->SetExporterSDKVersion(TEXT("2021.3.1"));
-			DatasmithScene->SetProductVersion(TEXT("Alias 2021.3.1"));
-		}
-		else if (AliasFileVersion < LibAlias2022_0_1_Version)
-		{
-			DatasmithScene->SetExporterSDKVersion(TEXT("2021.3.2"));
-			DatasmithScene->SetProductVersion(TEXT("Alias 2021.3.2"));
-		}
-		else if (AliasFileVersion < LibAlias2022_1_0_Version)
-		{
-			DatasmithScene->SetExporterSDKVersion(TEXT("2022"));
-			DatasmithScene->SetProductVersion(TEXT("Alias 2022"));
-		}
-		else if (AliasFileVersion < LibAlias2022_2_0_Version)
-		{
-			DatasmithScene->SetExporterSDKVersion(TEXT("2022.1"));
-			DatasmithScene->SetProductVersion(TEXT("Alias 2022.1"));
-		}
-		else
-		{
-			DatasmithScene->SetExporterSDKVersion(TEXT("2022.2"));
-			DatasmithScene->SetProductVersion(TEXT("Alias 2022.2"));
-		}
+		DatasmithScene->SetExporterSDKVersion(*AliasSdkVersion);
+		DatasmithScene->SetProductVersion(*AliasProductVersion);
 
 		CADLibrary::FImportParameters ImportParameters;
 		if (CADLibrary::FImportParameters::bGDisableCADKernelTessellation)
@@ -2434,7 +2389,8 @@ void FDatasmithWireTranslator::Initialize(FDatasmithTranslatorCapabilities& OutC
 
 			if (LibAliasVersionMin <= FileVersion && FileVersion < LibAliasVersionMax)
 			{
-				OutCapabilities.SupportedFileFormats.Add(FFileFormatInfo{ TEXT("wire"), *AliasVersionChar });
+				const FString AliasVersion = FString::Printf(TEXT("AliasStudio %s model files"), *AliasSdkVersion);
+				OutCapabilities.SupportedFileFormats.Add(FFileFormatInfo{ TEXT("wire"), *AliasVersion });
 				OutCapabilities.bIsEnabled = true;
 				return;
 			}
