Add CVar to enable all unsupported CAD formats
Add needed libraries for these new formats

diff --git a/Engine/Plugins/Enterprise/DatasmithCADImporter/Source/DatasmithCADTranslator/Private/DatasmithCADTranslator.cpp b/Engine/Plugins/Enterprise/DatasmithCADImporter/Source/DatasmithCADTranslator/Private/DatasmithCADTranslator.cpp
--- a/Engine/Plugins/Enterprise/DatasmithCADImporter/Source/DatasmithCADTranslator/Private/DatasmithCADTranslator.cpp
+++ b/Engine/Plugins/Enterprise/DatasmithCADImporter/Source/DatasmithCADTranslator/Private/DatasmithCADTranslator.cpp
@@ -29,6 +29,16 @@ FAutoConsoleVariableRef GCADTranslatorEnableNativeIFCTranslator(
 	bGEnableNativeIFCTranslator,
 	TEXT("\
 Enable/disable UE native IFC translator. If native translator is disabled, TechSoft is used.\n\
+Default is disable\n"),
+	ECVF_Default);
+
+static bool bGEnableUnsupportedCADFormats = false;
+FAutoConsoleVariableRef GEnableUnsupportedCADFormats(
+	TEXT("ds.CADTranslator.EnableUnsupportedCADFormats"),
+	bGEnableUnsupportedCADFormats,
+	TEXT("\
+Enable/disable unsupported CAD formats i.e. formats that can be imported with the CAD library but haven't been selected as officialy supported.\n\
+These formats stay \"experimental\" as they are out of our testing coverage. No support will be done on these formats.\n\
 Default is disable\n"),
 	ECVF_Default);
 }
@@ -93,10 +103,15 @@ void FDatasmithCADTranslator::Initialize(FDatasmithTranslatorCapabilities& OutCa
 
 	OutCapabilities.SupportedFileFormats.Add(FFileFormatInfo{ TEXT("step"), TEXT("Step files") });
 	OutCapabilities.SupportedFileFormats.Add(FFileFormatInfo{ TEXT("stp"), TEXT("Step files") });
+	OutCapabilities.SupportedFileFormats.Add(FFileFormatInfo{ TEXT("stpz"), TEXT("Step files") });
+	OutCapabilities.SupportedFileFormats.Add(FFileFormatInfo{ TEXT("stpx"), TEXT("Step/XML files") });
+	OutCapabilities.SupportedFileFormats.Add(FFileFormatInfo{ TEXT("stpxz"), TEXT("Step/XML files") });
 	OutCapabilities.SupportedFileFormats.Add(FFileFormatInfo{ TEXT("xml"), TEXT("AP242 Xml Step files, XPDM files") });
 
 	OutCapabilities.SupportedFileFormats.Add(FFileFormatInfo{ TEXT("x_t"), TEXT("Parasolid files (Text format)") });
 	OutCapabilities.SupportedFileFormats.Add(FFileFormatInfo{ TEXT("x_b"), TEXT("Parasolid files (Binary format)") });
+	OutCapabilities.SupportedFileFormats.Add(FFileFormatInfo{ TEXT("xmt"), TEXT("Parasolid files") });
+	OutCapabilities.SupportedFileFormats.Add(FFileFormatInfo{ TEXT("xmt_txt"), TEXT("Parasolid files") });
 
 	OutCapabilities.SupportedFileFormats.Add(FFileFormatInfo{ TEXT("asm"), TEXT("Unigraphics, NX, SolidEdge Assembly files") });
 	OutCapabilities.SupportedFileFormats.Add(FFileFormatInfo{ TEXT("prt"), TEXT("Unigraphics, NX Part files") });
@@ -104,15 +119,37 @@ void FDatasmithCADTranslator::Initialize(FDatasmithTranslatorCapabilities& OutCa
 	OutCapabilities.SupportedFileFormats.Add(FFileFormatInfo{ TEXT("par"), TEXT("SolidEdge Part files") });
 	OutCapabilities.SupportedFileFormats.Add(FFileFormatInfo{ TEXT("psm"), TEXT("SolidEdge Part files") });
 
-	OutCapabilities.SupportedFileFormats.Add(FFileFormatInfo{ TEXT("dwg"), TEXT("AutoCAD, Model files") });
+	OutCapabilities.SupportedFileFormats.Add(FFileFormatInfo{ TEXT("dwg"), TEXT("AutoCAD model files") });
+	OutCapabilities.SupportedFileFormats.Add(FFileFormatInfo{ TEXT("dxf"), TEXT("AutoCAD model files") });
 
 	if (!DatasmithCADTranslatorImpl::bGEnableNativeIFCTranslator)
 	{
 		OutCapabilities.SupportedFileFormats.Add(FFileFormatInfo{ TEXT("ifc"), TEXT("IFC (Industry Foundation Classes)") });
+		OutCapabilities.SupportedFileFormats.Add(FFileFormatInfo{ TEXT("ifczip"), TEXT("IFC (Industry Foundation Classes)") });
 	}
-
+	
 	OutCapabilities.SupportedFileFormats.Add(FFileFormatInfo{ TEXT("hsf"), TEXT("HOOPS stream files") });
 	OutCapabilities.SupportedFileFormats.Add(FFileFormatInfo{ TEXT("prc"), TEXT("HOOPS stream files") });
+
+	if (DatasmithCADTranslatorImpl::bGEnableUnsupportedCADFormats)
+	{
+		OutCapabilities.SupportedFileFormats.Add(FFileFormatInfo{ TEXT("3mf"), TEXT("3D Manufacturing Format") });
+		OutCapabilities.SupportedFileFormats.Add(FFileFormatInfo{ TEXT("3ds"), TEXT("Autodesk 3DS") });
+		OutCapabilities.SupportedFileFormats.Add(FFileFormatInfo{ TEXT("dwf"), TEXT("Autodesk DWF") });
+		OutCapabilities.SupportedFileFormats.Add(FFileFormatInfo{ TEXT("dwfx"), TEXT("Autodesk DWF") });
+		OutCapabilities.SupportedFileFormats.Add(FFileFormatInfo{ TEXT("nwd"), TEXT("Autodesk Navisworks") });
+		OutCapabilities.SupportedFileFormats.Add(FFileFormatInfo{ TEXT("mf1"), TEXT("I-Deas") });
+		OutCapabilities.SupportedFileFormats.Add(FFileFormatInfo{ TEXT("arc"), TEXT("I-Deas") });
+		OutCapabilities.SupportedFileFormats.Add(FFileFormatInfo{ TEXT("unv"), TEXT("I-Deas") });
+		OutCapabilities.SupportedFileFormats.Add(FFileFormatInfo{ TEXT("pkg"), TEXT("I-Deas") });
+		//OutCapabilities.SupportedFileFormats.Add(FFileFormatInfo{ TEXT("dgn"), TEXT("Microstation") });  // available with Hoops Exchange 2023
+		OutCapabilities.SupportedFileFormats.Add(FFileFormatInfo{ TEXT("stl"), TEXT("Stereo Lithography (STL)") });
+		OutCapabilities.SupportedFileFormats.Add(FFileFormatInfo{ TEXT("u3d"), TEXT("U3D (ECMA-363)") });
+		OutCapabilities.SupportedFileFormats.Add(FFileFormatInfo{ TEXT("vda"), TEXT("VDA-FS") });
+		OutCapabilities.SupportedFileFormats.Add(FFileFormatInfo{ TEXT("vrml"), TEXT("VRML") });
+		OutCapabilities.SupportedFileFormats.Add(FFileFormatInfo{ TEXT("wrl"), TEXT("VRML") });
+		OutCapabilities.SupportedFileFormats.Add(FFileFormatInfo{ TEXT("obj"), TEXT("Wavefront OBJ") });
+	}
 }
 
 bool FDatasmithCADTranslator::IsSourceSupported(const FDatasmithSceneSource& Source)
