[NO NON-PII COMMIT MESSAGE]

diff --git a/Engine/Build/InstalledEngineBuild.xml b/Engine/Build/InstalledEngineBuild.xml
--- a/Engine/Build/InstalledEngineBuild.xml
+++ b/Engine/Build/InstalledEngineBuild.xml
@@ -891,7 +891,7 @@
 	<Agent Name="Extra Plugins Agent Win64" Type="Win64">
 		<Node Name="Compile Datasmith Plugins Win64" Requires="Update Version Files" Produces="#DatasmithPluginBinaries Win64">
 			<Do If="$(WithWin64) And $(CompileDatasmithPlugins) == 'True'">
-				<Property Name="DatasmithPlugins" Value="DatasmithCADImporter;DatasmithIFCImporter;DatasmithC4DImporter;AxFImporter;MDLImporter"/>
+				<Property Name="DatasmithPlugins" Value="DatasmithCADImporter;DatasmithC4DImporter;AxFImporter;MDLImporter"/>
 				<ForEach Name="Plugin" Values="$(DatasmithPlugins)">
 					<Compile Target="UnrealEditor" Configuration="Development" Platform="Win64" Arguments="-Plugin=&quot;$(RootDir)/Engine/Plugins/Enterprise/$(Plugin)/$(Plugin).uplugin&quot;"/>
 					<Tag Files="$(RootDir)/Engine/Plugins/Enterprise/$(Plugin)/Binaries/Win64/..." With="#DatasmithPluginBinaries Win64"/>
@@ -903,7 +903,7 @@
 	<Agent Name="Extra Plugins Agent Mac" Type="Mac">
 		<Node Name="Compile Datasmith Plugins Mac" Requires="Update Version Files" Produces="#DatasmithPluginBinaries Mac">
 			<Do If="$(WithMac) And $(CompileDatasmithPlugins) == 'True'">
-				<Property Name="DatasmithPlugins" Value="DatasmithCADImporter;DatasmithIFCImporter;DatasmithC4DImporter;AxFImporter;MDLImporter"/>
+				<Property Name="DatasmithPlugins" Value="DatasmithCADImporter;DatasmithC4DImporter;AxFImporter;MDLImporter"/>
 				<ForEach Name="Plugin" Values="$(DatasmithPlugins)">
 					<Compile Target="UnrealEditor" Configuration="Development" Platform="Mac" Arguments="-Plugin=&quot;$(RootDir)/Engine/Plugins/Enterprise/$(Plugin)/$(Plugin).uplugin&quot; $(MacArchitectureArg)"/>
 					<Tag Files="$(RootDir)/Engine/Plugins/Enterprise/$(Plugin)/Binaries/Mac/..." With="#DatasmithPluginBinaries Mac"/>
diff --git a/Engine/Build/InstalledEngineFilters.xml b/Engine/Build/InstalledEngineFilters.xml
--- a/Engine/Build/InstalledEngineFilters.xml
+++ b/Engine/Build/InstalledEngineFilters.xml
@@ -27,7 +27,6 @@
 
 	<Property Name="PluginsExceptions">
 		Engine/Plugins/Enterprise/DatasmithCADImporter/...
-		Engine/Plugins/Enterprise/DatasmithIFCImporter/...
 		Engine/Plugins/Enterprise/DatasmithC4DImporter/...
 		Engine/Plugins/Enterprise/AxFImporter/...
 		Engine/Plugins/Enterprise/MDLImporter/...
diff --git a/Engine/Plugins/Enterprise/DatasmithCADImporter/Source/DatasmithCADTranslator/Private/DatasmithCADTranslator.cpp b/Engine/Plugins/Enterprise/DatasmithCADImporter/Source/DatasmithCADTranslator/Private/DatasmithCADTranslator.cpp
--- a/Engine/Plugins/Enterprise/DatasmithCADImporter/Source/DatasmithCADTranslator/Private/DatasmithCADTranslator.cpp
+++ b/Engine/Plugins/Enterprise/DatasmithCADImporter/Source/DatasmithCADTranslator/Private/DatasmithCADTranslator.cpp
@@ -23,14 +23,6 @@ DEFINE_LOG_CATEGORY(LogCADTranslator);
 
 namespace DatasmithCADTranslatorImpl
 {
-static bool bGEnableNativeIFCTranslator = false;
-FAutoConsoleVariableRef GCADTranslatorEnableNativeIFCTranslator(
-	TEXT("ds.IFC.EnableNativeTranslator"),
-	bGEnableNativeIFCTranslator,
-	TEXT("\
-Enable/disable UE native IFC translator. If native translator is disabled, TechSoft is used.\n\
-Default is disable\n"),
-	ECVF_Default);
 
 static bool bGEnableUnsupportedCADFormats = false;
 FAutoConsoleVariableRef GEnableUnsupportedCADFormats(
@@ -122,11 +114,8 @@ void FDatasmithCADTranslator::Initialize(FDatasmithTranslatorCapabilities& OutCa
 	OutCapabilities.SupportedFileFormats.Add(FFileFormatInfo{ TEXT("dwg"), TEXT("AutoCAD model files") });
 	OutCapabilities.SupportedFileFormats.Add(FFileFormatInfo{ TEXT("dxf"), TEXT("AutoCAD model files") });
 
-	if (!DatasmithCADTranslatorImpl::bGEnableNativeIFCTranslator)
-	{
-		OutCapabilities.SupportedFileFormats.Add(FFileFormatInfo{ TEXT("ifc"), TEXT("IFC (Industry Foundation Classes)") });
-		OutCapabilities.SupportedFileFormats.Add(FFileFormatInfo{ TEXT("ifczip"), TEXT("IFC (Industry Foundation Classes)") });
-	}
+	OutCapabilities.SupportedFileFormats.Add(FFileFormatInfo{ TEXT("ifc"), TEXT("IFC (Industry Foundation Classes)") });
+	OutCapabilities.SupportedFileFormats.Add(FFileFormatInfo{ TEXT("ifczip"), TEXT("IFC (Industry Foundation Classes)") });
 	
 	OutCapabilities.SupportedFileFormats.Add(FFileFormatInfo{ TEXT("hsf"), TEXT("HOOPS stream files") });
 	OutCapabilities.SupportedFileFormats.Add(FFileFormatInfo{ TEXT("prc"), TEXT("HOOPS stream files") });
diff --git a/Engine/Plugins/Enterprise/DatasmithContent/Source/DatasmithContent/Public/DatasmithAssetImportData.h b/Engine/Plugins/Enterprise/DatasmithContent/Source/DatasmithContent/Public/DatasmithAssetImportData.h
--- a/Engine/Plugins/Enterprise/DatasmithContent/Source/DatasmithContent/Public/DatasmithAssetImportData.h
+++ b/Engine/Plugins/Enterprise/DatasmithContent/Source/DatasmithContent/Public/DatasmithAssetImportData.h
@@ -321,22 +321,6 @@ public:
 	FString ClipInfoPath;
 };
 
-UCLASS(EditInlineNew)
-class DATASMITHCONTENT_API UDatasmithIFCSceneImportData : public UDatasmithSceneImportData
-{
-	GENERATED_BODY()
-};
-
-UCLASS(EditInlineNew)
-class DATASMITHCONTENT_API UDatasmithStaticMeshIFCImportData : public UDatasmithStaticMeshImportData
-{
-	GENERATED_BODY()
-
-public:
-	UPROPERTY(VisibleAnywhere, Category = InternalProperty)
-	FString SourceGlobalId;
-};
-
 namespace Datasmith
 {
 	DATASMITHCONTENT_API UAssetImportData* GetAssetImportData(UObject* Asset);
diff --git a/Engine/Plugins/Enterprise/DatasmithIFCImporter/DatasmithIFCImporter.uplugin b/Engine/Plugins/Enterprise/DatasmithIFCImporter/DatasmithIFCImporter.uplugin
deleted file mode 100644
diff --git a/Engine/Plugins/Enterprise/DatasmithIFCImporter/Resources/IFC/IFC2X3-Settings.xml b/Engine/Plugins/Enterprise/DatasmithIFCImporter/Resources/IFC/IFC2X3-Settings.xml
deleted file mode 100644
diff --git a/Engine/Plugins/Enterprise/DatasmithIFCImporter/Resources/IFC/IFC2X3_TC1.exp b/Engine/Plugins/Enterprise/DatasmithIFCImporter/Resources/IFC/IFC2X3_TC1.exp
deleted file mode 100644
diff --git a/Engine/Plugins/Enterprise/DatasmithIFCImporter/Resources/IFC/IFC4-Settings.xml b/Engine/Plugins/Enterprise/DatasmithIFCImporter/Resources/IFC/IFC4-Settings.xml
deleted file mode 100644
diff --git a/Engine/Plugins/Enterprise/DatasmithIFCImporter/Resources/IFC/IFC4.exp b/Engine/Plugins/Enterprise/DatasmithIFCImporter/Resources/IFC/IFC4.exp
deleted file mode 100644
diff --git a/Engine/Plugins/Enterprise/DatasmithIFCImporter/Resources/IFC/IFC4_ADD2.exp b/Engine/Plugins/Enterprise/DatasmithIFCImporter/Resources/IFC/IFC4_ADD2.exp
deleted file mode 100644
diff --git a/Engine/Plugins/Enterprise/DatasmithIFCImporter/Resources/IFC/IFC4x1_FINAL.exp b/Engine/Plugins/Enterprise/DatasmithIFCImporter/Resources/IFC/IFC4x1_FINAL.exp
deleted file mode 100644
diff --git a/Engine/Plugins/Enterprise/DatasmithIFCImporter/Source/DatasmithIFCTranslator/DatasmithIFCTranslator.Build.cs b/Engine/Plugins/Enterprise/DatasmithIFCImporter/Source/DatasmithIFCTranslator/DatasmithIFCTranslator.Build.cs
deleted file mode 100644
diff --git a/Engine/Plugins/Enterprise/DatasmithIFCImporter/Source/DatasmithIFCTranslator/Private/DatasmithIFCImportOptions.cpp b/Engine/Plugins/Enterprise/DatasmithIFCImporter/Source/DatasmithIFCTranslator/Private/DatasmithIFCImportOptions.cpp
deleted file mode 100644
diff --git a/Engine/Plugins/Enterprise/DatasmithIFCImporter/Source/DatasmithIFCTranslator/Private/DatasmithIFCImportOptions.h b/Engine/Plugins/Enterprise/DatasmithIFCImporter/Source/DatasmithIFCTranslator/Private/DatasmithIFCImportOptions.h
deleted file mode 100644
diff --git a/Engine/Plugins/Enterprise/DatasmithIFCImporter/Source/DatasmithIFCTranslator/Private/DatasmithIFCImporter.cpp b/Engine/Plugins/Enterprise/DatasmithIFCImporter/Source/DatasmithIFCTranslator/Private/DatasmithIFCImporter.cpp
deleted file mode 100644
diff --git a/Engine/Plugins/Enterprise/DatasmithIFCImporter/Source/DatasmithIFCTranslator/Private/DatasmithIFCImporter.h b/Engine/Plugins/Enterprise/DatasmithIFCImporter/Source/DatasmithIFCTranslator/Private/DatasmithIFCImporter.h
deleted file mode 100644
diff --git a/Engine/Plugins/Enterprise/DatasmithIFCImporter/Source/DatasmithIFCTranslator/Private/DatasmithIFCTranslator.cpp b/Engine/Plugins/Enterprise/DatasmithIFCImporter/Source/DatasmithIFCTranslator/Private/DatasmithIFCTranslator.cpp
deleted file mode 100644
diff --git a/Engine/Plugins/Enterprise/DatasmithIFCImporter/Source/DatasmithIFCTranslator/Private/DatasmithIFCTranslator.h b/Engine/Plugins/Enterprise/DatasmithIFCImporter/Source/DatasmithIFCTranslator/Private/DatasmithIFCTranslator.h
deleted file mode 100644
diff --git a/Engine/Plugins/Enterprise/DatasmithIFCImporter/Source/DatasmithIFCTranslator/Private/DatasmithIFCTranslatorModule.cpp b/Engine/Plugins/Enterprise/DatasmithIFCImporter/Source/DatasmithIFCTranslator/Private/DatasmithIFCTranslatorModule.cpp
deleted file mode 100644
diff --git a/Engine/Plugins/Enterprise/DatasmithIFCImporter/Source/DatasmithIFCTranslator/Private/IFC/IFCReader.cpp b/Engine/Plugins/Enterprise/DatasmithIFCImporter/Source/DatasmithIFCTranslator/Private/IFC/IFCReader.cpp
deleted file mode 100644
diff --git a/Engine/Plugins/Enterprise/DatasmithIFCImporter/Source/DatasmithIFCTranslator/Private/IFC/IFCReader.h b/Engine/Plugins/Enterprise/DatasmithIFCImporter/Source/DatasmithIFCTranslator/Private/IFC/IFCReader.h
deleted file mode 100644
diff --git a/Engine/Plugins/Enterprise/DatasmithIFCImporter/Source/DatasmithIFCTranslator/Private/IFC/IFCStaticMeshFactory.cpp b/Engine/Plugins/Enterprise/DatasmithIFCImporter/Source/DatasmithIFCTranslator/Private/IFC/IFCStaticMeshFactory.cpp
deleted file mode 100644
diff --git a/Engine/Plugins/Enterprise/DatasmithIFCImporter/Source/DatasmithIFCTranslator/Private/IFC/IFCStaticMeshFactory.h b/Engine/Plugins/Enterprise/DatasmithIFCImporter/Source/DatasmithIFCTranslator/Private/IFC/IFCStaticMeshFactory.h
deleted file mode 100644
diff --git a/Engine/Plugins/Enterprise/DatasmithIFCImporter/Source/DatasmithIFCTranslator/Public/DatasmithIFCTranslatorModule.h b/Engine/Plugins/Enterprise/DatasmithIFCImporter/Source/DatasmithIFCTranslator/Public/DatasmithIFCTranslatorModule.h
deleted file mode 100644
diff --git a/Engine/Source/Programs/AutomationTool/BuildGraph/Scripts/InstalledBuild.cs b/Engine/Source/Programs/AutomationTool/BuildGraph/Scripts/InstalledBuild.cs
--- a/Engine/Source/Programs/AutomationTool/BuildGraph/Scripts/InstalledBuild.cs
+++ b/Engine/Source/Programs/AutomationTool/BuildGraph/Scripts/InstalledBuild.cs
@@ -30,7 +30,6 @@ namespace AutomationTool
 		static string[] PluginsExceptions =
 		{
 			"Engine/Plugins/Enterprise/DatasmithCADImporter/...",
-			"Engine/Plugins/Enterprise/DatasmithIFCImporter/...",
 			"Engine/Plugins/Enterprise/DatasmithC4DImporter/...",
 			"Engine/Plugins/Enterprise/AxFImporter/...",
 			"Engine/Plugins/Enterprise/MDLImporter/..."
