Rollback 24762504 for TechSoftInterface

diff --git a/Engine/Plugins/Enterprise/DatasmithCADImporter/Source/CADInterfaces/Private/TechSoftInterface.cpp b/Engine/Plugins/Enterprise/DatasmithCADImporter/Source/CADInterfaces/Private/TechSoftInterface.cpp
--- a/Engine/Plugins/Enterprise/DatasmithCADImporter/Source/CADInterfaces/Private/TechSoftInterface.cpp
+++ b/Engine/Plugins/Enterprise/DatasmithCADImporter/Source/CADInterfaces/Private/TechSoftInterface.cpp
@@ -97,19 +97,7 @@ FString GetTechSoftVersion()
 FUniqueTechSoftModelFile LoadModelFileFromFile(const A3DImport& Import, A3DStatus& Status)
 {
 	A3DAsmModelFile* ModelFile = nullptr;
-#if !PLATFORM_EXCEPTIONS_DISABLED
-	try
-#endif
-	{
-		Status = A3DAsmModelFileLoadFromFile(Import.GetFilePath(), &Import.m_sLoadData, &ModelFile);
-	}
-#if !PLATFORM_EXCEPTIONS_DISABLED
-	catch (...)
-	{
-		UE_LOG(LogCADInterfaces, Warning, TEXT("Failed to load %s. An exception is thrown."), *Import.GetFilePath());
-		return FUniqueTechSoftModelFile();
-	}
-#endif
+	Status = A3DAsmModelFileLoadFromFile(Import.GetFilePath(), &Import.m_sLoadData, &ModelFile);
 
 	switch (Status)
 	{
