Fix Wire Matrix issue

diff --git a/Engine/Plugins/Enterprise/DatasmithCADImporter/Source/DatasmithWireTranslator/Private/DatasmithWireTranslator.cpp b/Engine/Plugins/Enterprise/DatasmithCADImporter/Source/DatasmithWireTranslator/Private/DatasmithWireTranslator.cpp
--- a/Engine/Plugins/Enterprise/DatasmithCADImporter/Source/DatasmithWireTranslator/Private/DatasmithWireTranslator.cpp
+++ b/Engine/Plugins/Enterprise/DatasmithCADImporter/Source/DatasmithWireTranslator/Private/DatasmithWireTranslator.cpp
@@ -209,7 +209,7 @@ public:
 		{
 			if (!ShaderNames.Contains(ShaderName))
 			{
-			UE_LOG(LogDatasmithWireTranslator, Warning, TEXT("The main UE5 rendering systems do not support more than 256 materials per mesh and the limit has been defined to %d materials. Only the first %d materials are kept. The others are replaced by the last one"), FImportParameters::GMaxMaterialCountPerMesh, FImportParameters::GMaxMaterialCountPerMesh);
+				UE_LOG(LogDatasmithWireTranslator, Warning, TEXT("The main UE5 rendering systems do not support more than 256 materials per mesh and the limit has been defined to %d materials. Only the first %d materials are kept. The others are replaced by the last one"), FImportParameters::GMaxMaterialCountPerMesh, FImportParameters::GMaxMaterialCountPerMesh);
 				bWarningNotLogYet = false;
 				Color = LastColor;
 			}
@@ -1545,9 +1545,6 @@ bool FWireTranslatorImpl::ProcessAlGroupNode(AlDagNode& GroupNode, FDagNodeInfo&
 	ThisGroupNodeInfo.ActorElement->SetLabel(*ThisGroupNodeInfo.Label);
 	ThisGroupNodeInfo.ExtractDagNodeLayer(GroupNode);
 
-	// Apply local transform to actor element
-	OpenModelUtils::SetActorTransform(ThisGroupNodeInfo.ActorElement, GroupNode);
-
 	TSharedPtr<IDatasmithActorElement> ParentActorElement = FindOrAddParentActor(ParentInfo, ThisGroupNodeInfo.LayerName);
 	if (ParentActorElement.IsValid())
 	{
@@ -1560,6 +1557,9 @@ bool FWireTranslatorImpl::ProcessAlGroupNode(AlDagNode& GroupNode, FDagNodeInfo&
 
 	RecurseDagForLeaves(ChildNode, ThisGroupNodeInfo);
 
+	// Apply local transform to actor element
+	OpenModelUtils::SetActorTransform(ThisGroupNodeInfo.ActorElement, GroupNode);
+
 	return true;
 }
 
@@ -1675,8 +1675,6 @@ void FWireTranslatorImpl::ProcessAlShellNode(const TSharedPtr<AlDagNode>& ShellN
 		return;
 	}
 
-	OpenModelUtils::SetActorTransform(ShellInfo.ActorElement, *ShellNode);
-
 	// Apply materials on the current part
 	if (ShaderName.Len())
 	{
@@ -1700,6 +1698,8 @@ void FWireTranslatorImpl::ProcessAlShellNode(const TSharedPtr<AlDagNode>& ShellN
 	{
 		DatasmithScene->AddActor(ActorElement);
 	}
+
+	OpenModelUtils::SetActorTransform(ShellInfo.ActorElement, *ShellNode);
 }
 
 void FWireTranslatorImpl::ProcessBodyNode(const TSharedPtr<BodyData>& Body, FDagNodeInfo& ParentInfo)
@@ -2073,7 +2073,7 @@ TOptional<FMeshDescription> FWireTranslatorImpl::MeshDagNodeWithExternalMesher(A
 	{
 		const TCHAR* StaticMeshLable = MeshElement->GetLabel();
 		const TCHAR* StaticMeshName = MeshElement->GetName();
-		UE_LOG(LogDatasmithWireTranslator, Warning, TEXT("Failed to generate the mesh of \"%s\" (%s) StaticMesh."), *StaticMeshLable, *StaticMeshName);
+		UE_LOG(LogDatasmithWireTranslator, Warning, TEXT("Failed to generate the mesh of \"%s\" (%s) StaticMesh."), StaticMeshLable, StaticMeshName);
 	}
 
 	return MoveTemp(MeshDescription);
