CAD Qualifier improvements

diff --git a/Engine/Plugins/Enterprise/DatasmithCADImporter/Source/CADInterfaces/Private/CADFileData.cpp b/Engine/Plugins/Enterprise/DatasmithCADImporter/Source/CADInterfaces/Private/CADFileData.cpp
new file mode 100644
--- /dev/null
+++ b/Engine/Plugins/Enterprise/DatasmithCADImporter/Source/CADInterfaces/Private/CADFileData.cpp
@@ -0,0 +1,108 @@
+// Copyright Epic Games, Inc. All Rights Reserved.
+#include "CADFileData.h"
+
+#include "Misc/FileHelper.h"
+
+namespace CADLibrary
+{
+
+void FCADFileData::InitCsvFile()
+{
+	FString CsvData;
+	CsvData.Reserve(FCString::Strlen(CSV_HEADER) * 2);
+	CsvData += CSV_HEADER;
+
+	const FString& RootPath = FileDescription.GetRootFolder();
+	const FString& FilePath = FileDescription.GetSourcePath();
+
+	//Input File
+	FString FileName = FilePath.Right(FilePath.Len() - RootPath.Len() - 1);
+	FileName.ReplaceCharInline(TEXT(','), TEXT('_'));
+	CsvData += FileName + TEXT(",");
+	//Input File
+	CsvData += SceneGraphArchive.ArchiveFileName + TEXT(",");
+
+	//FileSize
+	FFileStatData FileStatData = IFileManager::Get().GetStatData(*FileDescription.GetSourcePath());
+	constexpr int64 OneKiloBit = 1024;
+	CsvData += FString::Printf(TEXT("%d,"), FileStatData.FileSize / OneKiloBit);
+	CsvData += TEXT("\0\0\0\0\0\0\0\0\0\0\0\0\0\0\0");
+	FString CsvPath = GetValidationFilePath();
+	FFileHelper::SaveStringToFile(*CsvData, *CsvPath, FFileHelper::EEncodingOptions::ForceUTF8);
+}
+
+void FCADFileData::ExportValidationData()
+{
+	FString CsvData;
+	CsvData.Reserve(FCString::Strlen(CSV_HEADER));
+
+	//Time.LoadStep
+	CsvData += FString::Printf(TEXT("%f,"), Record.ImportTime);
+	//Time.SavePrcTime
+	CsvData += FString::Printf(TEXT("%f,"), Record.SavePrcTime);
+	//Time.AdaptBRepTime
+	CsvData += FString::Printf(TEXT("%f,"), Record.AdaptBRepTime);
+	//Time.SewStep
+	CsvData += FString::Printf(TEXT("%f,"), Record.SewTime);
+	//Time.MeshStep
+	CsvData += FString::Printf(TEXT("%f,"), Record.MeshTime);
+	//Time.Total
+	CsvData += FString::Printf(TEXT("%f,"), Record.LoadProcessTime);
+
+	//MemoryUsage
+	constexpr float OneMegaBit = 1024.f * 1024.f;
+	uint64 MaxUsedMemory = Record.MaxMemoryUsed - Record.StartMemoryUsed;
+	float MaxUsedMemoryF = static_cast<float>(MaxUsedMemory) / OneMegaBit;
+	CsvData += FString::Printf(TEXT("%f,"), MaxUsedMemoryF);
+
+	//SceneGrapheData.InstanceNum
+	CsvData.AppendInt(SceneGraphArchive.Instances.Num());
+	CsvData += TEXT(",");
+	//SceneGrapheData.Bodies
+	CsvData.AppendInt(SceneGraphArchive.Bodies.Num());
+	CsvData += TEXT(",");
+	//SceneGrapheData.References
+	CsvData.AppendInt(SceneGraphArchive.References.Num());
+	CsvData += TEXT(",");
+	//SceneGrapheData.UnloadedReferences
+	CsvData.AppendInt(SceneGraphArchive.UnloadedReferences.Num());
+	CsvData += TEXT(",");
+	//SceneGrapheData.ExternalReferenceFiles
+	CsvData.AppendInt(SceneGraphArchive.ExternalReferenceFiles.Num());
+	CsvData += TEXT(",");
+
+	//SceneGrapheData.ColorHIdToColor
+	CsvData.AppendInt(SceneGraphArchive.ColorHIdToColor.Num());
+	CsvData += TEXT(",");
+	//SceneGrapheData.MaterialHIdToMaterial
+	CsvData.AppendInt(SceneGraphArchive.MaterialHIdToMaterial.Num());
+	CsvData += TEXT(",");
+
+	int32 FaceCount = 0;
+	int32 VertexCount = 0;
+	GetMeshStatsToCsv(FaceCount, VertexCount);
+
+	//Meshes.FaceCount
+	CsvData.AppendInt(FaceCount);
+	CsvData += TEXT(",");
+	//Meshes.VertexCount
+	CsvData.AppendInt(VertexCount);
+	CsvData += TEXT(",");
+	CsvData += TEXT('\0');
+
+	FString CsvPath = GetValidationFilePath();
+	FFileHelper::SaveStringToFile(*CsvData, *CsvPath, FFileHelper::EEncodingOptions::ForceUTF8, &IFileManager::Get(), EFileWrite::FILEWRITE_Append);
+}
+
+void FCADFileData::GetMeshStatsToCsv(int32& FaceCount, int32& VertexCount)
+{
+	
+	for (const FBodyMesh& BodyMesh : BodyMeshes)
+	{
+		FaceCount += BodyMesh.Faces.Num();
+		VertexCount += BodyMesh.VertexArray.Num();
+	}
+}
+
+
+}
\ No newline at end of file
diff --git a/Engine/Plugins/Enterprise/DatasmithCADImporter/Source/CADInterfaces/Private/CADFileReader.cpp b/Engine/Plugins/Enterprise/DatasmithCADImporter/Source/CADInterfaces/Private/CADFileReader.cpp
--- a/Engine/Plugins/Enterprise/DatasmithCADImporter/Source/CADInterfaces/Private/CADFileReader.cpp
+++ b/Engine/Plugins/Enterprise/DatasmithCADImporter/Source/CADInterfaces/Private/CADFileReader.cpp
@@ -98,6 +98,8 @@ namespace CADLibrary
 			return ECADParsingResult::ProcessFailed;
 		}
 
+		CADFileData.SetArchiveNames();
+
 		if (!FindFile(CADFileData.GetCADFileDescription()))
 		{
 			return ECADParsingResult::FileNotFound;
@@ -105,8 +107,6 @@ namespace CADLibrary
 
 		if (FImportParameters::bGEnableCADCache)
 		{
-			CADFileData.SetArchiveNames();
-
 			bool bNeedToProceed = true;
 
 			FString CADFileCachePath = CADFileData.GetCADCachePath();
@@ -135,6 +135,11 @@ namespace CADLibrary
 			}
 		}
 
+		if (FImportParameters::bValidationProcess)
+		{
+			CADFileData.InitCsvFile();
+		}
+
 		// Process the file
 		ECADParsingResult Result = CADParser->Process();
 
@@ -147,6 +152,11 @@ namespace CADLibrary
 			}
 		}
 
+		if (FImportParameters::bValidationProcess)
+		{
+			CADFileData.ExportValidationData();
+		}
+
 		return Result;
 	}
 }
diff --git a/Engine/Plugins/Enterprise/DatasmithCADImporter/Source/CADInterfaces/Private/CADSceneGraph.cpp b/Engine/Plugins/Enterprise/DatasmithCADImporter/Source/CADInterfaces/Private/CADSceneGraph.cpp
--- a/Engine/Plugins/Enterprise/DatasmithCADImporter/Source/CADInterfaces/Private/CADSceneGraph.cpp
+++ b/Engine/Plugins/Enterprise/DatasmithCADImporter/Source/CADInterfaces/Private/CADSceneGraph.cpp
@@ -106,7 +106,7 @@ FArchive& operator<<(FArchive& Ar, FArchiveSceneGraph& SceneGraph)
 void FArchiveSceneGraph::SerializeMockUp(const TCHAR* Filename)
 {
 	TUniquePtr<FArchive> Archive(IFileManager::Get().CreateFileWriter(Filename));
-	if (Archive)
+	if (Archive.IsValid())
 	{
 		*Archive << *this;
 		Archive->Close();
@@ -386,6 +386,11 @@ void FArchiveSceneGraph::RemoveLastBody()
 
 void FArchiveSceneGraph::AddExternalReferenceFile(const FArchiveUnloadedReference& Reference)
 {
+	if (Reference.ExternalFile.GetFileName().IsEmpty())
+	{
+		return;
+	}
+
 	const int32 Index = CADIdToIndex[Reference.Id];
 	ensure(ExternalReferenceFiles.Add(Reference.ExternalFile) == Index);
 }
diff --git a/Engine/Plugins/Enterprise/DatasmithCADImporter/Source/CADInterfaces/Private/TechSoftFileParser.cpp b/Engine/Plugins/Enterprise/DatasmithCADImporter/Source/CADInterfaces/Private/TechSoftFileParser.cpp
--- a/Engine/Plugins/Enterprise/DatasmithCADImporter/Source/CADInterfaces/Private/TechSoftFileParser.cpp
+++ b/Engine/Plugins/Enterprise/DatasmithCADImporter/Source/CADInterfaces/Private/TechSoftFileParser.cpp
@@ -329,7 +329,18 @@ FTechSoftFileParser::FTechSoftFileParser(FCADFileData& InCADData, const FString&
 #ifdef USE_TECHSOFT_SDK
 ECADParsingResult FTechSoftFileParser::Process()
 {
+	bProcessIsRunning = true;
+	TArray<UE::Tasks::FTask> Checkers;
+	if(FImportParameters::bValidationProcess)
+	{
+		Checkers.Emplace(UE::Tasks::Launch(TEXT("MemoryChecker"), [&]() { CheckMemory(); }));
+	}
+
+	//FPlatformMemoryStats StartMem = FPlatformMemory::GetStats();
+	uint64 StartTime = FPlatformTime::Cycles64();
+
 	const FFileDescriptor& File = CADFileData.GetCADFileDescription();
+	FImportRecord& ProcessReport = CADFileData.GetRecord();
 
 	if (File.GetPathOfFileToLoad().IsEmpty())
 	{
@@ -346,6 +357,7 @@ ECADParsingResult FTechSoftFileParser::Process()
 
 	A3DStatus LoadStatus = A3DStatus::A3D_SUCCESS;
 	ModelFile = TechSoftInterface::LoadModelFileFromFile(Import, LoadStatus);
+	ProcessReport.ImportTime += FPlatformTime::ToMilliseconds64(FPlatformTime::Cycles64() - StartTime);
 
 	if (!ModelFile.IsValid())
 	{
@@ -398,31 +410,42 @@ ECADParsingResult FTechSoftFileParser::Process()
 	// save the file for the next load
 	if (CADFileData.IsCacheDefined())
 	{
+		uint64 StartSaveTime = FPlatformTime::Cycles64();
 		FString CacheFilePath = CADFileData.GetCADCachePath();
 		if (CacheFilePath != File.GetPathOfFileToLoad())
 		{
 			TechSoftUtils::SaveModelFileToPrcFile(ModelFile.Get(), CacheFilePath);
 		}
+		ProcessReport.SavePrcTime = FPlatformTime::ToMilliseconds64(FPlatformTime::Cycles64() - StartSaveTime);
 	}
 
 	// Adapt BRep to UE::CADKernel
 	{
+		uint64 StartAdaptTime = FPlatformTime::Cycles64();
 		if (AdaptBRepModel() != A3D_SUCCESS)
 		{
 			return ECADParsingResult::ProcessFailed;
 		}
+		ProcessReport.AdaptBRepTime = FPlatformTime::ToMilliseconds64(FPlatformTime::Cycles64() - StartAdaptTime);
 	}
 
 	// Some formats (like IGES) require a sew all the time. In this case, bForceSew = true
 	if (bForceSew || CADFileData.GetImportParameters().GetStitchingTechnique() == StitchingSew)
 	{
+		uint64 StartSewTime = FPlatformTime::Cycles64();
 		SewModel();
+		ProcessReport.SewTime = FPlatformTime::ToMilliseconds64(FPlatformTime::Cycles64() - StartSewTime);
 	}
 
+	uint64 StartTraversTime = FPlatformTime::Cycles64();
+
 	ReserveCADFileData();
 	ReadMaterialsAndColors();
 	ECADParsingResult Result = TraverseModel();
 
+	ProcessReport.ImportTime += FPlatformTime::ToMilliseconds64(FPlatformTime::Cycles64() - StartTraversTime);
+
+
 	if (Result == ECADParsingResult::ProcessOk)
 	{
 		GenerateBodyMeshes();
@@ -437,6 +460,11 @@ ECADParsingResult FTechSoftFileParser::Process()
 
 	ModelFile.Reset();
 
+	ProcessReport.LoadProcessTime = FPlatformTime::ToMilliseconds64(FPlatformTime::Cycles64() - StartTime);
+
+	bProcessIsRunning = false;
+	UE::Tasks::Wait(Checkers);
+
 	return Result;
 }
 
@@ -451,12 +479,14 @@ void FTechSoftFileParser::SewModel()
 
 void FTechSoftFileParser::GenerateBodyMeshes()
 {
+	uint64 StartGenerateBodyMeshesTime = FPlatformTime::Cycles64();
 	for (TPair<A3DRiRepresentationItem*, FCadId>& Entry : RepresentationItemsCache)
 	{
 		A3DRiRepresentationItem* RepresentationItemPtr = Entry.Key;
 		FArchiveBody& Body = SceneGraph.GetBody(Entry.Value);
 		GenerateBodyMesh(RepresentationItemPtr, Body);
 	}
+	CADFileData.GetRecord().MeshTime = FPlatformTime::ToMilliseconds64(FPlatformTime::Cycles64() - StartGenerateBodyMeshesTime);
 }
 
 void FTechSoftFileParser::GenerateBodyMesh(A3DRiRepresentationItem* Representation, FArchiveBody& Body)
@@ -1971,6 +2001,22 @@ void FTechSoftFileParser::ReadMaterialsAndColors()
 	}
 }
 
+void FTechSoftFileParser::CheckMemory()
+{
+	CADFileData.GetRecord().StartMemoryUsed = FPlatformMemory::GetStats().UsedPhysical;
+	uint64& MaxMemoryUsed = CADFileData.GetRecord().MaxMemoryUsed;
+	while (bProcessIsRunning)
+	{
+		FPlatformProcess::Sleep(0.1);
+		const uint64 MemoryUsed = FPlatformMemory::GetStats().UsedPhysical;
+		if (MaxMemoryUsed < MemoryUsed)
+		{
+			MaxMemoryUsed = MemoryUsed;
+		}
+	}
+}
+
+
 #endif  
 
 } // ns CADLibrary
diff --git a/Engine/Plugins/Enterprise/DatasmithCADImporter/Source/CADInterfaces/Private/TechSoftFileParser.h b/Engine/Plugins/Enterprise/DatasmithCADImporter/Source/CADInterfaces/Private/TechSoftFileParser.h
--- a/Engine/Plugins/Enterprise/DatasmithCADImporter/Source/CADInterfaces/Private/TechSoftFileParser.h
+++ b/Engine/Plugins/Enterprise/DatasmithCADImporter/Source/CADInterfaces/Private/TechSoftFileParser.h
@@ -142,6 +142,9 @@ private:
 	void ExtractGeneralTransformation(const A3DMiscTransformation* GeneralTransformation, FArchiveCADObject& Component);
 	void ExtractTransformation3D(const A3DMiscTransformation* CartesianTransformation, FArchiveCADObject& Component);
 
+	void CheckMemory();
+
+
 #endif
 
 protected:
@@ -163,6 +166,9 @@ protected:
 
 	TMap<A3DRiRepresentationItem*, FCadId> RepresentationItemsCache;
 	TMap<A3DAsmProductOccurrence*, FCadId> ReferenceCache;
+
+	std::atomic<bool> bProcessIsRunning;
+
 };
 
 } // ns CADLibrary
\ No newline at end of file
diff --git a/Engine/Plugins/Enterprise/DatasmithCADImporter/Source/CADInterfaces/Private/TechSoftFileParserCADKernelTessellator.cpp b/Engine/Plugins/Enterprise/DatasmithCADImporter/Source/CADInterfaces/Private/TechSoftFileParserCADKernelTessellator.cpp
--- a/Engine/Plugins/Enterprise/DatasmithCADImporter/Source/CADInterfaces/Private/TechSoftFileParserCADKernelTessellator.cpp
+++ b/Engine/Plugins/Enterprise/DatasmithCADImporter/Source/CADInterfaces/Private/TechSoftFileParserCADKernelTessellator.cpp
@@ -22,6 +22,8 @@
 #include "CADKernel/Topo/TopologicalShapeEntity.h"
 #include "CADKernel/Topo/Topomaker.h"
 
+#include "HAL/PlatformTime.h"
+
 namespace CADLibrary
 {
 
@@ -109,6 +111,7 @@ void FTechSoftFileParserCADKernelTessellator::SewAndMesh(TArray<A3DRiRepresentat
 	}
 
 	// Sew if needed
+	uint64 StartSewTime = FPlatformTime::Cycles64();
 	UE::CADKernel::FTopomakerOptions TopomakerOptions((UE::CADKernel::ESewOption) SewOption::GetFromImportParameters(), GeometricTolerance, FImportParameters::GStitchingForceFactor);
 
 	FTopomaker Topomaker(CADKernelSession, TopomakerOptions);
@@ -169,6 +172,10 @@ void FTechSoftFileParserCADKernelTessellator::SewAndMesh(TArray<A3DRiRepresentat
 		}
 	}
 
+	CADFileData.GetRecord().SewTime += FPlatformTime::ToMilliseconds64(FPlatformTime::Cycles64() - StartSewTime);
+
+	uint64 StartMeshTime = FPlatformTime::Cycles64();
+
 	// Process existing bodies
 	for (FBody* Body : ExistingBodies)
 	{
@@ -203,6 +210,8 @@ void FTechSoftFileParserCADKernelTessellator::SewAndMesh(TArray<A3DRiRepresentat
 			ArchiveBody.Delete();
 		}
 	}
+
+	CADFileData.GetRecord().MeshTime += FPlatformTime::ToMilliseconds64(FPlatformTime::Cycles64() - StartMeshTime);
 }
 
 void FTechSoftFileParserCADKernelTessellator::GenerateBodyMesh(A3DRiRepresentationItem* Representation, FArchiveBody& ArchiveBody)
diff --git a/Engine/Plugins/Enterprise/DatasmithCADImporter/Source/CADInterfaces/Public/CADFileData.h b/Engine/Plugins/Enterprise/DatasmithCADImporter/Source/CADInterfaces/Public/CADFileData.h
--- a/Engine/Plugins/Enterprise/DatasmithCADImporter/Source/CADInterfaces/Public/CADFileData.h
+++ b/Engine/Plugins/Enterprise/DatasmithCADImporter/Source/CADInterfaces/Public/CADFileData.h
@@ -10,6 +10,21 @@
 namespace CADLibrary
 {
 
+#define CSV_HEADER TEXT("Input File,FileHash,FileSize(KB),Time.LoadStep(ms),Time.SavePrcTime(ms),Time.AdaptBRepTime(ms),Time.SewStep(ms),Time.MeshStep(ms),Time.Total(ms),MemoryUsage(MB),SceneGrapheData.InstanceNum,SceneGrapheData.Bodies,SceneGrapheData.References,SceneGrapheData.UnloadedReferences,SceneGrapheData.ExternalReferenceFiles,SceneGrapheData.ColorHIdToColor,SceneGrapheData.MaterialHIdToMaterial,Meshes.FaceCount,Meshes.VertexCount, 0\n")
+
+struct FImportRecord
+{
+	double ImportTime = 0;
+	double SavePrcTime = 0;
+	double SewTime = 0;
+	double AdaptBRepTime = 0;
+	double MeshTime = 0;
+	double LoadProcessTime = 0; // Full process
+
+	uint64 StartMemoryUsed = 0;
+	uint64 MaxMemoryUsed = 0;
+};
+
 class FCADFileData
 {
 public:
@@ -64,6 +79,23 @@ public:
 		return FString();
 	}
 
+	const FString GetValidationFilePath() const
+	{
+		if (IsCacheDefined())
+		{
+			const FString& RootPath = FileDescription.GetRootFolder();
+			const FString& FilePath = FileDescription.GetSourcePath();
+			FString FileRootFolder = FilePath.Right(FilePath.Len() - RootPath.Len() - 1);
+			int32 Index;
+			if (FileRootFolder.FindChar(TEXT('/'), Index))
+			{
+				FileRootFolder = FileRootFolder.Left(Index);
+			}
+			return FPaths::Combine(CachePath, TEXT("validation"), FileRootFolder, SceneGraphArchive.ArchiveFileName + TEXT(".csv"));
+		}
+		return FString();
+	}
+
 	const FString GetMeshArchiveFilePath() const
 	{
 		if (IsCacheDefined())
@@ -261,14 +293,26 @@ public:
 
 	void UpdateExternalRefPath()
 	{
-		const FString Root = FileDescription.GetRootFolder();
-		const FString ArchiveRoot = FPaths::GetPath(SceneGraphArchive.FullPath);
-		for (FFileDescriptor& File : SceneGraphArchive.ExternalReferenceFiles)
+		if(!FImportParameters::bValidationProcess)
 		{
-			File.ChangePath(ArchiveRoot, Root);
+			const FString Root = FileDescription.GetRootFolder();
+			const FString ArchiveRoot = FPaths::GetPath(SceneGraphArchive.FullPath);
+			for (FFileDescriptor& File : SceneGraphArchive.ExternalReferenceFiles)
+			{
+				File.ChangePath(ArchiveRoot, Root);
+			}
 		}
 	}
 
+	void InitCsvFile();
+	void ExportValidationData();
+	void GetMeshStatsToCsv(int32& FaceCount, int32& VertexCount);
+
+	FImportRecord& GetRecord()
+	{
+		return Record;
+	}
+
 private:
 	const FImportParameters ImportParameters;
 	const FString CachePath;
@@ -283,6 +327,8 @@ private:
 
 	TArray<TPair<uint8, FString>> Messages;
 
+	FImportRecord Record;
+
 	mutable uint32 SceneFileHash = 0;
 	mutable uint32 GeomFileHash = 0;
 };
diff --git a/Engine/Plugins/Enterprise/DatasmithCADImporter/Source/DatasmithDispatcher/Private/DatasmithDispatcher.cpp b/Engine/Plugins/Enterprise/DatasmithCADImporter/Source/DatasmithDispatcher/Private/DatasmithDispatcher.cpp
--- a/Engine/Plugins/Enterprise/DatasmithCADImporter/Source/DatasmithDispatcher/Private/DatasmithDispatcher.cpp
+++ b/Engine/Plugins/Enterprise/DatasmithCADImporter/Source/DatasmithDispatcher/Private/DatasmithDispatcher.cpp
@@ -9,6 +9,7 @@
 #include "DatasmithDispatcherTask.h"
 
 #include "Algo/Count.h"
+#include "Algo/Find.h"
 #include "HAL/FileManager.h"
 #include "Misc/MessageDialog.h"
 #include "Misc/Paths.h"
@@ -73,12 +74,10 @@ FDatasmithDispatcher::FDatasmithDispatcher(const CADLibrary::FImportParameters&
 void FDatasmithDispatcher::AddTask(const CADLibrary::FFileDescriptor& InFileDescription)
 {
 	FScopeLock Lock(&TaskPoolCriticalSection);
-	for (const FTask& Task : TaskPool)
+
+	if (Algo::FindByPredicate(TaskPool, [&InFileDescription](const FTask& Task) { return Task.FileDescription == InFileDescription; }))
 	{
-		if (Task.FileDescription == InFileDescription)
-		{
-			return;
-		}
+		return;
 	}
 
 	int32 TaskIndex = TaskPool.Emplace(InFileDescription);
@@ -132,6 +131,8 @@ TOptional<FTask> FDatasmithDispatcher::GetNextTask()
 	}
 
 	TaskPool[NextTaskIndex].State = ETaskState::Running;
+	UE_LOG(LogDatasmithDispatcher, Display, TEXT(" Launch %s (%d / %d)"), *TaskPool[NextTaskIndex].FileDescription.GetSourcePath(), NextTaskIndex, TaskPool.Num());
+
 	return TaskPool[NextTaskIndex++];
 }
 
@@ -187,7 +188,7 @@ void FDatasmithDispatcher::SetTaskState(int32 TaskIndex, ETaskState TaskState)
 void FDatasmithDispatcher::Process(bool bWithProcessor)
 {
 	EstimationOfMemoryUsedByStaticMeshes = 0;
-	bool bCheckMemory = true;
+	bool bCheckMemory = CADLibrary::FImportParameters::bValidationProcess ? false : true;
 
 	// This function is added in 5.1.1 but is cloned with the lines 370 to 387
 	TFunction<bool()> CanImportContinue = [&]() -> bool
diff --git a/Engine/Plugins/Enterprise/DatasmithCADImporter/Source/DatasmithDispatcher/Private/DatasmithWorkerHandler.cpp b/Engine/Plugins/Enterprise/DatasmithCADImporter/Source/DatasmithDispatcher/Private/DatasmithWorkerHandler.cpp
--- a/Engine/Plugins/Enterprise/DatasmithCADImporter/Source/DatasmithDispatcher/Private/DatasmithWorkerHandler.cpp
+++ b/Engine/Plugins/Enterprise/DatasmithCADImporter/Source/DatasmithDispatcher/Private/DatasmithWorkerHandler.cpp
@@ -336,6 +336,10 @@ void FDatasmithWorkerHandler::ProcessCommand(FCompletedTaskCommand& CompletedTas
 
 	for (const CADLibrary::FFileDescriptor& ExternalReferenceFile : CompletedTaskCommand.ExternalReferences)
 	{
+		if (ExternalReferenceFile.GetFileName().IsEmpty())
+		{
+			continue;
+		}
 		Dispatcher.AddTask(ExternalReferenceFile);
 	}
 
diff --git a/Engine/Plugins/Enterprise/DatasmithCADImporter/Source/DatasmithPLMXMLTranslator/Private/DatasmithPlmXmlImporter.cpp b/Engine/Plugins/Enterprise/DatasmithCADImporter/Source/DatasmithPLMXMLTranslator/Private/DatasmithPlmXmlImporter.cpp
--- a/Engine/Plugins/Enterprise/DatasmithCADImporter/Source/DatasmithPLMXMLTranslator/Private/DatasmithPlmXmlImporter.cpp
+++ b/Engine/Plugins/Enterprise/DatasmithCADImporter/Source/DatasmithPLMXMLTranslator/Private/DatasmithPlmXmlImporter.cpp
@@ -277,7 +277,7 @@ namespace PlmXml
 
 		void Process(const FDatasmithSceneSource& Source)
 		{
-			DatasmithDispatcher->Process(true);
+			DatasmithDispatcher->Process(CADLibrary::GMaxImportThreads != 1);
 			SceneGraphBuilder = MakeUnique<FDatasmithSceneGraphBuilder>(CADFileToUEFileMap, CacheDir, DatasmithScene, Source, ImportParameters);
 			SceneGraphBuilder->LoadSceneGraphDescriptionFiles();
 			MeshBuilderPtr = MakeUnique<FDatasmithMeshBuilder>(CADFileToUEGeomMap, CacheDir, ImportParameters);
diff --git a/Engine/Source/Programs/Enterprise/Datasmith/DatasmithCADWorker/Private/DatasmithCADWorkerImpl.cpp b/Engine/Source/Programs/Enterprise/Datasmith/DatasmithCADWorker/Private/DatasmithCADWorkerImpl.cpp
--- a/Engine/Source/Programs/Enterprise/Datasmith/DatasmithCADWorker/Private/DatasmithCADWorkerImpl.cpp
+++ b/Engine/Source/Programs/Enterprise/Datasmith/DatasmithCADWorker/Private/DatasmithCADWorkerImpl.cpp
@@ -242,12 +242,9 @@ void FDatasmithCADWorkerImpl::CheckMemory()
 	uint64 MaxMemoryUsed = FPlatformMemory::GetStats().UsedPhysical;
 	UE_LOG(LogDatasmithCADWorker, Verbose, TEXT("    - Start Ram used %llu MB"), MaxMemoryUsed / OneMegaBit);
 
-	uint64 StartTime = FPlatformTime::Cycles64();
 	while (bProcessIsRunning)
 	{
 		FPlatformProcess::Sleep(0.1);
-		double Duration = (FPlatformTime::Cycles64() - StartTime) * FPlatformTime::GetSecondsPerCycle64();
-
 		const uint64 MemoryUsed = FPlatformMemory::GetStats().UsedPhysical;
 		if (MaxMemoryUsed < MemoryUsed)
 		{
