Rollback last update of TechSoftFileParser.cpp and TechSoftFileParser.h

diff --git a/Engine/Plugins/Enterprise/DatasmithCADImporter/Source/CADInterfaces/Private/TechSoftFileParser.cpp b/Engine/Plugins/Enterprise/DatasmithCADImporter/Source/CADInterfaces/Private/TechSoftFileParser.cpp
--- a/Engine/Plugins/Enterprise/DatasmithCADImporter/Source/CADInterfaces/Private/TechSoftFileParser.cpp
+++ b/Engine/Plugins/Enterprise/DatasmithCADImporter/Source/CADInterfaces/Private/TechSoftFileParser.cpp
@@ -473,86 +473,6 @@ ECADParsingResult FTechSoftFileParser::Process()
 	return Result;
 }
 
-
-A3DStatus FTechSoftFileParser::AdaptBRepModel()
-{
-	TRACE_CPUPROFILER_EVENT_SCOPE(FTechSoftFileParserCADKernelTessellator::AdaptBRepModel);
-
-	const A3DUns32 ValidSurfaceCount = 9;
-	A3DUns32 AcceptedSurfaces[ValidSurfaceCount] = {
-		//kA3DTypeSurfBlend01,
-		//kA3DTypeSurfBlend02,
-		//kA3DTypeSurfBlend03,
-		kA3DTypeSurfNurbs,
-		kA3DTypeSurfCone,
-		kA3DTypeSurfCylinder,
-		kA3DTypeSurfCylindrical,
-		//kA3DTypeSurfOffset,
-		//kA3DTypeSurfPipe,
-		kA3DTypeSurfPlane,
-		kA3DTypeSurfRuled,
-		kA3DTypeSurfSphere,
-		kA3DTypeSurfRevolution,
-		//kA3DTypeSurfExtrusion,
-		//kA3DTypeSurfFromCurves,
-		kA3DTypeSurfTorus,
-		//kA3DTypeSurfTransform,
-	};
-
-	const A3DUns32 ValidCurveCount = 7;
-	A3DUns32 AcceptedCurves[ValidCurveCount] = {
-		//kA3DTypeCrvBase,
-		//kA3DTypeCrvBlend02Boundary,
-		kA3DTypeCrvNurbs,
-		kA3DTypeCrvCircle,
-		//kA3DTypeCrvComposite,
-		//kA3DTypeCrvOnSurf,
-		kA3DTypeCrvEllipse,
-		//kA3DTypeCrvEquation,
-		//kA3DTypeCrvHelix,
-		kA3DTypeCrvHyperbola,
-		//kA3DTypeCrvIntersection,
-		kA3DTypeCrvLine,
-		//kA3DTypeCrvOffset,
-		kA3DTypeCrvParabola,
-		kA3DTypeCrvPolyLine,
-		//kA3DTypeCrvTransform,
-	};
-
-	TUniqueTSObj<A3DCopyAndAdaptBrepModelData> CopyAndAdaptBrepModelData;
-	CopyAndAdaptBrepModelData->m_bUseSameParam = false;                        // If `A3D_TRUE`, surfaces will keep their parametrization when converted to NURBS.       
-	CopyAndAdaptBrepModelData->m_dTol = 1e-3;                                  // Tolerance value of resulting B-rep. The value is relative to the scale of the model.
-	CopyAndAdaptBrepModelData->m_bDeleteCrossingUV = false;                    // If `A3D_TRUE`, UV curves that cross seams of periodic surfaces are replaced by 3D curves 
-	CopyAndAdaptBrepModelData->m_bSplitFaces = true;                           // If `A3D_TRUE`, the faces with a periodic basis surface are split on parametric seams
-	CopyAndAdaptBrepModelData->m_bSplitClosedFaces = false;                    // If `A3D_TRUE`, the faces with a closed basis surface are split into faces at the parametric seam and mid-parameter
-	CopyAndAdaptBrepModelData->m_bForceComputeUV = true;                       // If `A3D_TRUE`, UV curves are computed from the B-rep data
-	CopyAndAdaptBrepModelData->m_bAllowUVCrossingSeams = true;                 // If `A3D_TRUE` and m_bForceComputeUV is set to `A3D_TRUE`, computed UV curves can cross seams.
-	CopyAndAdaptBrepModelData->m_bForceCompute3D = false;                      // If `A3D_TRUE`, 3D curves are computed from the B-rep data
-	CopyAndAdaptBrepModelData->m_bContinueOnError = true;                      // Continue processing even if an error occurs. Use \ref A3DCopyAndAdaptBrepModelAdvanced to get the error status.
-	CopyAndAdaptBrepModelData->m_bClampTolerantUVCurvesInsideUVDomain = false; // If `A3D_FALSE`, UV curves may stray outside the UV domain as long as the 3D edge tolerance is respected. If set to `A3D_TRUE`, the UV curves will be clamped to the UV domain (if the clamp still leaves them within the edge tolerance). */
-	CopyAndAdaptBrepModelData->m_bForceDuplicateGeometries = false;            // If `A3D_TRUE`, break the sharing of surfaces and curves into topologies.*/
-
-	CopyAndAdaptBrepModelData->m_uiAcceptableSurfacesSize = ValidSurfaceCount;
-	CopyAndAdaptBrepModelData->m_puiAcceptableSurfaces = &AcceptedSurfaces[0];
-	CopyAndAdaptBrepModelData->m_uiAcceptableCurvesSize = ValidCurveCount;
-	CopyAndAdaptBrepModelData->m_puiAcceptableCurves = &AcceptedCurves[0];
-
-	int32 ErrorCount = 0;
-	A3DCopyAndAdaptBrepModelErrorData* Errors = nullptr;
-	A3DStatus Ret = TechSoftInterface::AdaptBRepInModelFile(ModelFile.Get(), *CopyAndAdaptBrepModelData, ErrorCount, &Errors);
-	if ((Ret == A3D_SUCCESS || Ret == A3D_TOOLS_CONTINUE_ON_ERROR) && ErrorCount > 0)
-	{
-		// Add message about non-critical errors during the adaptation
-		CADFileData.LogMessage(FString::Printf(TEXT("File %s had %d non-critical error(s) during BRep adaptation step."), *CADFileData.GetCADFileDescription().GetFileName(), ErrorCount));
-	}
-	else if (Ret != A3D_SUCCESS)
-	{
-		CADFileData.LogWarning(FString::Printf(TEXT("File %s failed during BRep adaptation step."), *CADFileData.GetCADFileDescription().GetFileName(), ErrorCount));
-		return A3D_ERROR;
-	}
-	return A3D_SUCCESS;
-}
-
 void FTechSoftFileParser::SewModel()
 {
 	CADLibrary::TUniqueTSObj<A3DSewOptionsData> SewData;
diff --git a/Engine/Plugins/Enterprise/DatasmithCADImporter/Source/CADInterfaces/Private/TechSoftFileParser.h b/Engine/Plugins/Enterprise/DatasmithCADImporter/Source/CADInterfaces/Private/TechSoftFileParser.h
--- a/Engine/Plugins/Enterprise/DatasmithCADImporter/Source/CADInterfaces/Private/TechSoftFileParser.h
+++ b/Engine/Plugins/Enterprise/DatasmithCADImporter/Source/CADInterfaces/Private/TechSoftFileParser.h
@@ -53,10 +53,10 @@ public:
 
 protected:
 
-	virtual A3DStatus AdaptBRepModel();
-	//{
-	//	return A3DStatus::A3D_SUCCESS;
-	//}
+	virtual A3DStatus AdaptBRepModel()
+	{
+		return A3DStatus::A3D_SUCCESS;
+	}
 
 	/**
 	 * If the tessellator is TechSoft, SewModel call TechSoftInterface::SewModel
