CADKernel update: Meshing of thin faces

diff --git a/Engine/Plugins/Enterprise/DatasmithCADImporter/Source/CADInterfaces/Private/TechSoftBridge.cpp b/Engine/Plugins/Enterprise/DatasmithCADImporter/Source/CADInterfaces/Private/TechSoftBridge.cpp
--- a/Engine/Plugins/Enterprise/DatasmithCADImporter/Source/CADInterfaces/Private/TechSoftBridge.cpp
+++ b/Engine/Plugins/Enterprise/DatasmithCADImporter/Source/CADInterfaces/Private/TechSoftBridge.cpp
@@ -218,7 +218,7 @@ UE::CADKernel::FBody* FTechSoftBridge::GetBody(A3DRiBrepModel* A3DBRepModel)
 UE::CADKernel::FBody* FTechSoftBridge::AddBody(A3DRiBrepModel* A3DBRepModel, TMap<FString, FString> MetaData, const double InBodyScale)
 {
 #ifdef CADKERNEL_DEV
-	Report->BodyCount++;
+	UE::CADKernel::FCADFileReport::Get().BodyCount++;
 #endif
 
 	//UE::CADKernel working unit is mm
@@ -315,7 +315,7 @@ void FTechSoftBridge::TraverseShell(const A3DTopoShell* A3DShell, TSharedRef<UE:
 	TSharedRef<UE::CADKernel::FShell> Shell = UE::CADKernel::FEntity::MakeShared<UE::CADKernel::FShell>();
 	Body->AddShell(Shell);
 #ifdef CADKERNEL_DEV
-	Report->ShellCount++;
+	UE::CADKernel::FCADFileReport::Get().ShellCount++;
 #endif
 	AddMetadata(MetaData, *Shell);
 
@@ -340,7 +340,7 @@ TSharedPtr<UE::CADKernel::FCurve> FTechSoftBridge::AddCurve(const A3DCrvBase* A3
 	if (Ret == A3D_SUCCESS)
 	{
 #ifdef CADKERNEL_DEV
-		Report->CurveCount++;
+		UE::CADKernel::FCADFileReport::Get().CurveCount++;
 #endif
 		switch (eType)
 		{
@@ -385,7 +385,7 @@ TSharedPtr<UE::CADKernel::FCurve> FTechSoftBridge::AddCurve(const A3DCrvBase* A3
 TSharedPtr<UE::CADKernel::FCurve> FTechSoftBridge::AddCurveLine(const A3DCrvLine* A3DCurve, const TechSoftUtils::FUVReparameterization& UVReparameterization)
 {
 #ifdef CADKERNEL_DEV
-	Report->CurveLineCount++;
+	UE::CADKernel::FCADFileReport::Get().CurveLineCount++;
 #endif
 
 	if (bUseCurveAsNurbs)
@@ -408,7 +408,7 @@ TSharedPtr<UE::CADKernel::FCurve> FTechSoftBridge::AddCurveLine(const A3DCrvLine
 TSharedPtr<UE::CADKernel::FCurve> FTechSoftBridge::AddCurveCircle(const A3DCrvCircle* A3DCurve, const TechSoftUtils::FUVReparameterization& UVReparameterization)
 {
 #ifdef CADKERNEL_DEV
-	Report->CurveCircleCount++;
+	UE::CADKernel::FCADFileReport::Get().CurveCircleCount++;
 #endif
 
 	if (bUseCurveAsNurbs)
@@ -430,7 +430,7 @@ TSharedPtr<UE::CADKernel::FCurve> FTechSoftBridge::AddCurveCircle(const A3DCrvCi
 TSharedPtr<UE::CADKernel::FCurve> FTechSoftBridge::AddCurveEllipse(const A3DCrvEllipse* A3DCurve, const TechSoftUtils::FUVReparameterization& UVReparameterization)
 {
 #ifdef CADKERNEL_DEV
-	Report->CurveEllipseCount++;
+	UE::CADKernel::FCADFileReport::Get().CurveEllipseCount++;
 #endif
 
 	if (bUseCurveAsNurbs)
@@ -452,7 +452,7 @@ TSharedPtr<UE::CADKernel::FCurve> FTechSoftBridge::AddCurveEllipse(const A3DCrvE
 TSharedPtr<UE::CADKernel::FCurve> FTechSoftBridge::AddCurveParabola(const A3DCrvParabola* A3DCurve, const TechSoftUtils::FUVReparameterization& UVReparameterization)
 {
 #ifdef CADKERNEL_DEV
-	Report->CurveParabolaCount++;
+	UE::CADKernel::FCADFileReport::Get().CurveParabolaCount++;
 #endif
 
 	if (bUseCurveAsNurbs)
@@ -474,7 +474,7 @@ TSharedPtr<UE::CADKernel::FCurve> FTechSoftBridge::AddCurveParabola(const A3DCrv
 TSharedPtr<UE::CADKernel::FCurve> FTechSoftBridge::AddCurveHyperbola(const A3DCrvHyperbola* A3DCurve, const TechSoftUtils::FUVReparameterization& UVReparameterization)
 {
 #ifdef CADKERNEL_DEV
-	Report->CurveHyperbolaCount++;
+	UE::CADKernel::FCADFileReport::Get().CurveHyperbolaCount++;
 #endif
 
 	if (bUseCurveAsNurbs)
@@ -496,7 +496,7 @@ TSharedPtr<UE::CADKernel::FCurve> FTechSoftBridge::AddCurveHyperbola(const A3DCr
 TSharedPtr<UE::CADKernel::FCurve> FTechSoftBridge::AddCurveHelix(const A3DCrvHelix* A3DCurve, const TechSoftUtils::FUVReparameterization& UVReparameterization)
 {
 #ifdef CADKERNEL_DEV
-	Report->CurveHelixCount++;
+	UE::CADKernel::FCADFileReport::Get().CurveHelixCount++;
 #endif
 
 	if (bUseCurveAsNurbs)
@@ -518,7 +518,7 @@ TSharedPtr<UE::CADKernel::FCurve> FTechSoftBridge::AddCurveHelix(const A3DCrvHel
 TSharedPtr<UE::CADKernel::FCurve> FTechSoftBridge::AddCurvePolyLine(const A3DCrvPolyLine* A3DCurve, const TechSoftUtils::FUVReparameterization& UVReparameterization)
 {
 #ifdef CADKERNEL_DEV
-	Report->CurvePolyLineCount++;
+	UE::CADKernel::FCADFileReport::Get().CurvePolyLineCount++;
 #endif
 
 	if (bUseCurveAsNurbs)
@@ -540,7 +540,7 @@ TSharedPtr<UE::CADKernel::FCurve> FTechSoftBridge::AddCurvePolyLine(const A3DCrv
 TSharedPtr<UE::CADKernel::FCurve> FTechSoftBridge::AddCurveComposite(const A3DCrvComposite* A3DCurve, const TechSoftUtils::FUVReparameterization& UVReparameterization)
 {
 #ifdef CADKERNEL_DEV
-	Report->CurveCompositeCount++;
+	UE::CADKernel::FCADFileReport::Get().CurveCompositeCount++;
 #endif
 
 	if (bUseCurveAsNurbs)
@@ -586,7 +586,7 @@ TSharedPtr<UE::CADKernel::FCurve> AddCurveNurbsFromData(A3DCrvNurbsData& A3DNurb
 TSharedPtr<UE::CADKernel::FCurve> FTechSoftBridge::AddCurveNurbs(const A3DCrvNurbs* A3DNurbs, const TechSoftUtils::FUVReparameterization& UVReparameterization)
 {
 #ifdef CADKERNEL_DEV
-	Report->CurveLineCount++;
+	UE::CADKernel::FCADFileReport::Get().CurveLineCount++;
 #endif
 
 	TUniqueTSObj<A3DCrvNurbsData> CrvNurbsData(A3DNurbs);
@@ -601,7 +601,7 @@ TSharedPtr<UE::CADKernel::FCurve> FTechSoftBridge::AddCurveNurbs(const A3DCrvNur
 TSharedPtr<UE::CADKernel::FCurve> FTechSoftBridge::AddCurveAsNurbs(const A3DCrvBase* A3DCurve, const TechSoftUtils::FUVReparameterization& UVReparameterization)
 {
 #ifdef CADKERNEL_DEV
-	Report->CurveAsNurbsCount++;
+	UE::CADKernel::FCADFileReport::Get().CurveAsNurbsCount++;
 #endif
 
 	TUniqueTSObj<A3DCrvNurbsData> NurbsData;
@@ -621,7 +621,7 @@ TSharedPtr<UE::CADKernel::FCurve> FTechSoftBridge::AddCurveAsNurbs(const A3DCrvB
 TSharedPtr<UE::CADKernel::FTopologicalEdge> FTechSoftBridge::AddEdge(const A3DTopoCoEdge* A3DCoedge, const TSharedRef<UE::CADKernel::FSurface>& Surface, const TechSoftUtils::FUVReparameterization& UVReparameterization, UE::CADKernel::EOrientation& OutOrientation)
 {
 #ifdef CADKERNEL_DEV
-	Report->EdgeCount++;
+	UE::CADKernel::FCADFileReport::Get().EdgeCount++;
 #endif
 
 	TUniqueTSObj<A3DTopoCoEdgeData> CoEdgeData(A3DCoedge);
@@ -659,7 +659,7 @@ TSharedPtr<UE::CADKernel::FTopologicalEdge> FTechSoftBridge::AddEdge(const A3DTo
 TSharedPtr<UE::CADKernel::FTopologicalLoop> FTechSoftBridge::AddLoop(const A3DTopoLoop* A3DLoop, const TSharedRef<UE::CADKernel::FSurface>& Surface, const TechSoftUtils::FUVReparameterization& UVReparameterization, const bool bIsExternalLoop)
 {
 #ifdef CADKERNEL_DEV
-	Report->LoopCount++;
+	UE::CADKernel::FCADFileReport::Get().LoopCount++;
 #endif
 
 	TArray<TSharedPtr<UE::CADKernel::FTopologicalEdge>> Edges;
@@ -669,7 +669,7 @@ TSharedPtr<UE::CADKernel::FTopologicalLoop> FTechSoftBridge::AddLoop(const A3DTo
 	if (!TopoLoopData.IsValid())
 	{
 #ifdef CADKERNEL_DEV
-		Report->DegeneratedLoopCount++;
+		UE::CADKernel::FCADFileReport::Get().DegeneratedLoopCount++;
 #endif
 		return TSharedPtr<UE::CADKernel::FTopologicalLoop>();
 	}
@@ -682,7 +682,7 @@ TSharedPtr<UE::CADKernel::FTopologicalLoop> FTechSoftBridge::AddLoop(const A3DTo
 		if (!Edge.IsValid())
 		{
 #ifdef CADKERNEL_DEV
-			Report->DegeneratedEdgeCount++;
+			UE::CADKernel::FCADFileReport::Get().DegeneratedEdgeCount++;
 #endif
 			continue;
 		}
@@ -694,7 +694,7 @@ TSharedPtr<UE::CADKernel::FTopologicalLoop> FTechSoftBridge::AddLoop(const A3DTo
 	if (Edges.Num() == 0)
 	{
 #ifdef CADKERNEL_DEV
-		Report->DegeneratedLoopCount++;
+		UE::CADKernel::FCADFileReport::Get().DegeneratedLoopCount++;
 #endif
 		return TSharedPtr<UE::CADKernel::FTopologicalLoop>();
 	}
@@ -750,7 +750,7 @@ void FTechSoftBridge::AddFace(const A3DTopoFace* A3DFace, UE::CADKernel::EOrient
 	using namespace UE::CADKernel;
 
 #ifdef CADKERNEL_DEV
-	Report->FaceCount++;
+	UE::CADKernel::FCADFileReport::Get().FaceCount++;
 #endif
 
 	FArchiveCADObject MetaData;
@@ -760,7 +760,7 @@ void FTechSoftBridge::AddFace(const A3DTopoFace* A3DFace, UE::CADKernel::EOrient
 	if (!TopoFaceData.IsValid())
 	{
 #ifdef CADKERNEL_DEV
-		Report->FailedFaceCount++;
+		UE::CADKernel::FCADFileReport::Get().FailedFaceCount++;
 #endif
 		return;
 	}
@@ -771,8 +771,8 @@ void FTechSoftBridge::AddFace(const A3DTopoFace* A3DFace, UE::CADKernel::EOrient
 	if (!SurfacePtr.IsValid())
 	{
 #ifdef CADKERNEL_DEV
-		Report->DegeneratedSurfaceCount++;
-		Report->FailedFaceCount++; 
+		UE::CADKernel::FCADFileReport::Get().DegeneratedSurfaceCount++;
+		UE::CADKernel::FCADFileReport::Get().FailedFaceCount++; 
 #endif
 		return;
 	}
@@ -871,7 +871,7 @@ void FTechSoftBridge::AddFace(const A3DTopoFace* A3DFace, UE::CADKernel::EOrient
 			if(!ExternalLoop->IsExternal())
 			{
 #ifdef CADKERNEL_DEV
-				Report->WrongExternalLoopCount++;
+				UE::CADKernel::FCADFileReport::Get().WrongExternalLoopCount++;
 				FMessage::Printf(Log, TEXT("ERROR: Face %d had wrong external loop\n"), Face->GetId());
 #endif
 	
@@ -890,7 +890,7 @@ void FTechSoftBridge::AddFace(const A3DTopoFace* A3DFace, UE::CADKernel::EOrient
 		if (Loops.Num() == 0)
 		{
 #ifdef CADKERNEL_DEV
-			Report->FailedFaceCount++;
+			UE::CADKernel::FCADFileReport::Get().FailedFaceCount++;
 #endif
 			Face->SetAsDegenerated();
 			Face->Delete();
@@ -901,7 +901,7 @@ void FTechSoftBridge::AddFace(const A3DTopoFace* A3DFace, UE::CADKernel::EOrient
 			int32 DoubtfulLoopOrientationCount = 0;
 			Face->AddLoops(Loops, DoubtfulLoopOrientationCount);
 #ifdef CADKERNEL_DEV
-			Report->DoubtfulLoopOrientationCount += DoubtfulLoopOrientationCount;
+			UE::CADKernel::FCADFileReport::Get().DoubtfulLoopOrientationCount += DoubtfulLoopOrientationCount;
 #endif
 		}
 	}
@@ -916,7 +916,7 @@ static bool bUseSurfaceAsNurbs = true;
 TSharedPtr<UE::CADKernel::FSurface> FTechSoftBridge::AddSurface(const A3DSurfBase* A3DSurface, TechSoftUtils::FUVReparameterization& OutUVReparameterization)
 {
 #ifdef CADKERNEL_DEV
-	Report->SurfaceCount++;
+	UE::CADKernel::FCADFileReport::Get().SurfaceCount++;
 #endif
 
 	FArchiveCADObject MetaData;
@@ -993,7 +993,7 @@ TSharedPtr<UE::CADKernel::FSurface> FTechSoftBridge::AddSurface(const A3DSurfBas
 TSharedPtr<UE::CADKernel::FSurface> FTechSoftBridge::AddConeSurface(const A3DSurfBase* Surface, TechSoftUtils::FUVReparameterization& OutUVReparameterization)
 {
 #ifdef CADKERNEL_DEV
-	Report->ConeSurfaceCount++;
+	UE::CADKernel::FCADFileReport::Get().ConeSurfaceCount++;
 #endif
 
 	TUniqueTSObj<A3DSurfConeData> A3DConeData(Surface);
@@ -1017,7 +1017,7 @@ TSharedPtr<UE::CADKernel::FSurface> FTechSoftBridge::AddConeSurface(const A3DSur
 TSharedPtr<UE::CADKernel::FSurface> FTechSoftBridge::AddCylinderSurface(const A3DSurfBase* Surface, TechSoftUtils::FUVReparameterization& OutUVReparameterization)
 {
 #ifdef CADKERNEL_DEV
-	Report->CylinderSurfaceCount++;
+	UE::CADKernel::FCADFileReport::Get().CylinderSurfaceCount++;
 #endif
 
 	TUniqueTSObj<A3DSurfCylinderData> A3DCylinderData(Surface);
@@ -1037,7 +1037,7 @@ TSharedPtr<UE::CADKernel::FSurface> FTechSoftBridge::AddCylinderSurface(const A3
 TSharedPtr<UE::CADKernel::FSurface> FTechSoftBridge::AddLinearTransfoSurface(const A3DSurfBase* Surface, TechSoftUtils::FUVReparameterization& OutUVReparameterization)
 {
 #ifdef CADKERNEL_DEV
-	Report->LinearTransfoSurfaceCount++;
+	UE::CADKernel::FCADFileReport::Get().LinearTransfoSurfaceCount++;
 #endif
 
 	if (bUseSurfaceAsNurbs)
@@ -1051,7 +1051,7 @@ TSharedPtr<UE::CADKernel::FSurface> FTechSoftBridge::AddLinearTransfoSurface(con
 TSharedPtr<UE::CADKernel::FSurface> FTechSoftBridge::AddNurbsSurface(const A3DSurfNurbs* Nurbs, TechSoftUtils::FUVReparameterization& OutUVReparameterization)
 {
 #ifdef CADKERNEL_DEV
-	Report->NurbsSurfaceCount++;
+	UE::CADKernel::FCADFileReport::Get().NurbsSurfaceCount++;
 #endif
 
 	TUniqueTSObj<A3DSurfNurbsData> A3DNurbsData(Nurbs);
@@ -1066,7 +1066,7 @@ TSharedPtr<UE::CADKernel::FSurface> FTechSoftBridge::AddNurbsSurface(const A3DSu
 TSharedPtr<UE::CADKernel::FSurface> FTechSoftBridge::AddOffsetSurface(const A3DSurfBase* Surface, TechSoftUtils::FUVReparameterization& OutUVReparameterization)
 {
 #ifdef CADKERNEL_DEV
-	Report->OffsetSurfaceCount++;
+	UE::CADKernel::FCADFileReport::Get().OffsetSurfaceCount++;
 #endif
 
 	if (bUseSurfaceAsNurbs)
@@ -1080,7 +1080,7 @@ TSharedPtr<UE::CADKernel::FSurface> FTechSoftBridge::AddOffsetSurface(const A3DS
 TSharedPtr<UE::CADKernel::FSurface> FTechSoftBridge::AddPlaneSurface(const A3DSurfPlane* Surface, TechSoftUtils::FUVReparameterization& OutUVReparameterization)
 {
 #ifdef CADKERNEL_DEV
-	Report->PlaneSurfaceCount++;
+	UE::CADKernel::FCADFileReport::Get().PlaneSurfaceCount++;
 #endif
 
 	TUniqueTSObj<A3DSurfPlaneData> A3DPlaneData(Surface);
@@ -1100,7 +1100,7 @@ TSharedPtr<UE::CADKernel::FSurface> FTechSoftBridge::AddPlaneSurface(const A3DSu
 TSharedPtr<UE::CADKernel::FSurface> FTechSoftBridge::AddRevolutionSurface(const A3DSurfRevolution* Surface, TechSoftUtils::FUVReparameterization& OutUVReparameterization)
 {
 #ifdef CADKERNEL_DEV
-	Report->RevolutionSurfaceCount++;
+	UE::CADKernel::FCADFileReport::Get().RevolutionSurfaceCount++;
 #endif
 
 	if (bUseSurfaceAsNurbs)
@@ -1120,7 +1120,7 @@ TSharedPtr<UE::CADKernel::FSurface> FTechSoftBridge::AddRevolutionSurface(const
 TSharedPtr<UE::CADKernel::FSurface> FTechSoftBridge::AddRuledSurface(const A3DSurfBase* Surface, TechSoftUtils::FUVReparameterization& OutUVReparameterization)
 {
 #ifdef CADKERNEL_DEV
-	Report->RuledSurfaceCount++;
+	UE::CADKernel::FCADFileReport::Get().RuledSurfaceCount++;
 #endif
 
 	if (bUseSurfaceAsNurbs)
@@ -1140,7 +1140,7 @@ TSharedPtr<UE::CADKernel::FSurface> FTechSoftBridge::AddRuledSurface(const A3DSu
 TSharedPtr<UE::CADKernel::FSurface> FTechSoftBridge::AddSphereSurface(const A3DSurfBase* Surface, TechSoftUtils::FUVReparameterization& OutUVReparameterization)
 {
 #ifdef CADKERNEL_DEV
-	Report->SphereSurfaceCount++;
+	UE::CADKernel::FCADFileReport::Get().SphereSurfaceCount++;
 #endif
 
 	TUniqueTSObj<A3DSurfSphereData> A3DSphereData(Surface);
@@ -1159,7 +1159,7 @@ TSharedPtr<UE::CADKernel::FSurface> FTechSoftBridge::AddSphereSurface(const A3DS
 TSharedPtr<UE::CADKernel::FSurface> FTechSoftBridge::AddTorusSurface(const A3DSurfBase* Surface, TechSoftUtils::FUVReparameterization& OutUVReparameterization)
 {
 #ifdef CADKERNEL_DEV
-	Report->TorusSurfaceCount++;
+	UE::CADKernel::FCADFileReport::Get().TorusSurfaceCount++;
 #endif
 
 	TUniqueTSObj<A3DSurfTorusData> A3DTorusData(Surface);
@@ -1177,7 +1177,7 @@ TSharedPtr<UE::CADKernel::FSurface> FTechSoftBridge::AddTorusSurface(const A3DSu
 TSharedPtr<UE::CADKernel::FSurface> FTechSoftBridge::AddBlend01Surface(const A3DSurfBase* Surface, TechSoftUtils::FUVReparameterization& OutUVReparameterization)
 {
 #ifdef CADKERNEL_DEV
-	Report->Blend01SurfaceCount++;
+	UE::CADKernel::FCADFileReport::Get().Blend01SurfaceCount++;
 #endif
 
 	if (bUseSurfaceAsNurbs)
@@ -1197,7 +1197,7 @@ TSharedPtr<UE::CADKernel::FSurface> FTechSoftBridge::AddBlend01Surface(const A3D
 TSharedPtr<UE::CADKernel::FSurface> FTechSoftBridge::AddBlend02Surface(const A3DSurfBase* Surface, TechSoftUtils::FUVReparameterization& OutUVReparameterization)
 {
 #ifdef CADKERNEL_DEV
-	Report->Blend02SurfaceCount++;
+	UE::CADKernel::FCADFileReport::Get().Blend02SurfaceCount++;
 #endif
 
 	if (bUseSurfaceAsNurbs)
@@ -1217,7 +1217,7 @@ TSharedPtr<UE::CADKernel::FSurface> FTechSoftBridge::AddBlend02Surface(const A3D
 TSharedPtr<UE::CADKernel::FSurface> FTechSoftBridge::AddBlend03Surface(const A3DSurfBase* Surface, TechSoftUtils::FUVReparameterization& OutUVReparameterization)
 {
 #ifdef CADKERNEL_DEV
-	Report->Blend03SurfaceCount++;
+	UE::CADKernel::FCADFileReport::Get().Blend03SurfaceCount++;
 #endif
 
 	if (bUseSurfaceAsNurbs)
@@ -1231,7 +1231,7 @@ TSharedPtr<UE::CADKernel::FSurface> FTechSoftBridge::AddBlend03Surface(const A3D
 TSharedPtr<UE::CADKernel::FSurface> FTechSoftBridge::AddCylindricalSurface(const A3DSurfBase* Surface, TechSoftUtils::FUVReparameterization& OutUVReparameterization)
 {
 #ifdef CADKERNEL_DEV
-	Report->CylindricalSurfaceCount++;
+	UE::CADKernel::FCADFileReport::Get().CylindricalSurfaceCount++;
 #endif
 
 	if (bUseSurfaceAsNurbs)
@@ -1251,7 +1251,7 @@ TSharedPtr<UE::CADKernel::FSurface> FTechSoftBridge::AddCylindricalSurface(const
 TSharedPtr<UE::CADKernel::FSurface> FTechSoftBridge::AddPipeSurface(const A3DSurfBase* Surface, TechSoftUtils::FUVReparameterization& OutUVReparameterization)
 {
 #ifdef CADKERNEL_DEV
-	Report->PipeSurfaceCount++;
+	UE::CADKernel::FCADFileReport::Get().PipeSurfaceCount++;
 #endif
 
 	if (bUseSurfaceAsNurbs)
@@ -1271,7 +1271,7 @@ TSharedPtr<UE::CADKernel::FSurface> FTechSoftBridge::AddPipeSurface(const A3DSur
 TSharedPtr<UE::CADKernel::FSurface> FTechSoftBridge::AddExtrusionSurface(const A3DSurfBase* Surface, TechSoftUtils::FUVReparameterization& OutUVReparameterization)
 {
 #ifdef CADKERNEL_DEV
-	Report->ExtrusionSurfaceCount++;
+	UE::CADKernel::FCADFileReport::Get().ExtrusionSurfaceCount++;
 #endif
 
 	if (bUseSurfaceAsNurbs)
@@ -1291,7 +1291,7 @@ TSharedPtr<UE::CADKernel::FSurface> FTechSoftBridge::AddExtrusionSurface(const A
 TSharedPtr<UE::CADKernel::FSurface> FTechSoftBridge::AddSurfaceFromCurves(const A3DSurfBase* Surface, TechSoftUtils::FUVReparameterization& OutUVReparameterization)
 {
 #ifdef CADKERNEL_DEV
-	Report->SurfaceFromCurvesCount++;
+	UE::CADKernel::FCADFileReport::Get().SurfaceFromCurvesCount++;
 #endif
 
 	if (bUseSurfaceAsNurbs)
@@ -1311,7 +1311,7 @@ TSharedPtr<UE::CADKernel::FSurface> FTechSoftBridge::AddSurfaceFromCurves(const
 TSharedPtr<UE::CADKernel::FSurface> FTechSoftBridge::AddTransformSurface(const A3DSurfBase* Surface, TechSoftUtils::FUVReparameterization& OutUVReparameterization)
 {
 #ifdef CADKERNEL_DEV
-	Report->TransformSurfaceCount++;
+	UE::CADKernel::FCADFileReport::Get().TransformSurfaceCount++;
 #endif
 
 	if (bUseSurfaceAsNurbs)
@@ -1365,7 +1365,7 @@ TSharedPtr<UE::CADKernel::FSurface> FTechSoftBridge::AddSurfaceNurbs(const A3DSu
 TSharedPtr<UE::CADKernel::FSurface> FTechSoftBridge::AddSurfaceAsNurbs(const A3DSurfBase* Surface, TechSoftUtils::FUVReparameterization& OutUVReparameterization)
 {
 #ifdef CADKERNEL_DEV
-	Report->SurfaceAsNurbsCount++;
+	UE::CADKernel::FCADFileReport::Get().SurfaceAsNurbsCount++;
 #endif
 
 	TUniqueTSObj<A3DSurfNurbsData> A3DNurbsData;
diff --git a/Engine/Plugins/Enterprise/DatasmithCADImporter/Source/CADInterfaces/Private/TechSoftBridge.h b/Engine/Plugins/Enterprise/DatasmithCADImporter/Source/CADInterfaces/Private/TechSoftBridge.h
--- a/Engine/Plugins/Enterprise/DatasmithCADImporter/Source/CADInterfaces/Private/TechSoftBridge.h
+++ b/Engine/Plugins/Enterprise/DatasmithCADImporter/Source/CADInterfaces/Private/TechSoftBridge.h
@@ -9,10 +9,6 @@
 
 #include "TechSoftInterface.h"
 
-#ifdef CADKERNEL_DEV
-#include "CADFileReport.h"
-#endif
-
 namespace UE::CADKernel
 {
 
@@ -169,10 +165,6 @@ private:
 
 	UE::CADKernel::FSession& Session;
 	UE::CADKernel::FModel& Model;
-#ifdef CADKERNEL_DEV
-	UE::CADKernel::FCADFileReport* Report;
-#endif
-
 
 	const double GeometricTolerance;
 	const double EdgeLengthTolerance;
@@ -191,13 +183,6 @@ public:
 	UE::CADKernel::FBody* GetBody(A3DRiBrepModel* A3DBRepModel);
 	const A3DRiBrepModel* GetA3DBody(UE::CADKernel::FBody* BRepModel);
 
-#ifdef CADKERNEL_DEV
-	void SetReport(UE::CADKernel::FCADFileReport& InReport)
-	{
-		Report = &InReport;
-	}
-#endif
-
 private:
 
 	void TraverseBrepData(const A3DTopoBrepData* A3DBrepData, TSharedRef<UE::CADKernel::FBody>& Body);
diff --git a/Engine/Plugins/Enterprise/DatasmithCADImporter/Source/CADInterfaces/Private/TechSoftFileParserCADKernelTessellator.cpp b/Engine/Plugins/Enterprise/DatasmithCADImporter/Source/CADInterfaces/Private/TechSoftFileParserCADKernelTessellator.cpp
--- a/Engine/Plugins/Enterprise/DatasmithCADImporter/Source/CADInterfaces/Private/TechSoftFileParserCADKernelTessellator.cpp
+++ b/Engine/Plugins/Enterprise/DatasmithCADImporter/Source/CADInterfaces/Private/TechSoftFileParserCADKernelTessellator.cpp
@@ -14,7 +14,7 @@
 #include "CADKernel/Core/Session.h"
 #include "CADKernel/Core/Types.h"
 
-#include "CADKernel/Mesh/Meshers/ParametricMesher.h"
+#include "CADKernel/Mesh/Meshers/Mesher.h"
 #include "CADKernel/Mesh/Structure/ModelMesh.h"
 
 #include "CADKernel/Topo/Body.h"
@@ -266,7 +266,9 @@ void FTechSoftFileParserCADKernelTessellator::MeshAndGetTessellation(UE::CADKern
 
 	FCADKernelTools::DefineMeshCriteria(CADKernelModelMesh.Get(), CADFileData.GetImportParameters(), CADKernelSession.GetGeometricTolerance());
 
-	FParametricMesher Mesher(*CADKernelModelMesh);
+	const double GeometricTolerance = FImportParameters::GStitchingTolerance * 10; // cm to mm
+	const bool bActivateThinZoneMeshing = FImportParameters::bGActivateThinZoneMeshing;
+	FMesher Mesher(*CADKernelModelMesh, GeometricTolerance, bActivateThinZoneMeshing);
 	Mesher.MeshEntity(CADKernelBody);
 
 	FCADKernelTools::GetBodyTessellation(*CADKernelModelMesh, CADKernelBody, BodyMesh);
diff --git a/Engine/Plugins/Enterprise/DatasmithCADImporter/Source/CADKernelSurface/Private/CADKernelSurfaceExtension.cpp b/Engine/Plugins/Enterprise/DatasmithCADImporter/Source/CADKernelSurface/Private/CADKernelSurfaceExtension.cpp
--- a/Engine/Plugins/Enterprise/DatasmithCADImporter/Source/CADKernelSurface/Private/CADKernelSurfaceExtension.cpp
+++ b/Engine/Plugins/Enterprise/DatasmithCADImporter/Source/CADKernelSurface/Private/CADKernelSurfaceExtension.cpp
@@ -17,7 +17,6 @@
 #include "CADKernel/Core/Entity.h"
 #include "CADKernel/Core/Session.h"
 #include "CADKernel/Core/Types.h"
-#include "CADKernel/Mesh/Meshers/ParametricMesher.h"
 #include "CADKernel/Mesh/Structure/ModelMesh.h"
 #include "CADKernel/Topo/Model.h"
 #include "CADKernel/Topo/Body.h"
diff --git a/Engine/Plugins/Enterprise/DatasmithCADImporter/Source/CADLibrary/Private/CADKernelTools.cpp b/Engine/Plugins/Enterprise/DatasmithCADImporter/Source/CADLibrary/Private/CADKernelTools.cpp
--- a/Engine/Plugins/Enterprise/DatasmithCADImporter/Source/CADLibrary/Private/CADKernelTools.cpp
+++ b/Engine/Plugins/Enterprise/DatasmithCADImporter/Source/CADLibrary/Private/CADKernelTools.cpp
@@ -15,7 +15,7 @@
 #include "CADKernel/Core/Types.h"
 
 #include "CADKernel/Mesh/Criteria/Criterion.h"
-#include "CADKernel/Mesh/Meshers/ParametricMesher.h"
+#include "CADKernel/Mesh/Meshers/Mesher.h"
 #include "CADKernel/Mesh/Structure/FaceMesh.h"
 #include "CADKernel/Mesh/Structure/ModelMesh.h"
 
@@ -32,10 +32,10 @@ typedef uint32 TriangleIndex[3];
 namespace CADLibrary
 {
 
-static void FillVertexPosition(FMeshConversionContext& Context, const TSharedRef<UE::CADKernel::FModelMesh>& ModelMesh, FMeshDescription& MeshDescription)
+static void FillVertexPosition(FMeshConversionContext& Context, UE::CADKernel::FModelMesh& ModelMesh, FMeshDescription& MeshDescription)
 {
 	TArray<FVector3f> VertexArray;
-	ModelMesh->GetNodeCoordinates(VertexArray);
+	ModelMesh.GetNodeCoordinates(VertexArray);
 
 	for (FVector3f& Vertex : VertexArray)
 	{
@@ -78,7 +78,7 @@ static void FillVertexPosition(FMeshConversionContext& Context, const TSharedRef
 	}
 }
 
-bool FillMesh(FMeshConversionContext& Context, const TSharedRef<UE::CADKernel::FModelMesh>& ModelMesh, FMeshDescription& MeshDescription)
+bool FillMesh(FMeshConversionContext& Context, UE::CADKernel::FModelMesh& ModelMesh, FMeshDescription& MeshDescription)
 {
 	using namespace UE::CADKernel;
 
@@ -108,9 +108,9 @@ bool FillMesh(FMeshConversionContext& Context, const TSharedRef<UE::CADKernel::F
 
 	// Find all the materials used
 	TMap<uint32, FPolygonGroupID> MaterialToPolygonGroupMapping;
-	for (const TSharedPtr<FFaceMesh>& FaceMesh : ModelMesh->GetFaceMeshes())
+	for (const FFaceMesh* FaceMesh : ModelMesh.GetFaceMeshes())
 	{
-		const FTopologicalFace& Face = (const FTopologicalFace&)FaceMesh->GetGeometricEntity();
+		const FTopologicalFace& Face = (const FTopologicalFace&) FaceMesh->GetGeometricEntity();
 		// we assume that face has only color
 		MaterialToPolygonGroupMapping.Add(Face.GetColorId(), INDEX_NONE);
 	}
@@ -149,7 +149,7 @@ bool FillMesh(FMeshConversionContext& Context, const TSharedRef<UE::CADKernel::F
 		TArray<int32>& VertexIdSet = (Step == 0) ? Context.VertexIds : Context.SymmetricVertexIds;
 
 		// Loop through the FaceMeshes and collect all tessellation data
-		for (const TSharedPtr<FFaceMesh>& FaceMesh : ModelMesh->GetFaceMeshes())
+		for (FFaceMesh* FaceMesh : ModelMesh.GetFaceMeshes())
 		{
 			const FTopologicalFace& Face = (const FTopologicalFace&)FaceMesh->GetGeometricEntity();
 			if (bImportOnlyAlreadyPresent && !PatchIdSet.Contains(Face.GetPatchId()))
@@ -168,8 +168,8 @@ bool FillMesh(FMeshConversionContext& Context, const TSharedRef<UE::CADKernel::F
 			//FVector Temp3D = { 0, 0, 0 };
 			//FVector2D TexCoord2D = { 0, 0 };
 
-			TArray<int32>& TriangleVertexIndices = FaceMesh->TrianglesVerticesIndex;
-			TArray<int32>& VerticesGlobalIndex = FaceMesh->VerticesGlobalIndex;
+			const TArray<int32>& TriangleVertexIndices = FaceMesh->TrianglesVerticesIndex;
+			const TArray<int32>& VerticesGlobalIndex = FaceMesh->VerticesGlobalIndex;
 			MeshVertexInstanceIDs.Empty(TriangleVertexIndices.Num());
 
 			PatchIndex++;
@@ -242,10 +242,10 @@ bool FillMesh(FMeshConversionContext& Context, const TSharedRef<UE::CADKernel::F
 	return true;
 }
 
-static bool ConvertModelMeshToMeshDescription(FMeshConversionContext& Context, const TSharedRef<UE::CADKernel::FModelMesh>& InModelMesh, FMeshDescription& MeshDescription)
+static bool ConvertModelMeshToMeshDescription(FMeshConversionContext& Context, UE::CADKernel::FModelMesh& InModelMesh, FMeshDescription& MeshDescription)
 {
-	int32 VertexCount = InModelMesh->GetVertexCount();
-	int32 TriangleCount = InModelMesh->GetTriangleCount();
+	int32 VertexCount = InModelMesh.GetVertexCount();
+	int32 TriangleCount = InModelMesh.GetTriangleCount();
 
 	MeshDescription.ReserveNewVertexInstances(VertexCount);
 	MeshDescription.ReserveNewPolygons(TriangleCount);
@@ -269,13 +269,15 @@ bool FCADKernelTools::Tessellate(UE::CADKernel::FTopologicalShapeEntity& CADTopo
 
 	// Tessellate the model
 	TSharedRef<FModelMesh> CADKernelModelMesh = FEntity::MakeShared<FModelMesh>();
-	FParametricMesher Mesher(*CADKernelModelMesh);
 
-	constexpr double GeometricTolerance = 0.01;
+	const double GeometricTolerance = FImportParameters::GStitchingTolerance * 10; // cm to mm
+	const bool bActivateThinZoneMeshing = FImportParameters::bGActivateThinZoneMeshing;
+	FMesher Mesher(*CADKernelModelMesh, GeometricTolerance, bActivateThinZoneMeshing);
+
 	DefineMeshCriteria(*CADKernelModelMesh, TessellationContext.ImportParameters, GeometricTolerance);
 	Mesher.MeshEntity(CADTopologicalEntity);
 
-	return ConvertModelMeshToMeshDescription(TessellationContext, CADKernelModelMesh, OutMeshDescription);
+	return ConvertModelMeshToMeshDescription(TessellationContext, *CADKernelModelMesh, OutMeshDescription);
 }
 
 uint32 FCADKernelTools::GetFaceTessellation(UE::CADKernel::FFaceMesh& FaceMesh, FBodyMesh& OutBodyMesh, FObjectDisplayDataId FaceMaterial)
diff --git a/Engine/Plugins/Enterprise/DatasmithCADImporter/Source/CADTools/Private/CADOption.cpp b/Engine/Plugins/Enterprise/DatasmithCADImporter/Source/CADTools/Private/CADOption.cpp
--- a/Engine/Plugins/Enterprise/DatasmithCADImporter/Source/CADTools/Private/CADOption.cpp
+++ b/Engine/Plugins/Enterprise/DatasmithCADImporter/Source/CADTools/Private/CADOption.cpp
@@ -100,6 +100,15 @@ Only available with CADKernel.\n\
 Default value of RemoveDuplicatedFaces is true\n"),
 ECVF_Default);
 
+bool FImportParameters::bGActivateThinZoneMeshing = true;
+FAutoConsoleVariableRef GActivateThinZoneMeshing(
+	TEXT("ds.CADTranslator.Meshing.ActivateThinZoneMeshing"),
+	FImportParameters::bGActivateThinZoneMeshing,
+	TEXT("During the meshing process, faces with thin zones are meshed with a specific process for the thin zones. \n\
+Only available with CADKernel.\n\
+Default value of ActivateThinZoneMeshing is true\n"),
+ECVF_Default);
+
 float FImportParameters::GUnitScale = 1.f;
 FAutoConsoleVariableRef GCADTranslatorUnitScale(
 	TEXT("ds.CADTranslator.UnitScale"),
diff --git a/Engine/Plugins/Enterprise/DatasmithCADImporter/Source/CADTools/Public/CADOptions.h b/Engine/Plugins/Enterprise/DatasmithCADImporter/Source/CADTools/Public/CADOptions.h
--- a/Engine/Plugins/Enterprise/DatasmithCADImporter/Source/CADTools/Public/CADOptions.h
+++ b/Engine/Plugins/Enterprise/DatasmithCADImporter/Source/CADTools/Public/CADOptions.h
@@ -53,6 +53,7 @@ namespace CADLibrary
 		CADTOOLS_API static bool bGStitchingForceSew;
 		CADTOOLS_API static bool bGStitchingRemoveThinFaces;
 		CADTOOLS_API static bool bGStitchingRemoveDuplicatedFaces;
+		CADTOOLS_API static bool bGActivateThinZoneMeshing;
 		CADTOOLS_API static float GStitchingForceFactor;
 		CADTOOLS_API static float GUnitScale;
 		CADTOOLS_API static float GMeshingParameterFactor;
@@ -88,7 +89,7 @@ namespace CADLibrary
 			{
 				Hash = HashCombine(Hash, ::GetTypeHash(Param));
 			}
-			for (bool Param : {bGStitchingForceSew, bGStitchingRemoveThinFaces, bGStitchingRemoveDuplicatedFaces, bGDisableCADKernelTessellation, bGPreferJtFileEmbeddedTessellation})
+			for (bool Param : {bGStitchingForceSew, bGStitchingRemoveThinFaces, bGStitchingRemoveDuplicatedFaces, bGDisableCADKernelTessellation, bGActivateThinZoneMeshing, bGPreferJtFileEmbeddedTessellation})
 			{
 				Hash = HashCombine(Hash, ::GetTypeHash(Param));
 			}
@@ -115,6 +116,7 @@ namespace CADLibrary
 			Ar << ImportParameters.bGStitchingForceSew;
 			Ar << ImportParameters.bGStitchingRemoveDuplicatedFaces;
 			Ar << ImportParameters.bGStitchingRemoveThinFaces;
+			Ar << ImportParameters.bGActivateThinZoneMeshing;
 			Ar << ImportParameters.bValidationProcess;
 			Ar << ImportParameters.GMaxMaterialCountPerMesh;
 			Ar << ImportParameters.GMeshingParameterFactor;
diff --git a/Engine/Plugins/Enterprise/DatasmithCADImporter/Source/DatasmithOpenNurbsTranslator/Private/OpenNurbsBRepToCADKernelConverter.cpp b/Engine/Plugins/Enterprise/DatasmithCADImporter/Source/DatasmithOpenNurbsTranslator/Private/OpenNurbsBRepToCADKernelConverter.cpp
--- a/Engine/Plugins/Enterprise/DatasmithCADImporter/Source/DatasmithOpenNurbsTranslator/Private/OpenNurbsBRepToCADKernelConverter.cpp
+++ b/Engine/Plugins/Enterprise/DatasmithCADImporter/Source/DatasmithOpenNurbsTranslator/Private/OpenNurbsBRepToCADKernelConverter.cpp
@@ -20,7 +20,6 @@
 
 #include "CADKernel/Math/Boundary.h"
 #include "CADKernel/Math/Point.h"
-#include "CADKernel/Mesh/Meshers/ParametricMesher.h"
 #include "CADKernel/Mesh/Structure/ModelMesh.h"
 
 #include "CADKernel/Topo/Body.h"
diff --git a/Engine/Source/Runtime/Datasmith/CADKernel/Private/CADKernel/Core/Database.cpp b/Engine/Source/Runtime/Datasmith/CADKernel/Private/CADKernel/Core/Database.cpp
--- a/Engine/Source/Runtime/Datasmith/CADKernel/Private/CADKernel/Core/Database.cpp
+++ b/Engine/Source/Runtime/Datasmith/CADKernel/Private/CADKernel/Core/Database.cpp
@@ -118,23 +118,40 @@ void FDatabase::GetEntities(const TArray<FIdent>& EntityIds, TArray<TSharedPtr<F
 {
 	Entities.Empty(EntityIds.Num());
 
-	// to avoid duplicated components
-	TSet<FIdent> AddedComponents;
-	AddedComponents.Reserve(EntityIds.Num());
-
 	for (FIdent EntityId : EntityIds)
 	{
 		TSharedPtr<FEntity> Entity = GetEntity(EntityId);
-		if (!Entity.IsValid() || Entity->IsDeleted())
+		if (!Entity.IsValid() || Entity->IsDeleted() || Entity->IsProcessed())
 		{
 			continue;
 		}
+		Entity->SetProcessedMarker();
+		Entities.Add(Entity);
+	}
+	for (TSharedPtr<FEntity>& Entity : Entities)
+	{
+		Entity->ResetProcessedMarker();
+	}
+}
 
-		if (AddedComponents.Find(EntityId) == nullptr)
+void FDatabase::GetEntities(const TArray<FIdent>& EntityIds, TArray<FEntity*>& Entities) const
+{
+	Entities.Empty(EntityIds.Num());
+
+	for (FIdent EntityId : EntityIds)
+	{
+		TSharedPtr<FEntity> Entity = GetEntity(EntityId);
+		if (!Entity.IsValid() || Entity->IsDeleted() || Entity->IsProcessed())
 		{
-			AddedComponents.Add(EntityId);
-			Entities.Add(Entity);
+			continue;
 		}
+
+		Entity->SetProcessedMarker();
+		Entities.Add(Entity.Get());
+	}
+	for (FEntity* Entity : Entities)
+	{
+		Entity->ResetProcessedMarker();
 	}
 }
 
@@ -142,40 +159,24 @@ void FDatabase::GetTopologicalEntities(const TArray<FIdent>& EntityIds, TArray<F
 {
 	Entities.Empty(EntityIds.Num());
 
-	// to avoid duplicated components
-	TSet<FIdent> AddedComponents;
-	AddedComponents.Reserve(EntityIds.Num());
-
 	for (FIdent EntityId : EntityIds)
 	{
 		TSharedPtr<FEntity> Entity = GetEntity(EntityId);
-		if (!Entity.IsValid() || Entity->IsDeleted())
+		if (!Entity.IsValid() || Entity->IsDeleted() || Entity->IsProcessed())
 		{
 			continue;
 		}
 
-		switch (Entity->GetEntityType())
+		if(Entity->IsTopologicalEntity())
 		{
-		case EEntity::EdgeLink:
-		case EEntity::VertexLink:
-		case EEntity::TopologicalEdge:
-		case EEntity::TopologicalFace:
-		case EEntity::TopologicalLink:
-		case EEntity::TopologicalLoop:
-		case EEntity::TopologicalVertex:
-		case EEntity::Shell:
-		case EEntity::Body:
-		case EEntity::Model:
-			if (AddedComponents.Find(EntityId) == nullptr)
-			{
-				AddedComponents.Add(EntityId);
-				Entities.Add((FTopologicalShapeEntity*)Entity.Get());
-			}
-			break;
-		default:
-			break;
+			Entity->SetProcessedMarker();
+			Entities.Add((FTopologicalEntity*) Entity.Get());
 		}
 	}
+	for (FTopologicalEntity* Entity : Entities)
+	{
+		Entity->ResetProcessedMarker();
+	}
 }
 
 void FDatabase::GetTopologicalShapeEntities(const TArray<FIdent>& EntityIds, TArray<FTopologicalShapeEntity*>& Entities) const
@@ -194,22 +195,16 @@ void FDatabase::GetTopologicalShapeEntities(const TArray<FIdent>& EntityIds, TAr
 			continue;
 		}
 
-		switch (Entity->GetEntityType())
+		if(Entity->IsTopologicalShapeEntity())
 		{
-		case EEntity::Body:
-		case EEntity::Model:
-		case EEntity::Shell:
-		case EEntity::TopologicalFace:
-			if (AddedComponents.Find(EntityId) == nullptr)
-			{
-				AddedComponents.Add(EntityId);
-				Entities.Add((FTopologicalShapeEntity*) Entity.Get());
-			}
-			break;
-		default:
-			break;
+			Entity->SetProcessedMarker();
+			Entities.Add((FTopologicalShapeEntity*) Entity.Get());
 		}
 	}
+	for (FTopologicalShapeEntity* Entity : Entities)
+	{
+		Entity->ResetProcessedMarker();
+	}
 }
 
 void FCADKernelArchive::SetReferencedEntityOrAddToWaitingList(FIdent ArchiveId, FEntity** Entity)
@@ -326,7 +321,7 @@ void FDatabase::Serialize(FCADKernelArchive& Ar)
 
 		if (!Entity.IsValid() || Entity->IsDeleted())
 		{
-			EEntity Type = EEntity::NullEntity;
+			EEntity Type = EEntity::None;
 			Ar << Type;
 			continue;
 		}
@@ -600,23 +595,12 @@ void FDatabase::ExpandSelection(TSharedPtr<FEntity> Entity, const TSet<EEntity>&
 
 	case EEntity::TopologicalEdge:
 	{
-		ensureCADKernel(false);
+		break;
 		TSharedPtr<FTopologicalEdge>Edge = StaticCastSharedPtr<FTopologicalEdge>(Entity);
 		ExpandSelection(Edge->GetCurve(), Filter, Selection);
 		break;
 	}
 
-	case EEntity::Group:
-	{
-		ensureCADKernel(false);
-		TSharedPtr<FGroup>Group = StaticCastSharedPtr<FGroup>(Entity);
-		for (const TSharedPtr<FEntity>& GroupEntity : Group->GetEntities())
-		{
-			ExpandSelection(GroupEntity, Filter, Selection);
-		}
-		break;
-	}
-
 	case EEntity::Model:
 	{
 		TSharedPtr<FModel> InModel = StaticCastSharedPtr<FModel>(Entity);
@@ -643,11 +627,12 @@ void FDatabase::ExpandSelection(TSharedPtr<FEntity> Entity, const TSet<EEntity>&
 
 	case EEntity::MeshModel:
 	{
-		TSharedPtr<FModelMesh> MeshModel = StaticCastSharedPtr<FModelMesh>(Entity);
-		for (const TSharedPtr<FMesh>& Mesh : MeshModel->GetMeshes())
-		{
-			ExpandSelection(Mesh, Filter, Selection);
-		}
+		ensureCADKernel(false);
+		//TSharedPtr<FModelMesh> MeshModel = StaticCastSharedPtr<FModelMesh>(Entity);
+		//for (const FMesh* Mesh : MeshModel->GetMeshes())
+		//{
+		//	ExpandSelection(Mesh, Filter, Selection);
+		//}
 		break;
 	}
 
diff --git a/Engine/Source/Runtime/Datasmith/CADKernel/Private/CADKernel/Core/Entity.cpp b/Engine/Source/Runtime/Datasmith/CADKernel/Private/CADKernel/Core/Entity.cpp
--- a/Engine/Source/Runtime/Datasmith/CADKernel/Private/CADKernel/Core/Entity.cpp
+++ b/Engine/Source/Runtime/Datasmith/CADKernel/Private/CADKernel/Core/Entity.cpp
@@ -5,7 +5,6 @@
 #include "CADKernel/Core/EntityGeom.h"
 
 #include "CADKernel/Core/Database.h"
-#include "CADKernel/Core/Group.h"
 #include "CADKernel/Core/Session.h"
 #include "CADKernel/Core/System.h"
 #include "CADKernel/Geo/Curves/Curve.h"
@@ -22,9 +21,12 @@
 
 namespace UE::CADKernel
 {
+
+	const FPairOfIndex FPairOfIndex::Undefined(-1, -1);
+
 	const TCHAR* FEntity::TypesNames[] = 
 	{
-		TEXT("NullEntity"),
+		TEXT("UndefinedEntity"),
 		TEXT("Curve"),
 		TEXT("Surface"),
 
@@ -50,11 +52,11 @@ namespace UE::CADKernel
 
 	const TCHAR* FEntity::GetTypeName(EEntity Type)
 	{
-		if (Type >= EEntity::EntityTypeEnd)
+		if (Type > EEntity::None && Type < EEntity::EntityTypeEnd)
 		{
-			return TypesNames[0];
+			return TypesNames[(int32)Type];
 		}
-		return TypesNames[(int32) Type];
+		return TypesNames[(int32)EEntity::None];
 	}
 
 	FEntity::~FEntity()
@@ -68,7 +70,7 @@ namespace UE::CADKernel
 
 		ensureCADKernel(Archive.IsLoading());
 
-		EEntity Type = EEntity::NullEntity;
+		EEntity Type = EEntity::None;
 		Archive << Type;
 
 		TSharedPtr<FEntity> Entity;
@@ -86,9 +88,6 @@ namespace UE::CADKernel
 		case EEntity::EdgeLink:
 			Entity = FEntity::MakeShared<FEdgeLink>(Archive);
 			break;
-		case EEntity::Group:
-			Entity = FEntity::MakeShared<FGroup>(Archive);
-			break;
 		case EEntity::Model:
 			Entity = FEntity::MakeShared<FModel>(Archive);
 			break;
diff --git a/Engine/Source/Runtime/Datasmith/CADKernel/Private/CADKernel/Core/Group.cpp b/Engine/Source/Runtime/Datasmith/CADKernel/Private/CADKernel/Core/Group.cpp
--- a/Engine/Source/Runtime/Datasmith/CADKernel/Private/CADKernel/Core/Group.cpp
+++ b/Engine/Source/Runtime/Datasmith/CADKernel/Private/CADKernel/Core/Group.cpp
@@ -20,7 +20,7 @@ EEntity FGroup::GetGroupType() const
 {
 	if (Entities.Num() == 0) 
 	{
-		return EEntity::NullEntity;
+		return EEntity::None;
 	}
 
 	return (*Entities.begin())->GetEntityType();
diff --git a/Engine/Source/Runtime/Datasmith/CADKernel/Private/CADKernel/Core/Group.h b/Engine/Source/Runtime/Datasmith/CADKernel/Private/CADKernel/Core/Group.h
--- a/Engine/Source/Runtime/Datasmith/CADKernel/Private/CADKernel/Core/Group.h
+++ b/Engine/Source/Runtime/Datasmith/CADKernel/Private/CADKernel/Core/Group.h
@@ -56,7 +56,7 @@ public:
 		SpawnIdentOnEntities(Entities, Database);
 	}
 
-	virtual void ResetMarkersRecursively() override
+	virtual void ResetMarkersRecursively() const override
 	{
 		ResetMarkers();
 		ResetMarkersRecursivelyOnEntities(Entities);
@@ -68,7 +68,7 @@ public:
 
 	virtual EEntity GetEntityType() const override
 	{
-		return EEntity::Group;
+		return EEntity::None;
 	}
 
 	void SetName(const FString& Name);
diff --git a/Engine/Source/Runtime/Datasmith/CADKernel/Private/CADKernel/Core/System.cpp b/Engine/Source/Runtime/Datasmith/CADKernel/Private/CADKernel/Core/System.cpp
--- a/Engine/Source/Runtime/Datasmith/CADKernel/Private/CADKernel/Core/System.cpp
+++ b/Engine/Source/Runtime/Datasmith/CADKernel/Private/CADKernel/Core/System.cpp
@@ -10,6 +10,7 @@
 #include "Misc/Paths.h"
 
 #ifdef CADKERNEL_DEV
+#include "CADKernel/Core/KernelParameters.h"
 #include "CADKernel/Core/Version.h"
 
 #include <stdlib.h>
@@ -128,7 +129,10 @@ void FSystem::DefineReportFile(const FString& InReportFilePath)
 	ReportFile = MakeShareable<FArchive>(IFileManager::Get().CreateFileWriter(*InReportFilePath, IO_WRITE));
 
 	ReportHeaderPath = FPaths::ConvertRelativePathToFull(FPaths::Combine(ReportHeaderPath, TEXT("ReportHeader.csv")));
-	ReportHeaderFile = MakeShareable<FArchive>(IFileManager::Get().CreateFileWriter(*ReportHeaderPath, IO_WRITE));
+	if(!FPaths::FileExists(ReportHeaderPath))
+	{
+		ReportHeaderFile = MakeShareable<FArchive>(IFileManager::Get().CreateFileWriter(*ReportHeaderPath, IO_WRITE));
+	}
 }
 #endif
 
diff --git a/Engine/Source/Runtime/Datasmith/CADKernel/Private/CADKernel/Geo/Curves/BoundedCurve.h b/Engine/Source/Runtime/Datasmith/CADKernel/Private/CADKernel/Geo/Curves/BoundedCurve.h
--- a/Engine/Source/Runtime/Datasmith/CADKernel/Private/CADKernel/Geo/Curves/BoundedCurve.h
+++ b/Engine/Source/Runtime/Datasmith/CADKernel/Private/CADKernel/Geo/Curves/BoundedCurve.h
@@ -40,7 +40,7 @@ public:
 		Curve->SpawnIdent(Database);
 	}
 
-	virtual void ResetMarkersRecursively() override
+	virtual void ResetMarkersRecursively() const override
 	{
 		ResetMarkers();
 		Curve->ResetMarkersRecursively();
diff --git a/Engine/Source/Runtime/Datasmith/CADKernel/Private/CADKernel/Geo/Curves/CompositeCurve.h b/Engine/Source/Runtime/Datasmith/CADKernel/Private/CADKernel/Geo/Curves/CompositeCurve.h
--- a/Engine/Source/Runtime/Datasmith/CADKernel/Private/CADKernel/Geo/Curves/CompositeCurve.h
+++ b/Engine/Source/Runtime/Datasmith/CADKernel/Private/CADKernel/Geo/Curves/CompositeCurve.h
@@ -16,6 +16,11 @@ public:
 	{
 	}
 
+	FOrientedCurve(const FOrientedCurve& OrientiredEntity)
+		: TOrientedEntity(OrientiredEntity)
+	{
+	}
+
 	FOrientedCurve()
 		: TOrientedEntity()
 	{
@@ -53,7 +58,7 @@ public:
 		SpawnIdentOnEntities((TArray<TOrientedEntity<FEntity>>&) Curves, Database);
 	}
 
-	virtual void ResetMarkersRecursively() override
+	virtual void ResetMarkersRecursively() const override
 	{
 		ResetMarkers();
 		ResetMarkersRecursivelyOnEntities((TArray<TOrientedEntity<FEntity>>&) Curves);
diff --git a/Engine/Source/Runtime/Datasmith/CADKernel/Private/CADKernel/Geo/Curves/CurveMaker.cpp b/Engine/Source/Runtime/Datasmith/CADKernel/Private/CADKernel/Geo/Curves/CurveMaker.cpp
--- a/Engine/Source/Runtime/Datasmith/CADKernel/Private/CADKernel/Geo/Curves/CurveMaker.cpp
+++ b/Engine/Source/Runtime/Datasmith/CADKernel/Private/CADKernel/Geo/Curves/CurveMaker.cpp
@@ -1,18 +1,9 @@
 // Copyright Epic Games, Inc. All Rights Reserved.
 
 #include "CADKernel/Geo/Curves/BezierCurve.h"
-//#include "CADKernel/Geo/Curves/BoundedCurve.h"
-//#include "CADKernel/Geo/Curves/CompositeCurve.h"
 #include "CADKernel/Geo/Curves/Curve.h"
-//#include "CADKernel/Geo/Curves/EllipseCurve.h"
-//#include "CADKernel/Geo/Curves/HyperbolaCurve.h"
 #include "CADKernel/Geo/Curves/NURBSCurve.h"
-//#include "CADKernel/Geo/Curves/ParabolaCurve.h"
-//#include "CADKernel/Geo/Curves/PolylineCurve.h"
-//#include "CADKernel/Geo/Curves/RestrictionCurve.h"
-//#include "CADKernel/Geo/Curves/SegmentCurve.h"
 #include "CADKernel/Geo/Curves/SplineCurve.h"
-//#include "CADKernel/Geo/Curves/SurfacicCurve.h"
 
 namespace UE::CADKernel
 {
diff --git a/Engine/Source/Runtime/Datasmith/CADKernel/Private/CADKernel/Geo/Sampler/SamplerAbstract.h b/Engine/Source/Runtime/Datasmith/CADKernel/Private/CADKernel/Geo/Sampler/SamplerAbstract.h
--- a/Engine/Source/Runtime/Datasmith/CADKernel/Private/CADKernel/Geo/Sampler/SamplerAbstract.h
+++ b/Engine/Source/Runtime/Datasmith/CADKernel/Private/CADKernel/Geo/Sampler/SamplerAbstract.h
@@ -272,6 +272,12 @@ protected:
 	 */
 	void CompletesPolyline(int32 NeededPointNumber, int32 IntermediateIndex)
 	{
+		if (NeededPointNumber > 50)
+		{
+			FMessage::Printf(EVerboseLevel::Debug, TEXT("SamplerAbstract needed point number %d\n"), NeededPointNumber);
+			NeededPointNumber = 50;
+		}
+
 		const TArray<double>& CandidateCoordinates = CandidatePoints.GetCoordinates();
 
 		if (NeededPointNumber > 1)
diff --git a/Engine/Source/Runtime/Datasmith/CADKernel/Private/CADKernel/Geo/Surfaces/CompositeSurface.h b/Engine/Source/Runtime/Datasmith/CADKernel/Private/CADKernel/Geo/Surfaces/CompositeSurface.h
--- a/Engine/Source/Runtime/Datasmith/CADKernel/Private/CADKernel/Geo/Surfaces/CompositeSurface.h
+++ b/Engine/Source/Runtime/Datasmith/CADKernel/Private/CADKernel/Geo/Surfaces/CompositeSurface.h
@@ -53,7 +53,7 @@ public:
 		SpawnIdentOnEntities(Surfaces, Database);
 	}
 
-	virtual void ResetMarkersRecursively() override
+	virtual void ResetMarkersRecursively() const override
 	{
 		ResetMarkers();
 		ResetMarkersRecursivelyOnEntities(Surfaces);
diff --git a/Engine/Source/Runtime/Datasmith/CADKernel/Private/CADKernel/Geo/Surfaces/RevolutionSurface.cpp b/Engine/Source/Runtime/Datasmith/CADKernel/Private/CADKernel/Geo/Surfaces/RevolutionSurface.cpp
--- a/Engine/Source/Runtime/Datasmith/CADKernel/Private/CADKernel/Geo/Surfaces/RevolutionSurface.cpp
+++ b/Engine/Source/Runtime/Datasmith/CADKernel/Private/CADKernel/Geo/Surfaces/RevolutionSurface.cpp
@@ -149,7 +149,7 @@ void FRevolutionSurface::SpawnIdent(FDatabase& Database)
 	Generatrix->SpawnIdent(Database);
 }
 
-void FRevolutionSurface::ResetMarkersRecursively()
+void FRevolutionSurface::ResetMarkersRecursively() const
 {
 	ResetMarkers();
 	Axis->ResetMarkersRecursively();
diff --git a/Engine/Source/Runtime/Datasmith/CADKernel/Private/CADKernel/Geo/Surfaces/RevolutionSurface.h b/Engine/Source/Runtime/Datasmith/CADKernel/Private/CADKernel/Geo/Surfaces/RevolutionSurface.h
--- a/Engine/Source/Runtime/Datasmith/CADKernel/Private/CADKernel/Geo/Surfaces/RevolutionSurface.h
+++ b/Engine/Source/Runtime/Datasmith/CADKernel/Private/CADKernel/Geo/Surfaces/RevolutionSurface.h
@@ -45,7 +45,7 @@ public:
 	}
 
 	virtual void SpawnIdent(FDatabase& Database) override;
-	virtual void ResetMarkersRecursively() override;
+	virtual void ResetMarkersRecursively() const override;
 
 	ESurface GetSurfaceType() const
 	{
diff --git a/Engine/Source/Runtime/Datasmith/CADKernel/Private/CADKernel/Geo/Surfaces/RuledSurface.cpp b/Engine/Source/Runtime/Datasmith/CADKernel/Private/CADKernel/Geo/Surfaces/RuledSurface.cpp
--- a/Engine/Source/Runtime/Datasmith/CADKernel/Private/CADKernel/Geo/Surfaces/RuledSurface.cpp
+++ b/Engine/Source/Runtime/Datasmith/CADKernel/Private/CADKernel/Geo/Surfaces/RuledSurface.cpp
@@ -181,7 +181,7 @@ void FRuledSurface::SpawnIdent(FDatabase& Database)
 	Curves[1]->SpawnIdent(Database);
 }
 
-void FRuledSurface::ResetMarkersRecursively() 
+void FRuledSurface::ResetMarkersRecursively() const
 {
 	ResetMarkers();
 	Curves[0]->ResetMarkersRecursively();
diff --git a/Engine/Source/Runtime/Datasmith/CADKernel/Private/CADKernel/Geo/Surfaces/RuledSurface.h b/Engine/Source/Runtime/Datasmith/CADKernel/Private/CADKernel/Geo/Surfaces/RuledSurface.h
--- a/Engine/Source/Runtime/Datasmith/CADKernel/Private/CADKernel/Geo/Surfaces/RuledSurface.h
+++ b/Engine/Source/Runtime/Datasmith/CADKernel/Private/CADKernel/Geo/Surfaces/RuledSurface.h
@@ -35,7 +35,7 @@ public:
 	}
 
 	virtual void SpawnIdent(FDatabase& Database) override;
-	virtual void ResetMarkersRecursively() override;
+	virtual void ResetMarkersRecursively() const override;
 
 	ESurface GetSurfaceType() const
 	{
diff --git a/Engine/Source/Runtime/Datasmith/CADKernel/Private/CADKernel/Geo/Surfaces/TabulatedCylinderSurface.cpp b/Engine/Source/Runtime/Datasmith/CADKernel/Private/CADKernel/Geo/Surfaces/TabulatedCylinderSurface.cpp
--- a/Engine/Source/Runtime/Datasmith/CADKernel/Private/CADKernel/Geo/Surfaces/TabulatedCylinderSurface.cpp
+++ b/Engine/Source/Runtime/Datasmith/CADKernel/Private/CADKernel/Geo/Surfaces/TabulatedCylinderSurface.cpp
@@ -19,7 +19,7 @@ void FTabulatedCylinderSurface::SpawnIdent(FDatabase& Database)
 	GuideCurve->SpawnIdent(Database);
 }
 
-void FTabulatedCylinderSurface::ResetMarkersRecursively() 
+void FTabulatedCylinderSurface::ResetMarkersRecursively() const
 {
 	ResetMarkers();
 	GuideCurve->ResetMarkersRecursively();
diff --git a/Engine/Source/Runtime/Datasmith/CADKernel/Private/CADKernel/Math/Geometry.cpp b/Engine/Source/Runtime/Datasmith/CADKernel/Private/CADKernel/Math/Geometry.cpp
--- a/Engine/Source/Runtime/Datasmith/CADKernel/Private/CADKernel/Math/Geometry.cpp
+++ b/Engine/Source/Runtime/Datasmith/CADKernel/Private/CADKernel/Math/Geometry.cpp
@@ -65,7 +65,71 @@ void FindLoopIntersectionsWithIso(const EIso Iso, const double IsoParameter, con
 	Algo::Sort(OutIntersections);
 }
 
-bool IntersectSegments2D(const TSegment<FPoint2D>& SegmentAB, const TSegment<FPoint2D>& SegmentCD)
+bool DoIntersectInside(const FSegment2D& SegmentAB, const FSegment2D& SegmentCD)
+{
+	constexpr const double Min = DOUBLE_SMALL_NUMBER;
+	constexpr const double Max = 1. - DOUBLE_SMALL_NUMBER;
+
+	TFunction<bool(double, double, double, double)> Intersect = [](double A, double B, double C, double D)
+	{
+		return !((D < A + DOUBLE_SMALL_NUMBER) || (B < C + DOUBLE_SMALL_NUMBER));
+	};
+
+	TFunction<bool(double, double, double, double)> TestWhenParallel = [&](double A, double B, double C, double D)
+	{
+		if (A < B)
+		{
+			if (C < D)
+			{
+				return Intersect(A, B, C, D);
+			}
+			else
+			{
+				return Intersect(A, B, D, C);
+			}
+		}
+		else
+		{
+			if (C < D)
+			{
+				return Intersect(B, A, C, D);
+			}
+			else
+			{
+				return Intersect(B, A, D, C);
+			}
+		}
+	};
+
+	const FPoint2D AB = SegmentAB[1] - SegmentAB[0];
+	const FPoint2D CD = SegmentCD[1] - SegmentCD[0];
+	const FPoint2D CA = SegmentAB[0] - SegmentCD[0];
+
+	const double ParallelCoef = CD ^ AB;
+	if (FMath::IsNearlyZero(ParallelCoef))
+	{
+		const double ParallelCoef2 = CA ^ AB;
+		if (!FMath::IsNearlyZero(ParallelCoef2))
+		{
+			return false;
+		}
+
+		if (fabs(AB.U) > fabs(AB.V))
+		{
+			return TestWhenParallel(SegmentAB[0].U, SegmentAB[1].U, SegmentCD[0].U, SegmentCD[1].U);
+		}
+		else
+		{
+			return TestWhenParallel(SegmentAB[0].V, SegmentAB[1].V, SegmentCD[0].V, SegmentCD[1].V);
+		}
+	}
+
+	const double ABIntersectionCoordinate = (CA ^ CD) / ParallelCoef;
+	const double CDIntersectionCoordinate = (CA ^ AB) / ParallelCoef;
+	return (ABIntersectionCoordinate < Max && ABIntersectionCoordinate > Min && CDIntersectionCoordinate < Max && CDIntersectionCoordinate > Min);
+}
+
+bool DoIntersect(const TSegment<FPoint2D>& SegmentAB, const TSegment<FPoint2D>& SegmentCD)
 {
 	constexpr const double Min = -DOUBLE_KINDA_SMALL_NUMBER;
 	constexpr const double Max = 1. + DOUBLE_KINDA_SMALL_NUMBER;
diff --git a/Engine/Source/Runtime/Datasmith/CADKernel/Private/CADKernel/Mesh/Criteria/CriteriaGrid.cpp b/Engine/Source/Runtime/Datasmith/CADKernel/Private/CADKernel/Mesh/Criteria/CriteriaGrid.cpp
--- a/Engine/Source/Runtime/Datasmith/CADKernel/Private/CADKernel/Mesh/Criteria/CriteriaGrid.cpp
+++ b/Engine/Source/Runtime/Datasmith/CADKernel/Private/CADKernel/Mesh/Criteria/CriteriaGrid.cpp
@@ -43,10 +43,80 @@ FCriteriaGrid::FCriteriaGrid(FTopologicalFace& InFace)
 	, CoordinateGrid(InFace.GetCrossingPointCoordinates())
 {
 	Init();
+	FaceMinMax.Init();
+
 #ifdef DISPLAY_CRITERIA_GRID
 	Display();
 #endif
 }
 
+void FCriteriaGrid::ComputeFaceMinMaxThicknessAlongIso()
+{
+	using FGetPointFunc = TFunction<const FPoint& (int32, int32)>;
+
+	TFunction<void(EIso, FGetPointFunc, FGetPointFunc, int32, int32)> ComputeFaceThicknessAlong = [this](EIso Iso, FGetPointFunc GetPointAlong, FGetPointFunc GetIntermediatePointAlong, int32 IsoUCount, int32 IsoVCount)
+	{
+		for (int32 VIndex = 0; VIndex < IsoVCount; ++VIndex)
+		{
+			const FPoint& StartPoint = GetPointAlong(0, VIndex);
+			const FPoint* PreviousPoint = &GetIntermediatePointAlong(0, VIndex);
+			double Length = StartPoint.Distance(*PreviousPoint);
+
+			for (int32 UIndex = 1; UIndex < IsoUCount - 1; ++UIndex)
+			{
+				const FPoint& Point = GetPointAlong(UIndex, VIndex);
+
+				PreviousPoint = &GetIntermediatePointAlong(UIndex, VIndex);
+				Length += Point.Distance(*PreviousPoint);
+			}
+			const FPoint& LastPoint = GetPointAlong(IsoUCount - 1, VIndex);
+			Length += LastPoint.Distance(*PreviousPoint);
+
+			this->FaceMinMax[Iso].ExtendTo(Length);
+		}
+	};
+
+	FGetPointFunc GetPointAlongU = [this](int32 UIndex, int32 VIndex) -> const FPoint&
+	{
+		return GetPoint(UIndex, VIndex);
+	};
+
+	FGetPointFunc GetPointAlongV = [this](int32 VIndex, int32 UIndex) -> const FPoint&
+	{
+		return GetPoint(UIndex, VIndex);
+	};
+
+	FGetPointFunc GetIntermediatePointAlongU = [this](int32 UIndex, int32 VIndex) -> const FPoint&
+	{
+		return GetIntermediateU(UIndex, VIndex);
+	};
+
+	FGetPointFunc GetIntermediatePointAlongV = [this](int32 VIndex, int32 UIndex) -> const FPoint&
+	{
+		return GetIntermediateV(UIndex, VIndex);
+	};
+
+	TFunction<int32(EIso)> GetCoordinateCount = [&](EIso Iso) -> int32
+	{
+		return Face.GetCrossingPointCoordinates().IsoCount(Iso);
+	};
+
+	const int32 IsoUCount = GetCoordinateCount(EIso::IsoU);
+	const int32 IsoVCount = GetCoordinateCount(EIso::IsoV);
+
+	ComputeFaceThicknessAlong(EIso::IsoU, GetPointAlongU, GetIntermediatePointAlongU, IsoUCount, IsoVCount);
+	ComputeFaceThicknessAlong(EIso::IsoV, GetPointAlongV, GetIntermediatePointAlongV, IsoVCount, IsoUCount);
+
+	Face.SetEstimatedMinimalElementLength(FMath::Min(FaceMinMax[EIso::IsoU].GetMax(), FaceMinMax[EIso::IsoV].GetMax()));
+}
+
+bool FCriteriaGrid::CheckIfIsDegenerate()
+{
+	ensureCADKernel(FaceMinMax.IsValid());
+	const double MinFaceThickness = Tolerance3D * 3.;
+	return (FaceMinMax[EIso::IsoU].Max < MinFaceThickness) || (FaceMinMax[EIso::IsoV].Max < MinFaceThickness);
+}
+
+
 } // namespace UE::CADKernel
 
diff --git a/Engine/Source/Runtime/Datasmith/CADKernel/Private/CADKernel/Mesh/Criteria/CriteriaGrid.h b/Engine/Source/Runtime/Datasmith/CADKernel/Private/CADKernel/Mesh/Criteria/CriteriaGrid.h
--- a/Engine/Source/Runtime/Datasmith/CADKernel/Private/CADKernel/Mesh/Criteria/CriteriaGrid.h
+++ b/Engine/Source/Runtime/Datasmith/CADKernel/Private/CADKernel/Mesh/Criteria/CriteriaGrid.h
@@ -2,6 +2,7 @@
 #pragma once
 
 #include "CADKernel/Core/Types.h"
+#include "CADKernel/Math/Boundary.h"
 #include "CADKernel/Mesh/Structure/GridBase.h"
 
 #include "CADKernel/Geo/Sampling/SurfacicSampling.h"
@@ -20,6 +21,8 @@ private:
 	 */
 	FCoordinateGrid CoordinateGrid;
 
+	FSurfacicBoundary FaceMinMax;
+
 protected:
 
 	virtual const FCoordinateGrid& GetCoordinateGrid() const override
@@ -81,6 +84,18 @@ public:
 		return GetPoint(iU, iV, true, true);
 	}
 
+	void ComputeFaceMinMaxThicknessAlongIso();
+
+	double GetCharacteristicThicknessOfFace() const
+	{
+		return FMath::Max(FaceMinMax[EIso::IsoU].GetMax(), FaceMinMax[EIso::IsoV].GetMax());
+	}
+
+	/**
+	 * @return true if all the iso lines are smaller than geometric tolerance
+	 */
+	bool CheckIfIsDegenerate();
+
 #ifdef CADKERNEL_DEV
 	void Display() const;
 #endif
diff --git a/Engine/Source/Runtime/Datasmith/CADKernel/Private/CADKernel/Mesh/Meshers/BowyerWatsonTriangulator.h b/Engine/Source/Runtime/Datasmith/CADKernel/Private/CADKernel/Mesh/Meshers/BowyerWatsonTriangulator.h
--- a/Engine/Source/Runtime/Datasmith/CADKernel/Private/CADKernel/Mesh/Meshers/BowyerWatsonTriangulator.h
+++ b/Engine/Source/Runtime/Datasmith/CADKernel/Private/CADKernel/Mesh/Meshers/BowyerWatsonTriangulator.h
@@ -326,7 +326,7 @@ public:
 						}
 #endif
 
-						if (SlopeStartCandidate > 0 && SlopeStartCandidate < Slope::RightSlope && SlopeEndCandidate > Slope::ThreeRightSlope && SlopeEndCandidate < Slope::FullSlope)
+						if (SlopeStartCandidate > 0 && SlopeStartCandidate < Slope::RightSlope && SlopeEndCandidate > Slope::ThreeRightSlope && SlopeEndCandidate < Slope::TwoPiSlope)
 						{
 							bTriangleHasBeenAdded = true;
 							if (EndIndex == 0 && NewVertexSlope > BoundaryVertexToSlope.Last().Value)
