Switch mesher in case of failure

diff --git a/Engine/Plugins/Enterprise/DatasmithCADImporter/Source/CADInterfaces/Private/CADFileReader.cpp b/Engine/Plugins/Enterprise/DatasmithCADImporter/Source/CADInterfaces/Private/CADFileReader.cpp
--- a/Engine/Plugins/Enterprise/DatasmithCADImporter/Source/CADInterfaces/Private/CADFileReader.cpp
+++ b/Engine/Plugins/Enterprise/DatasmithCADImporter/Source/CADInterfaces/Private/CADFileReader.cpp
@@ -18,13 +18,13 @@ namespace CADLibrary
 		: CADFileData(ImportParams, InFile, InCachePath)
 	{
 #ifdef USE_TECHSOFT_SDK
-		if (FImportParameters::bGDisableCADKernelTessellation)
+		if (ImportParams.GetMesher() == EMesher::CADKernel)
 		{
-			CADParser = MakeUnique<FTechSoftFileParser>(CADFileData, EnginePluginsPath);
+			CADParser = MakeUnique<FTechSoftFileParserCADKernelTessellator>(CADFileData, EnginePluginsPath);
 		}
 		else
 		{
-			CADParser = MakeUnique<FTechSoftFileParserCADKernelTessellator>(CADFileData, EnginePluginsPath);
+			CADParser = MakeUnique<FTechSoftFileParser>(CADFileData, EnginePluginsPath);
 		}
 #endif
 	}
diff --git a/Engine/Plugins/Enterprise/DatasmithCADImporter/Source/CADInterfaces/Private/CADSceneGraph.cpp b/Engine/Plugins/Enterprise/DatasmithCADImporter/Source/CADInterfaces/Private/CADSceneGraph.cpp
--- a/Engine/Plugins/Enterprise/DatasmithCADImporter/Source/CADInterfaces/Private/CADSceneGraph.cpp
+++ b/Engine/Plugins/Enterprise/DatasmithCADImporter/Source/CADInterfaces/Private/CADSceneGraph.cpp
@@ -61,6 +61,7 @@ FArchive& operator<<(FArchive& Ar, FArchiveBody& Body)
 	Ar << (FArchiveCADObject&) Body;
 	Ar << Body.MaterialFaceSet;
 	Ar << Body.ColorFaceSet;
+	Ar << Body.Mesher;
 	Ar << Body.ParentId;
 	Ar << Body.MeshActorUId;
 
@@ -341,12 +342,12 @@ void FArchiveSceneGraph::RemoveLastUnloadedReference()
 	CADIdToIndex.SetNum(CADIdToIndex.Num() - 1, false);
 }
 
-FArchiveBody& FArchiveSceneGraph::AddBody(FArchiveReference& Parent)
+FArchiveBody& FArchiveSceneGraph::AddBody(FArchiveReference& Parent, EMesher InMesher)
 {
 	ensure(Bodies.Num() < Bodies.Max());
 
 	const int32 BodyId = LastEntityId++;
-	int32 Index = Bodies.Emplace(BodyId, Parent);
+	int32 Index = Bodies.Emplace(BodyId, Parent, InMesher);
 	CADIdToIndex.Add(Index);
 
 	return Bodies[Index];
diff --git a/Engine/Plugins/Enterprise/DatasmithCADImporter/Source/CADInterfaces/Private/TechSoftBridge.cpp b/Engine/Plugins/Enterprise/DatasmithCADImporter/Source/CADInterfaces/Private/TechSoftBridge.cpp
--- a/Engine/Plugins/Enterprise/DatasmithCADImporter/Source/CADInterfaces/Private/TechSoftBridge.cpp
+++ b/Engine/Plugins/Enterprise/DatasmithCADImporter/Source/CADInterfaces/Private/TechSoftBridge.cpp
@@ -246,7 +246,7 @@ UE::CADKernel::FBody* FTechSoftBridge::AddBody(A3DRiBrepModel* A3DBRepModel, TMa
 		TraverseBrepData(BRepModelData->m_pBrepData, Body);
 	}
 
-	if(Body->FaceCount() == 0)
+	if (Body->FaceCount() == 0 || bConvertionFailed)
 	{
 		Body->Delete();
 		return nullptr;
@@ -288,6 +288,10 @@ void FTechSoftBridge::TraverseBrepData(const A3DTopoBrepData* A3DBrepData, TShar
 		for (A3DUns32 Index = 0; Index < TopoBrepData->m_uiConnexSize; ++Index)
 		{
 			TraverseConnex(TopoBrepData->m_ppConnexes[Index], Body);
+			if (bConvertionFailed)
+			{
+				return;
+			}
 		}
 	}
 }
@@ -303,6 +307,10 @@ void FTechSoftBridge::TraverseConnex(const A3DTopoConnex* A3DTopoConnex, TShared
 		for (A3DUns32 Index = 0; Index < TopoConnexData->m_uiShellSize; ++Index)
 		{
 			TraverseShell(TopoConnexData->m_ppShells[Index], Body);
+			if (bConvertionFailed)
+			{
+				return;
+			}
 		}
 	}
 }
@@ -327,6 +335,10 @@ void FTechSoftBridge::TraverseShell(const A3DTopoShell* A3DShell, TSharedRef<UE:
 		for (A3DUns32 Index = 0; Index < ShellData->m_uiFaceSize; ++Index)
 		{
 			AddFace(ShellData->m_ppFaces[Index], ShellData->m_pucOrientationWithShell[Index] == 1 ? UE::CADKernel::Front : UE::CADKernel::Back, Shell, Index);
+			if (bConvertionFailed)
+			{
+				return;
+			}
 		}
 	}
 }
@@ -632,6 +644,8 @@ TSharedPtr<UE::CADKernel::FTopologicalEdge> FTechSoftBridge::AddEdge(const A3DTo
 
 	if (CoEdgeData->m_pUVCurve == nullptr)
 	{
+		bConvertionFailed = true;
+		FailureReason = EFailureReason::Curve3D;
 		return TSharedPtr<UE::CADKernel::FTopologicalEdge>();
 	}
 
diff --git a/Engine/Plugins/Enterprise/DatasmithCADImporter/Source/CADInterfaces/Private/TechSoftBridge.h b/Engine/Plugins/Enterprise/DatasmithCADImporter/Source/CADInterfaces/Private/TechSoftBridge.h
--- a/Engine/Plugins/Enterprise/DatasmithCADImporter/Source/CADInterfaces/Private/TechSoftBridge.h
+++ b/Engine/Plugins/Enterprise/DatasmithCADImporter/Source/CADInterfaces/Private/TechSoftBridge.h
@@ -6,6 +6,7 @@
 
 #include "CADKernel/Geo/GeoEnum.h"
 #include "CADKernel/Math/MatrixH.h"
+#include "CADOptions.h"
 
 #include "TechSoftInterface.h"
 
@@ -175,6 +176,9 @@ private:
 	TMap<UE::CADKernel::FBody*, const A3DEntity*> CADKernelToTechSoft;
 	TMap<const A3DTopoCoEdge*, TSharedPtr<UE::CADKernel::FTopologicalEdge>> A3DEdgeToEdge;
 
+	EFailureReason FailureReason = EFailureReason::None;
+	bool bConvertionFailed = false;
+
 	double BodyScale = 1;
 
 public:
@@ -183,6 +187,11 @@ public:
 	UE::CADKernel::FBody* GetBody(A3DRiBrepModel* A3DBRepModel);
 	const A3DRiBrepModel* GetA3DBody(UE::CADKernel::FBody* BRepModel);
 
+	EFailureReason GetFailureReason()
+	{
+		return FailureReason;
+	}
+
 private:
 
 	void TraverseBrepData(const A3DTopoBrepData* A3DBrepData, TSharedRef<UE::CADKernel::FBody>& Body);
diff --git a/Engine/Plugins/Enterprise/DatasmithCADImporter/Source/CADInterfaces/Private/TechSoftFileParser.cpp b/Engine/Plugins/Enterprise/DatasmithCADImporter/Source/CADInterfaces/Private/TechSoftFileParser.cpp
--- a/Engine/Plugins/Enterprise/DatasmithCADImporter/Source/CADInterfaces/Private/TechSoftFileParser.cpp
+++ b/Engine/Plugins/Enterprise/DatasmithCADImporter/Source/CADInterfaces/Private/TechSoftFileParser.cpp
@@ -405,7 +405,7 @@ ECADParsingResult FTechSoftFileParser::Process()
 			return ECADParsingResult::ProcessFailed;
 		}
 
-		ModellerType = (EModellerType)ModelFileData->m_eModellerType;
+		ModelerType = (EModelerType) ModelFileData->m_eModellerType;
 		FileUnit = TechSoftInterface::GetModelFileUnit(ModelFile.Get());
 	}
 
@@ -451,6 +451,10 @@ ECADParsingResult FTechSoftFileParser::Process()
 	if (Result == ECADParsingResult::ProcessOk)
 	{
 		GenerateBodyMeshes();
+		if (bConvertionFailed)
+		{
+			Result = ECADParsingResult::ProcessFailed;
+		}
 
 		FString TechSoftVersion = TechSoftInterface::GetTechSoftVersion();
 		if (!TechSoftVersion.IsEmpty())
@@ -490,6 +494,10 @@ void FTechSoftFileParser::GenerateBodyMeshes()
 		A3DRiRepresentationItem* RepresentationItemPtr = Entry.Key;
 		FArchiveBody& Body = SceneGraph.GetBody(Entry.Value);
 		GenerateBodyMesh(RepresentationItemPtr, Body);
+		if (bConvertionFailed)
+		{
+			return;
+		}
 	}
 	CADFileData.GetRecord().MeshTime = FPlatformTime::ToMilliseconds64(FPlatformTime::Cycles64() - StartGenerateBodyMeshesTime);
 }
@@ -1289,7 +1297,7 @@ void FTechSoftFileParser::CountUnderRepresentationSet(const A3DRiSet* Representa
 
 void FTechSoftFileParser::TraverseBRepModel(A3DRiBrepModel* BRepModelPtr, FArchiveReference& Parent)
 {
-	FArchiveBody& BRep = SceneGraph.AddBody(Parent);
+	FArchiveBody& BRep = SceneGraph.AddBody(Parent, CADFileData.GetImportParameters().GetMesher());
 	ExtractMetaData(BRepModelPtr, BRep);
 	BRep.DefineGraphicsPropertiesFromNoOverwrite(Parent);
 
@@ -1314,7 +1322,7 @@ void FTechSoftFileParser::TraverseBRepModel(A3DRiBrepModel* BRepModelPtr, FArchi
 
 void FTechSoftFileParser::TraversePolyBRepModel(A3DRiPolyBrepModel* PolygonalPtr, FArchiveReference& Parent)
 {
-	FArchiveBody& BRep = SceneGraph.AddBody(Parent);
+	FArchiveBody& BRep = SceneGraph.AddBody(Parent, EMesher::TechSoft);
 	BRep.bIsFromCad = false;
 
 	ExtractMetaData(PolygonalPtr, BRep);
@@ -1523,9 +1531,9 @@ void FTechSoftFileParser::BuildRepresentationSetName(FArchiveCADObject& Occurren
 void FTechSoftFileParser::ExtractSpecificMetaData(const A3DAsmProductOccurrence* Occurrence, FArchiveCADObject& OutMetaData)
 {
 	//----------- Export Specific information per CAD format -----------
-	switch (ModellerType)
+	switch (ModelerType)
 	{
-	case ModellerSlw:
+	case ModelerSlw:
 	{
 		TUniqueTSObj<A3DAsmProductOccurrenceDataSLW> SolidWorksSpecificData(Occurrence);
 		if (SolidWorksSpecificData.IsValid())
@@ -1540,7 +1548,7 @@ void FTechSoftFileParser::ExtractSpecificMetaData(const A3DAsmProductOccurrence*
 		}
 		break;
 	}
-	case ModellerUnigraphics:
+	case ModelerUnigraphics:
 	{
 #ifdef WIP
 		TUniqueTSObj<A3DAsmProductOccurrenceDataUg> UnigraphicsSpecificData(Occurrence);
@@ -1610,7 +1618,7 @@ void FTechSoftFileParser::ExtractSpecificMetaData(const A3DAsmProductOccurrence*
 		break;
 	}
 
-	case ModellerCatiaV5:
+	case ModelerCatiaV5:
 	{
 		TUniqueTSObj<A3DAsmProductOccurrenceDataCV5> CatiaV5SpecificData(Occurrence);
 		if (CatiaV5SpecificData.IsValid())
diff --git a/Engine/Plugins/Enterprise/DatasmithCADImporter/Source/CADInterfaces/Private/TechSoftFileParser.h b/Engine/Plugins/Enterprise/DatasmithCADImporter/Source/CADInterfaces/Private/TechSoftFileParser.h
--- a/Engine/Plugins/Enterprise/DatasmithCADImporter/Source/CADInterfaces/Private/TechSoftFileParser.h
+++ b/Engine/Plugins/Enterprise/DatasmithCADImporter/Source/CADInterfaces/Private/TechSoftFileParser.h
@@ -164,12 +164,14 @@ protected:
 	ECADFormat Format;
 	bool bForceSew = false;
 
-	EModellerType ModellerType;
+	EModelerType ModelerType;
 	double FileUnit = 1;
 
 	TMap<A3DRiRepresentationItem*, FCadId> RepresentationItemsCache;
 	TMap<A3DAsmProductOccurrence*, FCadId> ReferenceCache;
 
+	bool bConvertionFailed = false;
+
 	std::atomic<bool> bProcessIsRunning;
 
 };
diff --git a/Engine/Plugins/Enterprise/DatasmithCADImporter/Source/CADInterfaces/Private/TechSoftFileParserCADKernelTessellator.cpp b/Engine/Plugins/Enterprise/DatasmithCADImporter/Source/CADInterfaces/Private/TechSoftFileParserCADKernelTessellator.cpp
--- a/Engine/Plugins/Enterprise/DatasmithCADImporter/Source/CADInterfaces/Private/TechSoftFileParserCADKernelTessellator.cpp
+++ b/Engine/Plugins/Enterprise/DatasmithCADImporter/Source/CADInterfaces/Private/TechSoftFileParserCADKernelTessellator.cpp
@@ -72,6 +72,10 @@ void FTechSoftFileParserCADKernelTessellator::SewAndGenerateBodyMeshes()
 	for (TPair<FCadId, TArray<A3DRiRepresentationItem*>>& Representations : OccurenceIdToRepresentations)
 	{
 		SewAndMesh(Representations.Value);
+		if (bConvertionFailed)
+		{
+			return;
+		}
 	}
 
 	// Delete unused ArchiveBody
@@ -108,6 +112,28 @@ void FTechSoftFileParserCADKernelTessellator::SewAndMesh(TArray<A3DRiRepresentat
 		FArchiveBody& Body = SceneGraph.GetBody(*BodyIndex);
 
 		FBody* CADKernelBody = TechSoftBridge.AddBody(Representation, Body.MetaData, Body.Unit);
+		if (CADKernelBody == nullptr)
+		{
+			const EFailureReason FailureReason = TechSoftBridge.GetFailureReason();
+			switch(FailureReason)
+			{
+				case EFailureReason::Curve3D :
+				{
+					// The failure of the current body is due to an incompatibility of modeler.
+					// Reimport of the full file with another modeler
+					bConvertionFailed = true;
+					return;
+				}
+
+				case EFailureReason::Unknown:
+				default:
+				{
+					// The failure of the current body is not due to an incompatibility of modeler.
+					// A reimport of the full file with another modeler is not wanted
+					break;
+				}
+			}
+		}
 	}
 
 	// Sew if needed
@@ -230,7 +256,12 @@ void FTechSoftFileParserCADKernelTessellator::GenerateBodyMesh(A3DRiRepresentati
 	if (CADKernelBody == nullptr)
 	{
 		ArchiveBody.Delete();
-		return;
+		const EFailureReason FailureReason = TechSoftBridge.GetFailureReason();
+		if(FailureReason == EFailureReason::Curve3D)
+		{
+			bConvertionFailed = true;
+			return;
+		}
 	}
 
 	if (CADFileData.GetImportParameters().GetStitchingTechnique() == StitchingHeal)
diff --git a/Engine/Plugins/Enterprise/DatasmithCADImporter/Source/CADInterfaces/Private/TechSoftFileParserCADKernelTessellator.h b/Engine/Plugins/Enterprise/DatasmithCADImporter/Source/CADInterfaces/Private/TechSoftFileParserCADKernelTessellator.h
--- a/Engine/Plugins/Enterprise/DatasmithCADImporter/Source/CADInterfaces/Private/TechSoftFileParserCADKernelTessellator.h
+++ b/Engine/Plugins/Enterprise/DatasmithCADImporter/Source/CADInterfaces/Private/TechSoftFileParserCADKernelTessellator.h
@@ -4,6 +4,7 @@
 #include "CoreMinimal.h"
 
 #include "TechSoftFileParser.h"
+#include "CADOptions.h"
 
 namespace UE::CADKernel
 {
diff --git a/Engine/Plugins/Enterprise/DatasmithCADImporter/Source/CADInterfaces/Public/CADFileData.h b/Engine/Plugins/Enterprise/DatasmithCADImporter/Source/CADInterfaces/Public/CADFileData.h
--- a/Engine/Plugins/Enterprise/DatasmithCADImporter/Source/CADInterfaces/Public/CADFileData.h
+++ b/Engine/Plugins/Enterprise/DatasmithCADImporter/Source/CADInterfaces/Public/CADFileData.h
@@ -107,7 +107,7 @@ public:
 
 	const FString GetBodyCachePath(uint32 BodyHash) const
 	{
-		return CADLibrary::BuildCacheFilePath(*CachePath, TEXT("body"), BodyHash);
+		return CADLibrary::BuildCacheFilePath(*CachePath, TEXT("body"), BodyHash, ImportParameters.GetMesher());
 	}
 
 	/**
diff --git a/Engine/Plugins/Enterprise/DatasmithCADImporter/Source/CADInterfaces/Public/CADSceneGraph.h b/Engine/Plugins/Enterprise/DatasmithCADImporter/Source/CADInterfaces/Public/CADSceneGraph.h
--- a/Engine/Plugins/Enterprise/DatasmithCADImporter/Source/CADInterfaces/Public/CADSceneGraph.h
+++ b/Engine/Plugins/Enterprise/DatasmithCADImporter/Source/CADInterfaces/Public/CADSceneGraph.h
@@ -4,6 +4,7 @@
 #include "CoreMinimal.h"
 
 #include "CADData.h"
+#include "CADOptions.h"
 
 class FArchive;
 
@@ -112,9 +113,10 @@ class CADINTERFACES_API FArchiveBody : public FArchiveCADObject
 {
 public:
 	FArchiveBody() = default;
-	FArchiveBody(FCadId Id, const FArchiveCADObject& Parent)
+	FArchiveBody(FCadId Id, const FArchiveCADObject& Parent, EMesher InMesher)
 		: FArchiveCADObject(Id, Parent)
 		, ParentId(Parent.Id)
+		, Mesher(InMesher)
 	{
 	}
 
@@ -127,6 +129,8 @@ public:
 	TSet<FMaterialUId> MaterialFaceSet;
 	TSet<FMaterialUId> ColorFaceSet;
 
+	EMesher Mesher = EMesher::TechSoft;
+
 	void Delete();
 
 	// non serialized fields
@@ -225,7 +229,7 @@ public:
 	void RemoveLastUnloadedReference();
 	bool IsAUnloadedReference(FCadId CadId) const;
 
-	FArchiveBody& AddBody(FArchiveReference& Parent);
+	FArchiveBody& AddBody(FArchiveReference& Parent, EMesher InMesher);
 	FArchiveBody& GetBody(FCadId CadId);
 	void RemoveLastBody();
 	bool IsABody(FCadId CadId) const;
diff --git a/Engine/Plugins/Enterprise/DatasmithCADImporter/Source/CADInterfaces/Public/TechSoftInterface.h b/Engine/Plugins/Enterprise/DatasmithCADImporter/Source/CADInterfaces/Public/TechSoftInterface.h
--- a/Engine/Plugins/Enterprise/DatasmithCADImporter/Source/CADInterfaces/Public/TechSoftInterface.h
+++ b/Engine/Plugins/Enterprise/DatasmithCADImporter/Source/CADInterfaces/Public/TechSoftInterface.h
@@ -228,59 +228,59 @@ private:
 	A3DAsmModelFile* ModelFile;
 };
 
-enum EModellerType : uint32
+enum EModelerType : uint32
 {
-	ModellerUnknown = 0,              /*!< User modeller. */
-	ModellerCatia = 2,                /*!< CATIA modeller. */
-	ModellerCatiaV5 = 3,              /*!< CATIA V5 modeller. */
-	ModellerCadds = 4,                /*!< CADDS modeller. */
-	ModellerUnigraphics = 5,          /*!< Unigraphics modeller. */
-	ModellerParasolid = 6,            /*!< Parasolid modeller. */
-	ModellerEuclid = 7,               /*!< Euclid modeller. */
-	ModellerIges = 9,                 /*!< IGES modeller. */
-	ModellerUnisurf = 10,             /*!< Unisurf modeller. */
-	ModellerVda = 11,                 /*!< VDA modeller. */
-	ModellerStl = 12,                 /*!< STL modeller. */
-	ModellerWrl = 13,                 /*!< WRL modeller. */
-	ModellerDxf = 14,                 /*!< DXF modeller. */
-	ModellerAcis = 15,                /*!< ACIS modeller. */
-	ModellerProE = 16,                /*!< Pro/E modeller. */
-	ModellerStep = 18,                /*!< STEP modeller. */
-	ModellerIdeas = 19,               /*!< I-DEAS modeller. */
-	ModellerJt = 20,                  /*!< JT modeller. */
-	ModellerSlw = 22,                 /*!< SolidWorks modeller. */
-	ModellerCgr = 23,                 /*!< CGR modeller. */
-	ModellerPrc = 24,                 /*!< PRC modeller. */
-	ModellerXvl = 25,                 /*!< XVL modeller. */
-	ModellerHpgl = 26,                /*!< HPGL modeller. */
-	ModellerTopSolid = 27,            /*!< TopSolid modeller. */
-	ModellerOneSpaceDesigner = 28,    /*!< OneSpace designer modeller. */
-	Modeller3dxml = 29,               /*!< 3DXML modeller. */
-	ModellerInventor = 30,            /*!< Inventor modeller. */
-	ModellerPostScript = 31,          /*!< Postscript modeller. */
-	ModellerPDF = 32,                 /*!< PDF modeller. */
-	ModellerU3D = 33,                 /*!< U3D modeller. */
-	ModellerIFC = 34,                 /*!< IFC modeller. */
-	ModellerDWG = 35,                 /*!< DWG modeller. */
-	ModellerDWF = 36,                 /*!< DWF modeller. */
-	ModellerSE = 37,                  /*!< SolidEdge modeller. */
-	ModellerOBJ = 38,                 /*!< OBJ modeller. */
-	ModellerKMZ = 39,                 /*!< KMZ modeller. */
-	ModellerDAE = 40,                 /*!< COLLADA modeller. */
-	Modeller3DS = 41,                 /*!< 3DS modeller. */
-	ModellerRhino = 43,               /*!< Rhino modeller. */
-	ModellerXML = 44,                 /*!< XML modeller. */
-	Modeller3mf = 45,                 /*!< 3MF modeller. */
-	ModellerScs = 46,                 /*!< SCS modeller. */
-	Modeller3dHtml = 47,              /*!< 3DHTML modeller. */
-	ModellerHsf = 48,                 /*!< Hsf modeller. */
-	ModellerGltf = 49,                /*!< GL modeller. */
-	ModellerRevit = 50,               /*!< Revit modeller. */
-	ModellerFBX = 51,                 /*!< FBX modeller. */
-	ModellerStepXML = 52,             /*!< StepXML modeller. */
-	ModellerPLMXML = 53,              /*!< PLMXML modeller. */
-	ModellerNavisworks = 54,			/*!< For Future Use: Navisworks modeller. */
-	ModellerLast
+	ModelerUnknown = 0,              /*!< User Modeler. */
+	ModelerCatia = 2,                /*!< CATIA Modeler. */
+	ModelerCatiaV5 = 3,              /*!< CATIA V5 Modeler. */
+	ModelerCadds = 4,                /*!< CADDS Modeler. */
+	ModelerUnigraphics = 5,          /*!< Unigraphics Modeler. */
+	ModelerParasolid = 6,            /*!< Parasolid Modeler. */
+	ModelerEuclid = 7,               /*!< Euclid Modeler. */
+	ModelerIges = 9,                 /*!< IGES Modeler. */
+	ModelerUnisurf = 10,             /*!< Unisurf Modeler. */
+	ModelerVda = 11,                 /*!< VDA Modeler. */
+	ModelerStl = 12,                 /*!< STL Modeler. */
+	ModelerWrl = 13,                 /*!< WRL Modeler. */
+	ModelerDxf = 14,                 /*!< DXF Modeler. */
+	ModelerAcis = 15,                /*!< ACIS Modeler. */
+	ModelerProE = 16,                /*!< Pro/E Modeler. */
+	ModelerStep = 18,                /*!< STEP Modeler. */
+	ModelerIdeas = 19,               /*!< I-DEAS Modeler. */
+	ModelerJt = 20,                  /*!< JT Modeler. */
+	ModelerSlw = 22,                 /*!< SolidWorks Modeler. */
+	ModelerCgr = 23,                 /*!< CGR Modeler. */
+	ModelerPrc = 24,                 /*!< PRC Modeler. */
+	ModelerXvl = 25,                 /*!< XVL Modeler. */
+	ModelerHpgl = 26,                /*!< HPGL Modeler. */
+	ModelerTopSolid = 27,            /*!< TopSolid Modeler. */
+	ModelerOneSpaceDesigner = 28,    /*!< OneSpace designer Modeler. */
+	Modeler3dxml = 29,               /*!< 3DXML Modeler. */
+	ModelerInventor = 30,            /*!< Inventor Modeler. */
+	ModelerPostScript = 31,          /*!< Postscript Modeler. */
+	ModelerPDF = 32,                 /*!< PDF Modeler. */
+	ModelerU3D = 33,                 /*!< U3D Modeler. */
+	ModelerIFC = 34,                 /*!< IFC Modeler. */
+	ModelerDWG = 35,                 /*!< DWG Modeler. */
+	ModelerDWF = 36,                 /*!< DWF Modeler. */
+	ModelerSE = 37,                  /*!< SolidEdge Modeler. */
+	ModelerOBJ = 38,                 /*!< OBJ Modeler. */
+	ModelerKMZ = 39,                 /*!< KMZ Modeler. */
+	ModelerDAE = 40,                 /*!< COLLADA Modeler. */
+	Modeler3DS = 41,                 /*!< 3DS Modeler. */
+	ModelerRhino = 43,               /*!< Rhino Modeler. */
+	ModelerXML = 44,                 /*!< XML Modeler. */
+	Modeler3mf = 45,                 /*!< 3MF Modeler. */
+	ModelerScs = 46,                 /*!< SCS Modeler. */
+	Modeler3dHtml = 47,              /*!< 3DHTML Modeler. */
+	ModelerHsf = 48,                 /*!< Hsf Modeler. */
+	ModelerGltf = 49,                /*!< GL Modeler. */
+	ModelerRevit = 50,               /*!< Revit Modeler. */
+	ModelerFBX = 51,                 /*!< FBX Modeler. */
+	ModelerStepXML = 52,             /*!< StepXML Modeler. */
+	ModelerPLMXML = 53,              /*!< PLMXML Modeler. */
+	ModelerNavisworks = 54,			/*!< For Future Use: Navisworks Modeler. */
+	ModelerLast
 };
 
 } // CADLibrary
diff --git a/Engine/Plugins/Enterprise/DatasmithCADImporter/Source/CADTools/Public/CADOptions.h b/Engine/Plugins/Enterprise/DatasmithCADImporter/Source/CADTools/Public/CADOptions.h
--- a/Engine/Plugins/Enterprise/DatasmithCADImporter/Source/CADTools/Public/CADOptions.h
+++ b/Engine/Plugins/Enterprise/DatasmithCADImporter/Source/CADTools/Public/CADOptions.h
@@ -12,6 +12,20 @@ namespace CADLibrary
 {
 	CADTOOLS_API extern int32 GMaxImportThreads;
 
+	enum class EMesher
+	{
+		CADKernel,
+		TechSoft,
+		None
+	};
+
+	enum class EFailureReason : uint8
+	{
+		Curve3D,
+		Unknown,
+		None
+	};
+
 	enum EStitchingTechnique
 	{
 		StitchingNone = 0,
@@ -39,6 +53,7 @@ namespace CADLibrary
 		double MaxEdgeLength;
 		double MaxNormalAngle;
 		EStitchingTechnique StitchingTechnique;
+		EMesher Mesher;
 		FDatasmithUtils::EModelCoordSystem ModelCoordSys;
 		
 	public:
@@ -66,16 +81,34 @@ namespace CADLibrary
 			, MaxEdgeLength(0.0)
 			, MaxNormalAngle(20.0)
 			, StitchingTechnique(EStitchingTechnique::StitchingNone)
+			, Mesher(EMesher::TechSoft)
 			, ModelCoordSys(NewCoordinateSystem)
 		{
 		}
 	
-		void SetTesselationParameters(double InChordTolerance, double InMaxEdgeLength, double InMaxNormalAngle, CADLibrary::EStitchingTechnique InStitchingTechnique)
+		FImportParameters(FImportParameters InParamneters, EMesher InMesher)
+			: ChordTolerance(InParamneters.ChordTolerance)
+			, MaxEdgeLength(InParamneters.MaxEdgeLength)
+			, MaxNormalAngle(InParamneters.MaxNormalAngle)
+			, StitchingTechnique(InParamneters.StitchingTechnique)
+			, Mesher(InMesher)
+			, ModelCoordSys(InParamneters.ModelCoordSys)
+		{
+		}
+
+		void SetTesselationParameters(double InChordTolerance, double InMaxEdgeLength, double InMaxNormalAngle, EStitchingTechnique InStitchingTechnique)
 		{
 			ChordTolerance = InChordTolerance * GMeshingParameterFactor;
 			MaxEdgeLength = InMaxEdgeLength * GMeshingParameterFactor;
 			MaxNormalAngle = InMaxNormalAngle;
 			StitchingTechnique = InStitchingTechnique;
+			Mesher = FImportParameters::bGDisableCADKernelTessellation ? EMesher::TechSoft : EMesher::CADKernel;
+		}
+
+		void SetTesselationParameters(double InChordTolerance, double InMaxEdgeLength, double InMaxNormalAngle, EStitchingTechnique InStitchingTechnique, EMesher InMesher)
+		{
+			SetTesselationParameters(InChordTolerance, InMaxEdgeLength, InMaxNormalAngle, InStitchingTechnique);
+			Mesher = InMesher;
 		}
 
 		uint32 GetHash() const
@@ -85,7 +118,7 @@ namespace CADLibrary
 			{
 				Hash = HashCombine(Hash, ::GetTypeHash(Param));
 			}
-			for (uint32 Param : {uint32(StitchingTechnique), uint32(ModelCoordSys)})
+			for (uint32 Param : {uint32(StitchingTechnique), uint32(Mesher), uint32(ModelCoordSys)})
 			{
 				Hash = HashCombine(Hash, ::GetTypeHash(Param));
 			}
@@ -102,6 +135,7 @@ namespace CADLibrary
 			Ar << ImportParameters.MaxEdgeLength;
 			Ar << ImportParameters.MaxNormalAngle;
 			Ar << (uint32&) ImportParameters.StitchingTechnique;
+			Ar << (uint8&) ImportParameters.Mesher;
 			Ar << (uint8&) ImportParameters.ModelCoordSys;
 
 			// these static variables have to be serialized to be transmitted to CADWorkers
@@ -147,6 +181,11 @@ namespace CADLibrary
 			return StitchingTechnique;
 		}
 
+		EMesher GetMesher() const
+		{
+			return Mesher;
+		}
+
 		FDatasmithUtils::EModelCoordSystem GetModelCoordSys() const
 		{
 			return ModelCoordSys;
@@ -193,18 +232,18 @@ namespace CADLibrary
 		return FPaths::Combine(CachePath, TEXT("cad"), FileName);
 	}
 
-	inline FString BuildCacheFilePath(const TCHAR* CachePath, const TCHAR* Folder, uint32 BodyHash)
+	inline FString BuildCacheFilePath(const TCHAR* CachePath, const TCHAR* Folder, uint32 BodyHash, const EMesher Mesher)
 	{
 		FString BodyFileName = FString::Printf(TEXT("UEx%08x"), BodyHash);
 		FString OutFileName = FPaths::Combine(CachePath, Folder, BodyFileName);
 
-		if (FImportParameters::bGDisableCADKernelTessellation)
+		if (Mesher == EMesher::CADKernel)
 		{
-			OutFileName += TEXT(".prc");
+			OutFileName += TEXT(".ugeom");
 		}
 		else
 		{
-			OutFileName += TEXT(".ugeom");
+			OutFileName += TEXT(".prc");
 		}
 		return OutFileName;
 	}
diff --git a/Engine/Plugins/Enterprise/DatasmithCADImporter/Source/DatasmithCADTranslator/Private/DatasmithCADTranslator.cpp b/Engine/Plugins/Enterprise/DatasmithCADImporter/Source/DatasmithCADTranslator/Private/DatasmithCADTranslator.cpp
--- a/Engine/Plugins/Enterprise/DatasmithCADImporter/Source/DatasmithCADTranslator/Private/DatasmithCADTranslator.cpp
+++ b/Engine/Plugins/Enterprise/DatasmithCADImporter/Source/DatasmithCADTranslator/Private/DatasmithCADTranslator.cpp
@@ -276,8 +276,9 @@ bool FDatasmithCADTranslator::LoadScene(TSharedRef<IDatasmithScene> DatasmithSce
 				NumberOfWorkers = 1;
 			}
 
+			const EMesher Mesher = FImportParameters::bGDisableCADKernelTessellation ? EMesher::TechSoft : EMesher::CADKernel;
 			DatasmithDispatcher::FDatasmithDispatcher Dispatcher(ImportParameters, CachePath, NumberOfWorkers, CADFileToUEFileMap, CADFileToUEGeomMap);
-			Dispatcher.AddTask(FileDescriptor);
+			Dispatcher.AddTask(FileDescriptor, Mesher);
 
 			Dispatcher.Process(GMaxImportThreads != 1);
 		}
@@ -321,11 +322,16 @@ bool FDatasmithCADTranslator::LoadStaticMesh(const TSharedRef<IDatasmithMeshElem
 
 	CADLibrary::FMeshParameters MeshParameters;
 
+	// Two meshers can be used (TechSoft and CADKernel) during the import indeed in case of failure of one, the other is tried.
+	// At this stage, the only clue to know which mesher has been used for the current body is the extension of the rawdata file
+	const TCHAR* FileName = MeshElement->GetFile();
+	FString Extension = FPaths::GetExtension(FileName);
+
 	if (TOptional< FMeshDescription > Mesh = MeshBuilderPtr->GetMeshDescription(MeshElement, MeshParameters))
 	{
 		OutMeshPayload.LodMeshes.Add(MoveTemp(Mesh.GetValue()));
 
-		if (CADLibrary::FImportParameters::bGDisableCADKernelTessellation)
+		if (Extension == TEXT("prc"))
 		{
 			ParametricSurfaceUtils::AddSurfaceData(MeshElement->GetFile(), ImportParameters, MeshParameters, GetCommonTessellationOptions(), OutMeshPayload);
 		}
diff --git a/Engine/Plugins/Enterprise/DatasmithCADImporter/Source/DatasmithCADTranslator/Private/DatasmithSceneGraphBuilder.cpp b/Engine/Plugins/Enterprise/DatasmithCADImporter/Source/DatasmithCADTranslator/Private/DatasmithSceneGraphBuilder.cpp
--- a/Engine/Plugins/Enterprise/DatasmithCADImporter/Source/DatasmithCADTranslator/Private/DatasmithSceneGraphBuilder.cpp
+++ b/Engine/Plugins/Enterprise/DatasmithCADImporter/Source/DatasmithCADTranslator/Private/DatasmithSceneGraphBuilder.cpp
@@ -514,7 +514,7 @@ TSharedPtr<IDatasmithMeshElement> FDatasmithSceneBaseGraphBuilder::FindOrAddMesh
 
 	BodyUuidToMeshElement.Add(Body.MeshActorUId, MeshElement);
 
-	FString BodyCachePath = CADLibrary::BuildCacheFilePath(*CachePath, TEXT("body"), Body.MeshActorUId);
+	FString BodyCachePath = CADLibrary::BuildCacheFilePath(*CachePath, TEXT("body"), Body.MeshActorUId, CADLibrary::EMesher::CADKernel);
 	MeshElement->SetFile(*BodyCachePath);
 
 	return MeshElement;
diff --git a/Engine/Plugins/Enterprise/DatasmithCADImporter/Source/DatasmithDispatcher/Private/DatasmithDispatcher.cpp b/Engine/Plugins/Enterprise/DatasmithCADImporter/Source/DatasmithDispatcher/Private/DatasmithDispatcher.cpp
--- a/Engine/Plugins/Enterprise/DatasmithCADImporter/Source/DatasmithDispatcher/Private/DatasmithDispatcher.cpp
+++ b/Engine/Plugins/Enterprise/DatasmithCADImporter/Source/DatasmithDispatcher/Private/DatasmithDispatcher.cpp
@@ -71,16 +71,17 @@ FDatasmithDispatcher::FDatasmithDispatcher(const CADLibrary::FImportParameters&
 	}
 }
 
-void FDatasmithDispatcher::AddTask(const CADLibrary::FFileDescriptor& InFileDescription)
+void FDatasmithDispatcher::AddTask(const CADLibrary::FFileDescriptor& InFileDescription, const CADLibrary::EMesher Mesher)
 {
+	using namespace CADLibrary;
 	FScopeLock Lock(&TaskPoolCriticalSection);
 
-	if (Algo::FindByPredicate(TaskPool, [&InFileDescription](const FTask& Task) { return Task.FileDescription == InFileDescription; }))
+	if (Algo::FindByPredicate(TaskPool, [&InFileDescription, &Mesher](const FTask& Task) { return Task.FileDescription == InFileDescription && Task.Mesher == Mesher; }))
 	{
 		return;
 	}
 
-	int32 TaskIndex = TaskPool.Emplace(InFileDescription);
+	int32 TaskIndex = TaskPool.Emplace(InFileDescription, Mesher);
 	TaskPool[TaskIndex].Index = TaskIndex;
 }
 
@@ -160,8 +161,20 @@ void FDatasmithDispatcher::SetTaskState(int32 TaskIndex, ETaskState TaskState)
 			break;
 		}
 
-		case ETaskState::ProcessOk:
 		case ETaskState::ProcessFailed:
+		{
+			if (Task.Mesher == CADLibrary::EMesher::CADKernel)
+			{
+				UE_LOG(LogDatasmithDispatcher, Log, TEXT("   - Task failed with CADKernel: %s"), *Task.FileDescription.GetFileName());
+				UE_LOG(LogDatasmithDispatcher, Log, TEXT("      => Add task to process with Techsoft"));
+				AddTask(Task.FileDescription, CADLibrary::EMesher::TechSoft);
+			}
+
+			CompletedTaskCount++;
+			break;
+		}
+
+		case ETaskState::ProcessOk:
 		case ETaskState::FileNotFound:
 		{
 			CompletedTaskCount++;
@@ -362,26 +375,32 @@ void FDatasmithDispatcher::CloseHandlers()
 
 void FDatasmithDispatcher::ProcessLocal()
 {
+	using namespace CADLibrary;
+	EMesher DefaultMesher = FImportParameters::bGDisableCADKernelTessellation ? EMesher::TechSoft : EMesher::CADKernel;
+
 	bool bCheckMemory = true;
 	while (TOptional<FTask> Task = GetNextTask())
 	{
-		CADLibrary::FFileDescriptor& FileDescription = Task->FileDescription;
+		FFileDescriptor& FileDescription = Task->FileDescription;
+		FImportParameters FileImporParameters(ImportParameters, Task->Mesher);
 
-		CADLibrary::FCADFileReader FileReader(ImportParameters, FileDescription, *FPaths::EnginePluginsDir(), ProcessCacheFolder);
+		FCADFileReader FileReader(FileImporParameters, FileDescription, *FPaths::EnginePluginsDir(), ProcessCacheFolder);
 		ETaskState ProcessResult = FileReader.ProcessFile();
 
 		ETaskState TaskState = ProcessResult;
 		SetTaskState(Task->Index, TaskState);
 
-		if (TaskState == ETaskState::ProcessOk)
+		switch (TaskState)
 		{
-			const CADLibrary::FCADFileData& CADFileData = FileReader.GetCADFileData();
-			const TArray<CADLibrary::FFileDescriptor>& ExternalRefSet = CADFileData.GetExternalRefSet();
+		case ETaskState::ProcessOk:
+		{
+			const FCADFileData& CADFileData = FileReader.GetCADFileData();
+			const TArray<FFileDescriptor>& ExternalRefSet = CADFileData.GetExternalRefSet();
 			if (ExternalRefSet.Num() > 0)
 			{
-				for (const CADLibrary::FFileDescriptor& ExternalFile : ExternalRefSet)
+				for (const FFileDescriptor& ExternalFile : ExternalRefSet)
 				{
-					AddTask(ExternalFile);
+					AddTask(ExternalFile, DefaultMesher);
 				}
 			}
 
@@ -406,6 +425,18 @@ void FDatasmithDispatcher::ProcessLocal()
 				}
 				bCheckMemory = false;
 			}
+			break;
+		}
+		case ETaskState::ProcessFailed:
+		{
+			if (Task->Mesher == EMesher::CADKernel)
+			{
+				AddTask(Task->FileDescription, EMesher::TechSoft);
+			}
+			break;
+		}
+		default:
+			break;
 		}
 	}
 }
diff --git a/Engine/Plugins/Enterprise/DatasmithCADImporter/Source/DatasmithDispatcher/Private/DatasmithWorkerHandler.cpp b/Engine/Plugins/Enterprise/DatasmithCADImporter/Source/DatasmithDispatcher/Private/DatasmithWorkerHandler.cpp
--- a/Engine/Plugins/Enterprise/DatasmithCADImporter/Source/DatasmithDispatcher/Private/DatasmithWorkerHandler.cpp
+++ b/Engine/Plugins/Enterprise/DatasmithCADImporter/Source/DatasmithDispatcher/Private/DatasmithWorkerHandler.cpp
@@ -329,23 +329,35 @@ void FDatasmithWorkerHandler::ProcessCommand(FPingCommand& PingCommand)
 
 void FDatasmithWorkerHandler::ProcessCommand(FCompletedTaskCommand& CompletedTaskCommand)
 {
+	using namespace CADLibrary;
 	if (!CurrentTask.IsSet())
 	{
 		return;
 	}
 
-	for (const CADLibrary::FFileDescriptor& ExternalReferenceFile : CompletedTaskCommand.ExternalReferences)
+	const EMesher DefaultMesher = FImportParameters::bGDisableCADKernelTessellation ? EMesher::TechSoft : EMesher::CADKernel;
+
+	if(CompletedTaskCommand.ProcessResult == ETaskState::ProcessOk)
 	{
-		if (ExternalReferenceFile.GetFileName().IsEmpty())
+		for (const FFileDescriptor& ExternalReferenceFile : CompletedTaskCommand.ExternalReferences)
 		{
-			continue;
+			if (ExternalReferenceFile.GetFileName().IsEmpty())
+			{
+				continue;
+			}
+			Dispatcher.AddTask(ExternalReferenceFile, DefaultMesher);
 		}
-		Dispatcher.AddTask(ExternalReferenceFile);
+		UE_LOG(LogDatasmithDispatcher, Log, TEXT("   - Task completed: %s"), *CurrentTask->FileDescription.GetFileName());
+	}
+	else if (CompletedTaskCommand.ProcessResult == ETaskState::ProcessFailed && DefaultMesher == EMesher::CADKernel)
+	{
+		UE_LOG(LogDatasmithDispatcher, Log, TEXT("   - Task failed with CADKernel: %s"), *CurrentTask->FileDescription.GetFileName());
+		UE_LOG(LogDatasmithDispatcher, Log, TEXT("      => Add task to process with Techsoft"));
+		Dispatcher.AddTask(CurrentTask->FileDescription, EMesher::TechSoft);
 	}
 
 	Dispatcher.SetTaskState(CurrentTask->Index, CompletedTaskCommand.ProcessResult);
 
-	UE_LOG(LogDatasmithDispatcher, Log, TEXT("   - Task completed: %s"), *CurrentTask->FileDescription.GetFileName());
 
 	Dispatcher.LinkCTFileToUnrealCacheFile(CurrentTask->FileDescription, CompletedTaskCommand.SceneGraphFileName, CompletedTaskCommand.GeomFileName);
 	Dispatcher.LogMessages(CompletedTaskCommand.Messages);
diff --git a/Engine/Plugins/Enterprise/DatasmithCADImporter/Source/DatasmithDispatcher/Public/DatasmithCommands.h b/Engine/Plugins/Enterprise/DatasmithCADImporter/Source/DatasmithDispatcher/Public/DatasmithCommands.h
--- a/Engine/Plugins/Enterprise/DatasmithCADImporter/Source/DatasmithDispatcher/Public/DatasmithCommands.h
+++ b/Engine/Plugins/Enterprise/DatasmithCADImporter/Source/DatasmithDispatcher/Public/DatasmithCommands.h
@@ -79,7 +79,11 @@ class FRunTaskCommand : public ICommand
 {
 public:
 	FRunTaskCommand() = default;
-	FRunTaskCommand(const FTask& Task) : JobFileDescription(Task.FileDescription), JobIndex(Task.Index) {}
+	FRunTaskCommand(const FTask& Task) 
+		: JobFileDescription(Task.FileDescription)
+		, Mesher(Task.Mesher)
+		, JobIndex(Task.Index) 
+	{}
 	virtual ECommandId GetType() const override { return ECommandId::RunTask; }
 
 protected:
@@ -87,6 +91,8 @@ protected:
 
 public:
 	CADLibrary::FFileDescriptor JobFileDescription;
+	CADLibrary::EMesher Mesher;
+
 	int32 JobIndex = -1;
 };
 
diff --git a/Engine/Plugins/Enterprise/DatasmithCADImporter/Source/DatasmithDispatcher/Public/DatasmithDispatcher.h b/Engine/Plugins/Enterprise/DatasmithCADImporter/Source/DatasmithDispatcher/Public/DatasmithDispatcher.h
--- a/Engine/Plugins/Enterprise/DatasmithCADImporter/Source/DatasmithDispatcher/Public/DatasmithDispatcher.h
+++ b/Engine/Plugins/Enterprise/DatasmithCADImporter/Source/DatasmithDispatcher/Public/DatasmithDispatcher.h
@@ -18,7 +18,7 @@ class DATASMITHDISPATCHER_API FDatasmithDispatcher
 public:
 	FDatasmithDispatcher(const CADLibrary::FImportParameters& InImportParameters, const FString& InCacheDir, int32 InNumberOfWorkers, TMap<uint32, FString>& CADFileToUnrealFileMap, TMap<uint32, FString>& CADFileToUnrealGeomMap);
 
-	void AddTask(const CADLibrary::FFileDescriptor & FileDescription);
+	void AddTask(const CADLibrary::FFileDescriptor & FileDescription, CADLibrary::EMesher Mesher);
 	TOptional<FTask> GetNextTask();
 	void SetTaskState(int32 TaskIndex, ETaskState TaskState);
 
diff --git a/Engine/Plugins/Enterprise/DatasmithCADImporter/Source/DatasmithDispatcher/Public/DatasmithDispatcherTask.h b/Engine/Plugins/Enterprise/DatasmithCADImporter/Source/DatasmithDispatcher/Public/DatasmithDispatcherTask.h
--- a/Engine/Plugins/Enterprise/DatasmithCADImporter/Source/DatasmithDispatcher/Public/DatasmithDispatcherTask.h
+++ b/Engine/Plugins/Enterprise/DatasmithCADImporter/Source/DatasmithDispatcher/Public/DatasmithDispatcherTask.h
@@ -4,6 +4,7 @@
 
 #include "CoreMinimal.h"
 #include "CADData.h"
+#include "CADOptions.h"
 
 namespace DatasmithDispatcher
 {
@@ -13,15 +14,18 @@ struct FTask
 {
 	FTask() = default;
 
-	FTask(const CADLibrary::FFileDescriptor& InFile)
+	FTask(const CADLibrary::FFileDescriptor& InFile, const CADLibrary::EMesher InMesher)
 		: FileDescription(InFile)
+		, Mesher(InMesher)
 		, State(ETaskState::UnTreated)
 	{
 	}
 
 	CADLibrary::FFileDescriptor FileDescription;
+	CADLibrary::EMesher Mesher;
 	int32 Index = -1;
 	ETaskState State = ETaskState::Unknown;
+
 };
 
 } // NS DatasmithDispatcher
diff --git a/Engine/Plugins/Enterprise/DatasmithCADImporter/Source/DatasmithPLMXMLTranslator/Private/DatasmithPlmXmlImporter.cpp b/Engine/Plugins/Enterprise/DatasmithCADImporter/Source/DatasmithPLMXMLTranslator/Private/DatasmithPlmXmlImporter.cpp
--- a/Engine/Plugins/Enterprise/DatasmithCADImporter/Source/DatasmithPLMXMLTranslator/Private/DatasmithPlmXmlImporter.cpp
+++ b/Engine/Plugins/Enterprise/DatasmithCADImporter/Source/DatasmithPLMXMLTranslator/Private/DatasmithPlmXmlImporter.cpp
@@ -255,8 +255,11 @@ namespace PlmXml
 		// Adds geom file to load and returns Id to use in InstantiateMesh later(after all is loaded)
 		int32 AddMeshToLoad(const FString& FullPath)
 		{
-			CADLibrary::FFileDescriptor FileDescription(*FullPath);
-			DatasmithDispatcher->AddTask(FileDescription);
+			using namespace CADLibrary;
+			const EMesher Mesher = FImportParameters::bGDisableCADKernelTessellation ? EMesher::TechSoft : EMesher::CADKernel;
+
+			FFileDescriptor FileDescription(*FullPath);
+			DatasmithDispatcher->AddTask(FileDescription, Mesher);
 			return FilePaths.Add(FullPath);
 		}
 
diff --git a/Engine/Source/Programs/Enterprise/Datasmith/DatasmithCADWorker/Private/DatasmithCADWorkerImpl.cpp b/Engine/Source/Programs/Enterprise/Datasmith/DatasmithCADWorker/Private/DatasmithCADWorkerImpl.cpp
--- a/Engine/Source/Programs/Enterprise/Datasmith/DatasmithCADWorker/Private/DatasmithCADWorkerImpl.cpp
+++ b/Engine/Source/Programs/Enterprise/Datasmith/DatasmithCADWorker/Private/DatasmithCADWorkerImpl.cpp
@@ -155,7 +155,8 @@ uint64 DefineMaximumAllowedDuration(const CADLibrary::FFileDescriptor& FileDescr
 
 void FDatasmithCADWorkerImpl::ProcessCommand(const FRunTaskCommand& RunTaskCommand)
 {
-	CADLibrary::FFileDescriptor FileToProcess = RunTaskCommand.JobFileDescription;
+	using namespace CADLibrary;
+	FFileDescriptor FileToProcess = RunTaskCommand.JobFileDescription;
 	UE_LOG(LogDatasmithCADWorker, Verbose, TEXT("Process %s %s"), *FileToProcess.GetFileName(), *FileToProcess.GetConfiguration());
 
 	FCompletedTaskCommand CompletedTask;
@@ -166,7 +167,9 @@ void FDatasmithCADWorkerImpl::ProcessCommand(const FRunTaskCommand& RunTaskComma
 	Checkers.Emplace(UE::Tasks::Launch(TEXT("TimeChecker"), [&FileToProcess, &MaxDuration]() { CheckDuration(FileToProcess, MaxDuration); }));
 	Checkers.Emplace(UE::Tasks::Launch(TEXT("MemoryChecker"), []() { CheckMemory(); }));
 
-	CADLibrary::FCADFileReader FileReader(ImportParameters, FileToProcess, EnginePluginsPath, CachePath);
+	FImportParameters FileImporParameters(ImportParameters, RunTaskCommand.Mesher);
+
+	FCADFileReader FileReader(FileImporParameters, FileToProcess, EnginePluginsPath, CachePath);
 	CompletedTask.ProcessResult = FileReader.ProcessFile();
 
 	bProcessIsRunning = false;
@@ -179,7 +182,7 @@ void FDatasmithCADWorkerImpl::ProcessCommand(const FRunTaskCommand& RunTaskComma
 			CompletedTask.ProcessResult = ETaskState::Unknown;
 		}
 			
-		const CADLibrary::FCADFileData& CADFileData = FileReader.GetCADFileData();
+		const FCADFileData& CADFileData = FileReader.GetCADFileData();
 		CompletedTask.ExternalReferences = CADFileData.GetExternalRefSet();
 		CompletedTask.SceneGraphFileName = CADFileData.GetSceneGraphFileName();
 		CompletedTask.GeomFileName = CADFileData.GetMeshFileName();
@@ -187,7 +190,7 @@ void FDatasmithCADWorkerImpl::ProcessCommand(const FRunTaskCommand& RunTaskComma
 
 		UE_LOG(LogDatasmithCADWorker, Verbose, TEXT("=> Process %s %s saved into %s%s and %s%s."), *FileToProcess.GetFileName(), *FileToProcess.GetConfiguration(), *CompletedTask.SceneGraphFileName, TEXT(".sg"), *CompletedTask.GeomFileName, TEXT(".gm"));
 		UE_LOG(LogDatasmithCADWorker, Verbose, TEXT("     It generates %d bodies"), CADFileData.GetBodyMeshes().Num());
-		for (const CADLibrary::FBodyMesh& BodyMesh : CADFileData.GetBodyMeshes())
+		for (const FBodyMesh& BodyMesh : CADFileData.GetBodyMeshes())
 		{
 			FString BodyFileName = FString::Printf(TEXT("UEx%08x"), BodyMesh.MeshActorUId);
 			UE_LOG(LogDatasmithCADWorker, Verbose, TEXT("     - Body %s"), *BodyFileName);
