Update DatasmithCADCache version to fix problem of datasmtih retessellation action

diff --git a/Engine/Plugins/Enterprise/DatasmithCADImporter/Source/CADTools/Private/CADToolsModule.cpp b/Engine/Plugins/Enterprise/DatasmithCADImporter/Source/CADTools/Private/CADToolsModule.cpp
--- a/Engine/Plugins/Enterprise/DatasmithCADImporter/Source/CADTools/Private/CADToolsModule.cpp
+++ b/Engine/Plugins/Enterprise/DatasmithCADImporter/Source/CADTools/Private/CADToolsModule.cpp
@@ -18,7 +18,7 @@ bool FCADToolsModule::IsAvailable()
 
 uint32 FCADToolsModule::GetCacheVersion()
 {
-	return 503;
+	return 504;
 }
 
 IMPLEMENT_MODULE(FCADToolsModule, CADTools);
