Fix CADWorker commands

diff --git a/Engine/Plugins/Enterprise/DatasmithCADImporter/Source/DatasmithDispatcher/Private/DatasmithCommands.cpp b/Engine/Plugins/Enterprise/DatasmithCADImporter/Source/DatasmithDispatcher/Private/DatasmithCommands.cpp
--- a/Engine/Plugins/Enterprise/DatasmithCADImporter/Source/DatasmithDispatcher/Private/DatasmithCommands.cpp
+++ b/Engine/Plugins/Enterprise/DatasmithCADImporter/Source/DatasmithDispatcher/Private/DatasmithCommands.cpp
@@ -47,6 +47,7 @@ namespace DatasmithDispatcher
 	void FRunTaskCommand::SerializeImpl(FArchive& Ar)
 	{
 		Ar << JobFileDescription;
+		Ar << Mesher;
 		Ar << JobIndex;
 	}
 
