[CAD Import] Disable SewMeshIfNeeded by default. It can be activated with the CVar ds.CADTranslator.SewMeshIfNeeded = true

diff --git a/Engine/Plugins/Enterprise/DatasmithCADImporter/Source/CADTools/Private/CADOption.cpp b/Engine/Plugins/Enterprise/DatasmithCADImporter/Source/CADTools/Private/CADOption.cpp
--- a/Engine/Plugins/Enterprise/DatasmithCADImporter/Source/CADTools/Private/CADOption.cpp
+++ b/Engine/Plugins/Enterprise/DatasmithCADImporter/Source/CADTools/Private/CADOption.cpp
@@ -116,13 +116,13 @@ FAutoConsoleVariableRef GCADTranslatorUnitScale(
 Default value of UnitScale is 1 i.e. unit = cm\n"),
 ECVF_Default);
 
-bool FImportParameters::bGSewMeshIfNeeded = true;
+bool FImportParameters::bGSewMeshIfNeeded = false;
 FAutoConsoleVariableRef GCADTranslatorSewMeshIfNeeded(
 	TEXT("ds.CADTranslator.SewMeshIfNeeded"),
 	FImportParameters::bGSewMeshIfNeeded,
 	TEXT("Perform a welding of the mesh to try to stitch mesh cracks\n\
 This welding is performed respecting the ds.CADTranslator.StitchingTolerance\n\
-Default value is true\n"),
+Default value is false\n"),
 ECVF_Default);
 
 bool FImportParameters::bGRemoveDuplicatedTriangle = false;
