Fix import PlmXml

diff --git a/Engine/Plugins/Enterprise/DatasmithCADImporter/Source/DatasmithCADTranslator/Private/DatasmithCADTranslator.cpp b/Engine/Plugins/Enterprise/DatasmithCADImporter/Source/DatasmithCADTranslator/Private/DatasmithCADTranslator.cpp
--- a/Engine/Plugins/Enterprise/DatasmithCADImporter/Source/DatasmithCADTranslator/Private/DatasmithCADTranslator.cpp
+++ b/Engine/Plugins/Enterprise/DatasmithCADImporter/Source/DatasmithCADTranslator/Private/DatasmithCADTranslator.cpp
@@ -254,30 +254,8 @@ bool FDatasmithCADTranslator::LoadScene(TSharedRef<IDatasmithScene> DatasmithSce
 	{
 		TMap<uint32, FString> CADFileToUEFileMap;
 		{
-			constexpr double RecommandedRamPerWorkers = 6.;
-			constexpr double OneGigaByte = 1024. * 1024. * 1024.;
-			const double AvailableRamGB = (double) (FPlatformMemory::GetStats().AvailablePhysical / OneGigaByte);
-
-			const int32 MaxNumberOfWorkers = FPlatformMisc::NumberOfCores();
-			const int32 RecommandedNumberOfWorkers = (int32) (AvailableRamGB / RecommandedRamPerWorkers + 0.5);
-
-			// UE recommendation 
-			int32 NumberOfWorkers = FMath::Min(MaxNumberOfWorkers, RecommandedNumberOfWorkers);
-
-			// User choice but limited by the number of cores. More brings nothing
-			if (GMaxImportThreads > 1)
-			{
-				NumberOfWorkers = FMath::Min(GMaxImportThreads, MaxNumberOfWorkers);
-			}
-
-			// If the file cannot have reference. One worker is enough.
-			if (GMaxImportThreads != 1 && !FileDescriptor.CanReferenceOtherFiles())
-			{
-				NumberOfWorkers = 1;
-			}
-
 			const EMesher Mesher = FImportParameters::bGDisableCADKernelTessellation ? EMesher::TechSoft : EMesher::CADKernel;
-			DatasmithDispatcher::FDatasmithDispatcher Dispatcher(ImportParameters, CachePath, NumberOfWorkers, CADFileToUEFileMap, CADFileToUEGeomMap);
+			DatasmithDispatcher::FDatasmithDispatcher Dispatcher(ImportParameters, CachePath, CADFileToUEFileMap, CADFileToUEGeomMap);
 			Dispatcher.AddTask(FileDescriptor, Mesher);
 
 			Dispatcher.Process(GMaxImportThreads != 1);
diff --git a/Engine/Plugins/Enterprise/DatasmithCADImporter/Source/DatasmithDispatcher/Private/DatasmithDispatcher.cpp b/Engine/Plugins/Enterprise/DatasmithCADImporter/Source/DatasmithDispatcher/Private/DatasmithDispatcher.cpp
--- a/Engine/Plugins/Enterprise/DatasmithCADImporter/Source/DatasmithDispatcher/Private/DatasmithDispatcher.cpp
+++ b/Engine/Plugins/Enterprise/DatasmithCADImporter/Source/DatasmithDispatcher/Private/DatasmithDispatcher.cpp
@@ -4,6 +4,7 @@
 
 #include "CADFileData.h"
 #include "CADFileReader.h"
+#include "CADOptions.h"
 #include "DatasmithDispatcherConfig.h"
 #include "DatasmithDispatcherLog.h"
 #include "DatasmithDispatcherTask.h"
@@ -51,16 +52,32 @@ const EAppReturnType::Type OpenMessageDialog(uint64 AvailableRam, uint64 Estimat
 }
 
 
-FDatasmithDispatcher::FDatasmithDispatcher(const CADLibrary::FImportParameters& InImportParameters, const FString& InCacheDir, int32 InNumberOfWorkers, TMap<uint32, FString>& OutCADFileToUnrealFileMap, TMap<uint32, FString>& OutCADFileToUnrealGeomMap)
+FDatasmithDispatcher::FDatasmithDispatcher(const CADLibrary::FImportParameters& InImportParameters, const FString& InCacheDir, TMap<uint32, FString>& OutCADFileToUnrealFileMap, TMap<uint32, FString>& OutCADFileToUnrealGeomMap)
 	: NextTaskIndex(0)
 	, CompletedTaskCount(0)
 	, CADFileToUnrealFileMap(OutCADFileToUnrealFileMap)
 	, CADFileToUnrealGeomMap(OutCADFileToUnrealGeomMap)
 	, ProcessCacheFolder(InCacheDir)
 	, ImportParameters(InImportParameters)
-	, NumberOfWorkers(InNumberOfWorkers)
+	, NumberOfWorkers(0)
 	, NextWorkerId(0)
 {
+	constexpr double RecommandedRamPerWorkers = 6.;
+	constexpr double OneGigaByte = 1024. * 1024. * 1024.;
+	const double AvailableRamGB = (double)(FPlatformMemory::GetStats().AvailablePhysical / OneGigaByte);
+
+	const int32 MaxNumberOfWorkers = FPlatformMisc::NumberOfCores();
+	const int32 RecommandedNumberOfWorkers = (int32)(AvailableRamGB / RecommandedRamPerWorkers + 0.5);
+
+	// UE recommendation 
+	NumberOfWorkers = FMath::Min(MaxNumberOfWorkers, RecommandedNumberOfWorkers);
+
+	// User choice but limited by the number of cores. More brings nothing
+	if (CADLibrary::GMaxImportThreads > 1)
+	{
+		NumberOfWorkers = FMath::Min(CADLibrary::GMaxImportThreads, MaxNumberOfWorkers);
+	}
+
 	if (CADLibrary::FImportParameters::bGEnableCADCache)
 	{
 		// init cache folders
@@ -200,6 +217,11 @@ void FDatasmithDispatcher::SetTaskState(int32 TaskIndex, ETaskState TaskState)
 
 void FDatasmithDispatcher::Process(bool bWithProcessor)
 {
+	if (TaskPool.Num() == 1 && !TaskPool[0].FileDescription.CanReferenceOtherFiles())
+	{
+		NumberOfWorkers = 1;
+	}
+
 	EstimationOfMemoryUsedByStaticMeshes = 0;
 	bool bCheckMemory = CADLibrary::FImportParameters::bValidationProcess ? false : true;
 
diff --git a/Engine/Plugins/Enterprise/DatasmithCADImporter/Source/DatasmithDispatcher/Public/DatasmithDispatcher.h b/Engine/Plugins/Enterprise/DatasmithCADImporter/Source/DatasmithDispatcher/Public/DatasmithDispatcher.h
--- a/Engine/Plugins/Enterprise/DatasmithCADImporter/Source/DatasmithDispatcher/Public/DatasmithDispatcher.h
+++ b/Engine/Plugins/Enterprise/DatasmithCADImporter/Source/DatasmithDispatcher/Public/DatasmithDispatcher.h
@@ -16,7 +16,7 @@ namespace DatasmithDispatcher
 class DATASMITHDISPATCHER_API FDatasmithDispatcher
 {
 public:
-	FDatasmithDispatcher(const CADLibrary::FImportParameters& InImportParameters, const FString& InCacheDir, int32 InNumberOfWorkers, TMap<uint32, FString>& CADFileToUnrealFileMap, TMap<uint32, FString>& CADFileToUnrealGeomMap);
+	FDatasmithDispatcher(const CADLibrary::FImportParameters& InImportParameters, const FString& InCacheDir, TMap<uint32, FString>& CADFileToUnrealFileMap, TMap<uint32, FString>& CADFileToUnrealGeomMap);
 
 	void AddTask(const CADLibrary::FFileDescriptor & FileDescription, CADLibrary::EMesher Mesher);
 	TOptional<FTask> GetNextTask();
diff --git a/Engine/Plugins/Enterprise/DatasmithCADImporter/Source/DatasmithPLMXMLTranslator/Private/DatasmithPlmXmlImporter.cpp b/Engine/Plugins/Enterprise/DatasmithCADImporter/Source/DatasmithPLMXMLTranslator/Private/DatasmithPlmXmlImporter.cpp
--- a/Engine/Plugins/Enterprise/DatasmithCADImporter/Source/DatasmithPLMXMLTranslator/Private/DatasmithPlmXmlImporter.cpp
+++ b/Engine/Plugins/Enterprise/DatasmithCADImporter/Source/DatasmithPLMXMLTranslator/Private/DatasmithPlmXmlImporter.cpp
@@ -163,6 +163,10 @@ namespace PlmXml
 	TArray<FIdRef> GetAttributeIDREFS(const FXmlNode* Node, const FString Name)
 	{
 		FString IdRefsText = Node->GetAttribute(*Name);
+		if (IdRefsText.IsEmpty())
+		{
+			return TArray<FIdRef>();
+		}
 
 		TArray<FString> IdRefs;
 		IdRefsText.ParseIntoArray(IdRefs, TEXT(" "), false);
@@ -249,7 +253,7 @@ namespace PlmXml
 			// Setup of import parameters for DatasmithDispatcher copied from FDatasmithCADTranslator's setup
 			ImportParameters.SetTesselationParameters(TessellationOptions.ChordTolerance, TessellationOptions.MaxEdgeLength, TessellationOptions.NormalTolerance, (CADLibrary::EStitchingTechnique) TessellationOptions.StitchingTechnique);
 
-			DatasmithDispatcher = MakeUnique<DatasmithDispatcher::FDatasmithDispatcher>(ImportParameters, CacheDir, FPlatformMisc::NumberOfCores(), CADFileToUEFileMap, CADFileToUEGeomMap);
+			DatasmithDispatcher = MakeUnique<DatasmithDispatcher::FDatasmithDispatcher>(ImportParameters, CacheDir, CADFileToUEFileMap, CADFileToUEGeomMap);
 		}
 
 		// Adds geom file to load and returns Id to use in InstantiateMesh later(after all is loaded)
@@ -1174,7 +1178,7 @@ namespace PlmXml
 
 		PlmXml::FParsedProductDef ProductDef;
 
-		// Read all ProductInstance, ProductRevisionView, ProductView and memoize them, they are refernced in PlmXml graph by their 'id's
+		// Read all ProductInstance, ProductRevisionView, ProductView and memorize them, they are referenced in PlmXml graph by their 'id's
 		for (const FXmlNode* InstanceGraphChildNode : InstanceGraphNode->GetChildrenNodes())
 		{
 			FString IdText = PlmXml::GetAttributeId(InstanceGraphChildNode);
diff --git a/Engine/Plugins/Enterprise/DatasmithCADImporter/Source/DatasmithPLMXMLTranslator/Private/DatasmithPlmXmlTranslator.cpp b/Engine/Plugins/Enterprise/DatasmithCADImporter/Source/DatasmithPLMXMLTranslator/Private/DatasmithPlmXmlTranslator.cpp
--- a/Engine/Plugins/Enterprise/DatasmithCADImporter/Source/DatasmithPLMXMLTranslator/Private/DatasmithPlmXmlTranslator.cpp
+++ b/Engine/Plugins/Enterprise/DatasmithCADImporter/Source/DatasmithPLMXMLTranslator/Private/DatasmithPlmXmlTranslator.cpp
@@ -6,6 +6,8 @@
 
 #include "CADInterfacesModule.h"
 
+#include "CADData.h"
+#include "CADOptions.h"
 #include "DatasmithSceneFactory.h"
 #include "DatasmithSceneSource.h"
 #include "DatasmithImportOptions.h"
@@ -60,12 +62,66 @@ bool FDatasmithPlmXmlTranslator::IsSourceSupported(const FDatasmithSceneSource&
 
 bool FDatasmithPlmXmlTranslator::LoadScene(TSharedRef<IDatasmithScene> OutScene)
 {
+	using namespace CADLibrary;
+
+	const FString& FilePath = GetSource().GetSourceFile();
+
+	FFileDescriptor FileDescriptor(*FPaths::ConvertRelativePathToFull(FilePath));
+
+	UE_LOG(LogDatasmithXMLPLMTranslator, Display, TEXT("PlmXml translation [%s]."), *FileDescriptor.GetSourcePath());
+	UE_LOG(LogDatasmithXMLPLMTranslator, Display, TEXT(" - Parsing Library:      %s"), TEXT("TechSoft"));
+	UE_LOG(LogDatasmithXMLPLMTranslator, Display, TEXT(" - Tessellation Library: %s") , FImportParameters::bGDisableCADKernelTessellation ? TEXT("TechSoft") : TEXT("CADKernel"));
+	UE_LOG(LogDatasmithXMLPLMTranslator, Display, TEXT(" - Cache mode:           %s") , FImportParameters::bGEnableCADCache ? (FImportParameters::bGOverwriteCache ? TEXT("Override") : TEXT("Enabled")) : TEXT("Disabled"));
+	UE_LOG(LogDatasmithXMLPLMTranslator, Display, TEXT(" - Processing:           %s") , FImportParameters::bGEnableCADCache ? (GMaxImportThreads == 1 ? TEXT("Sequencial") : TEXT("Parallel")) : TEXT("Sequencial"));
+
+	TFunction<double(double, double, const TCHAR*)> CheckParameterValue = [](double Value, double MinValue, const TCHAR* ParameterName) -> double
+	{
+		if (Value < MinValue)
+		{
+			UE_LOG(LogDatasmithXMLPLMTranslator, Warning, TEXT("%s value (%f) of tessellation parameters is smaller than the minimal value %f. It's value is modified to respect the limit"), ParameterName, Value, MinValue);
+			return MinValue;
+		}
+		return Value;
+	};
+
+	CommonTessellationOptions.ChordTolerance = CheckParameterValue(CommonTessellationOptions.ChordTolerance, UE::DatasmithTessellation::MinTessellationChord, TEXT("Chord tolerance"));
+	CommonTessellationOptions.MaxEdgeLength = FMath::IsNearlyZero(CommonTessellationOptions.MaxEdgeLength) ? 0. : CheckParameterValue(CommonTessellationOptions.MaxEdgeLength, UE::DatasmithTessellation::MinTessellationEdgeLength, TEXT("Max Edge Length"));
+	CommonTessellationOptions.NormalTolerance = CheckParameterValue(CommonTessellationOptions.NormalTolerance, UE::DatasmithTessellation::MinTessellationAngle, TEXT("Max Angle"));
+
+	UE_LOG(LogDatasmithXMLPLMTranslator, Display, TEXT(" - Import parameters:"));
+	UE_LOG(LogDatasmithXMLPLMTranslator, Display, TEXT("     - ChordTolerance:     %f"), CommonTessellationOptions.ChordTolerance);
+	UE_LOG(LogDatasmithXMLPLMTranslator, Display, TEXT("     - MaxEdgeLength:      %f"), CommonTessellationOptions.MaxEdgeLength);
+	UE_LOG(LogDatasmithXMLPLMTranslator, Display, TEXT("     - MaxNormalAngle:     %f"), CommonTessellationOptions.NormalTolerance);
+
+	FString StitchingTechnique;
+	switch ((EStitchingTechnique)CommonTessellationOptions.StitchingTechnique)
+	{
+	case EStitchingTechnique::StitchingHeal:
+		StitchingTechnique = TEXT("Heal");
+		break;
+	case EStitchingTechnique::StitchingSew:
+		StitchingTechnique = TEXT("Sew");
+		break;
+	default:
+		StitchingTechnique = TEXT("None");
+		break;
+	}
+	UE_LOG(LogDatasmithXMLPLMTranslator, Display, TEXT("     - StitchingTechnique: %s"), *StitchingTechnique);
+
+	if (!FImportParameters::bGDisableCADKernelTessellation)
+	{
+		UE_LOG(LogDatasmithXMLPLMTranslator, Display, TEXT("     - Stitching Options:"));
+		UE_LOG(LogDatasmithXMLPLMTranslator, Display, TEXT("         - ForceSew:              %s"), FImportParameters::bGStitchingForceSew ? TEXT("True") : TEXT("False"));
+		UE_LOG(LogDatasmithXMLPLMTranslator, Display, TEXT("         - RemoveThinFaces:       %s"), FImportParameters::bGStitchingRemoveThinFaces ? TEXT("True") : TEXT("False"));
+		UE_LOG(LogDatasmithXMLPLMTranslator, Display, TEXT("         - RemoveDuplicatedFaces: %s"), FImportParameters::bGStitchingRemoveDuplicatedFaces ? TEXT("True") : TEXT("False"));
+		UE_LOG(LogDatasmithXMLPLMTranslator, Display, TEXT("         - ForceFactor:           %f"), FImportParameters::GStitchingForceFactor);
+	}
+
 	OutScene->SetHost(TEXT("PlmXmlTranslator"));
 	OutScene->SetProductName(TEXT("PlmXml"));
 
     Importer = MakeUnique<FDatasmithPlmXmlImporter>(OutScene);
 
-	const FString& FilePath = GetSource().GetSourceFile();
 	if (!Importer->OpenFile(FilePath, GetSource(), CommonTessellationOptions))
 	{
 		return false;
