[NO NON-PII COMMIT MESSAGE]

diff --git a/Engine/Plugins/Enterprise/DatasmithCADImporter/Source/DatasmithCADTranslator/Private/DatasmithSceneGraphBuilder.cpp b/Engine/Plugins/Enterprise/DatasmithCADImporter/Source/DatasmithCADTranslator/Private/DatasmithSceneGraphBuilder.cpp
--- a/Engine/Plugins/Enterprise/DatasmithCADImporter/Source/DatasmithCADTranslator/Private/DatasmithSceneGraphBuilder.cpp
+++ b/Engine/Plugins/Enterprise/DatasmithCADImporter/Source/DatasmithCADTranslator/Private/DatasmithSceneGraphBuilder.cpp
@@ -514,7 +514,7 @@ TSharedPtr<IDatasmithMeshElement> FDatasmithSceneBaseGraphBuilder::FindOrAddMesh
 
 	BodyUuidToMeshElement.Add(Body.MeshActorUId, MeshElement);
 
-	FString BodyCachePath = CADLibrary::BuildCacheFilePath(*CachePath, TEXT("body"), Body.MeshActorUId, CADLibrary::EMesher::CADKernel);
+	FString BodyCachePath = CADLibrary::BuildCacheFilePath(*CachePath, TEXT("body"), Body.MeshActorUId, Body.Mesher);
 	MeshElement->SetFile(*BodyCachePath);
 
 	return MeshElement;
