[NO NON-PII COMMIT MESSAGE]

diff --git a/Engine/Plugins/Enterprise/DatasmithCADImporter/Source/DatasmithCADTranslator/Private/DatasmithCADTranslator.cpp b/Engine/Plugins/Enterprise/DatasmithCADImporter/Source/DatasmithCADTranslator/Private/DatasmithCADTranslator.cpp
--- a/Engine/Plugins/Enterprise/DatasmithCADImporter/Source/DatasmithCADTranslator/Private/DatasmithCADTranslator.cpp
+++ b/Engine/Plugins/Enterprise/DatasmithCADImporter/Source/DatasmithCADTranslator/Private/DatasmithCADTranslator.cpp
@@ -169,7 +169,7 @@ bool FDatasmithCADTranslator::LoadScene(TSharedRef<IDatasmithScene> DatasmithSce
 	{
 		if (Value < MinValue)
 		{
-			UE_LOG(LogCADTranslator, Warning, TEXT("%s value (%f) of tessellation parameters is smaller than the minimal value %f. It's value is modified to respect the limit"), ParameterName, Value, MinValue);
+			UE_LOG(LogCADTranslator, Warning, TEXT("%s value (%lf) of tessellation parameters is smaller than the minimal value %lf. It's value is modified to respect the limit"), ParameterName, Value, MinValue);
 			return MinValue;
 		}
 		return Value;
@@ -182,9 +182,9 @@ bool FDatasmithCADTranslator::LoadScene(TSharedRef<IDatasmithScene> DatasmithSce
 	ImportParameters.SetModelCoordinateSystem(FDatasmithUtils::EModelCoordSystem::ZUp_RightHanded);
 
 	UE_LOG(LogCADTranslator, Display, TEXT(" - Import parameters:"));
-	UE_LOG(LogCADTranslator, Display, TEXT("     - ChordTolerance:     %f"), ImportParameters.GetChordTolerance());
-	UE_LOG(LogCADTranslator, Display, TEXT("     - MaxEdgeLength:      %f"), ImportParameters.GetMaxEdgeLength());
-	UE_LOG(LogCADTranslator, Display, TEXT("     - MaxNormalAngle:     %f"), ImportParameters.GetMaxNormalAngle());
+	UE_LOG(LogCADTranslator, Display, TEXT("     - ChordTolerance:     %lf"), ImportParameters.GetChordTolerance());
+	UE_LOG(LogCADTranslator, Display, TEXT("     - MaxEdgeLength:      %lf"), ImportParameters.GetMaxEdgeLength());
+	UE_LOG(LogCADTranslator, Display, TEXT("     - MaxNormalAngle:     %lf"), ImportParameters.GetMaxNormalAngle());
 	FString StitchingTechnique;
 	switch(ImportParameters.GetStitchingTechnique())
 	{
diff --git a/Engine/Plugins/Enterprise/DatasmithCADImporter/Source/DatasmithDispatcher/Private/DatasmithDispatcher.cpp b/Engine/Plugins/Enterprise/DatasmithCADImporter/Source/DatasmithDispatcher/Private/DatasmithDispatcher.cpp
--- a/Engine/Plugins/Enterprise/DatasmithCADImporter/Source/DatasmithDispatcher/Private/DatasmithDispatcher.cpp
+++ b/Engine/Plugins/Enterprise/DatasmithCADImporter/Source/DatasmithDispatcher/Private/DatasmithDispatcher.cpp
@@ -149,7 +149,7 @@ TOptional<FTask> FDatasmithDispatcher::GetNextTask()
 	}
 
 	TaskPool[NextTaskIndex].State = ETaskState::Running;
-	UE_LOG(LogDatasmithDispatcher, Display, TEXT(" Launch %s (%d / %d)"), *TaskPool[NextTaskIndex].FileDescription.GetSourcePath(), NextTaskIndex, TaskPool.Num());
+	UE_LOG(LogDatasmithDispatcher, Display, TEXT("Launch %s (%d / %d)"), *TaskPool[NextTaskIndex].FileDescription.GetSourcePath(), NextTaskIndex, TaskPool.Num());
 
 	return TaskPool[NextTaskIndex++];
 }
@@ -182,9 +182,10 @@ void FDatasmithDispatcher::SetTaskState(int32 TaskIndex, ETaskState TaskState)
 		{
 			if (Task.Mesher == CADLibrary::EMesher::CADKernel)
 			{
-				UE_LOG(LogDatasmithDispatcher, Log, TEXT("   - Task failed with CADKernel: %s"), *Task.FileDescription.GetFileName());
-				UE_LOG(LogDatasmithDispatcher, Log, TEXT("      => Add task to process with Techsoft"));
+				UE_LOG(LogDatasmithDispatcher, Warning, TEXT("   - Task failed with CADKernel: %s"), *Task.FileDescription.GetFileName());
+				UE_LOG(LogDatasmithDispatcher, Warning, TEXT("      => Add task to process with Techsoft"));
 				AddTask(Task.FileDescription, CADLibrary::EMesher::TechSoft);
+				TaskState = ETaskState::UnTreated;
 			}
 
 			CompletedTaskCount++;
