[NO NON-PII COMMIT MESSAGE]

diff --git a/Engine/Plugins/Enterprise/DatasmithCADImporter/Source/CADInterfaces/Private/TechSoftFileParser.cpp b/Engine/Plugins/Enterprise/DatasmithCADImporter/Source/CADInterfaces/Private/TechSoftFileParser.cpp
--- a/Engine/Plugins/Enterprise/DatasmithCADImporter/Source/CADInterfaces/Private/TechSoftFileParser.cpp
+++ b/Engine/Plugins/Enterprise/DatasmithCADImporter/Source/CADInterfaces/Private/TechSoftFileParser.cpp
@@ -1888,7 +1888,7 @@ void FTechSoftFileParser::ExtractCoordinateSystem(const A3DRiCoordinateSystem* C
 	TUniqueTSObj<A3DRiCoordinateSystemData> CoordinateSystemData(CoordinateSystem);
 	if (CoordinateSystemData.IsValid())
 	{
-		ExtractTransformation3D(CoordinateSystemData->m_pTransformation, OutMetaData);
+		ExtractTransformation(CoordinateSystemData->m_pTransformation, OutMetaData);
 	}
 	else
 	{
