[NO NON-PII COMMIT MESSAGE]

diff --git a/Engine/Plugins/Enterprise/DatasmithCADImporter/Source/DatasmithDispatcher/Private/DatasmithDispatcher.cpp b/Engine/Plugins/Enterprise/DatasmithCADImporter/Source/DatasmithDispatcher/Private/DatasmithDispatcher.cpp
--- a/Engine/Plugins/Enterprise/DatasmithCADImporter/Source/DatasmithDispatcher/Private/DatasmithDispatcher.cpp
+++ b/Engine/Plugins/Enterprise/DatasmithCADImporter/Source/DatasmithDispatcher/Private/DatasmithDispatcher.cpp
@@ -22,6 +22,18 @@
 namespace DatasmithDispatcher
 {
 
+bool bCheckMemory = true;
+bool bEnableMemoryControl = true;
+FAutoConsoleVariableRef GCADTranslatorEnableMemoryControl(
+	TEXT("ds.CADTranslator.EnableMemoryControl"),
+	bEnableMemoryControl,
+	TEXT("\
+Enable/disable the control of the memory used by static meshes to avoid running out of memory.\n\
+if CheckMemoryUsage = true, a warning is displayed allowing the user to interrupt the import and have only data already processed or to continue with the risk of running out of memory.\n\
+Default value is true.\n"),
+ECVF_Default);
+
+
 static uint64 EstimationOfMemoryUsedByStaticMeshes = 0;
 static uint64 AvailableMemory = 0;
 static bool bDisplayWorkerWarning = true;
@@ -224,7 +236,7 @@ void FDatasmithDispatcher::Process(bool bWithProcessor)
 	}
 
 	EstimationOfMemoryUsedByStaticMeshes = 0;
-	bool bCheckMemory = CADLibrary::FImportParameters::bValidationProcess ? false : true;
+	bCheckMemory = bEnableMemoryControl & !CADLibrary::FImportParameters::bValidationProcess;
 
 	// This function is added in 5.1.1 but is cloned with the lines 370 to 387
 	TFunction<bool()> CanImportContinue = [&]() -> bool
@@ -401,7 +413,6 @@ void FDatasmithDispatcher::ProcessLocal()
 	using namespace CADLibrary;
 	EMesher DefaultMesher = FImportParameters::bGDisableCADKernelTessellation ? EMesher::TechSoft : EMesher::CADKernel;
 
-	bool bCheckMemory = true;
 	while (TOptional<FTask> Task = GetNextTask())
 	{
 		FFileDescriptor& FileDescription = Task->FileDescription;
