To fix it, for JT file, add in the achive name the import mode
Also update FCADFileData::GetGeomFileHash() for this case

diff --git a/Engine/Plugins/Enterprise/DatasmithCADImporter/Source/CADInterfaces/Public/CADFileData.h b/Engine/Plugins/Enterprise/DatasmithCADImporter/Source/CADInterfaces/Public/CADFileData.h
--- a/Engine/Plugins/Enterprise/DatasmithCADImporter/Source/CADInterfaces/Public/CADFileData.h
+++ b/Engine/Plugins/Enterprise/DatasmithCADImporter/Source/CADInterfaces/Public/CADFileData.h
@@ -51,15 +51,22 @@ public:
 	{
 		if (!GeomFileHash)
 		{
-			GeomFileHash = GetSceneFileHash();
-			GeomFileHash = HashCombine(GeomFileHash, GetTypeHash(ImportParameters));
-			GeomFileHash = HashCombine(GeomFileHash, ::GetTypeHash(FImportParameters::bGSewMeshIfNeeded));
-			GeomFileHash = HashCombine(GeomFileHash, ::GetTypeHash(FImportParameters::bGRemoveDuplicatedTriangle));
-			GeomFileHash = HashCombine(GeomFileHash, ::GetTypeHash(FImportParameters::GStitchingTolerance));
-			GeomFileHash = HashCombine(GeomFileHash, ::GetTypeHash(FImportParameters::bGStitchingForceSew));
-			GeomFileHash = HashCombine(GeomFileHash, ::GetTypeHash(FImportParameters::bGStitchingRemoveThinFaces));
-			GeomFileHash = HashCombine(GeomFileHash, ::GetTypeHash(FImportParameters::bGStitchingRemoveDuplicatedFaces));
-			GeomFileHash = HashCombine(GeomFileHash, ::GetTypeHash(FImportParameters::GStitchingForceFactor));
+			if (FileDescription.GetFileFormat() == ECADFormat::JT && FImportParameters::bGPreferJtFileEmbeddedTessellation)
+			{
+				GeomFileHash = GetSceneFileHash();
+			}
+			else
+			{
+				GeomFileHash = GetSceneFileHash();
+				GeomFileHash = HashCombine(GeomFileHash, GetTypeHash(ImportParameters));
+				GeomFileHash = HashCombine(GeomFileHash, ::GetTypeHash(FImportParameters::bGSewMeshIfNeeded));
+				GeomFileHash = HashCombine(GeomFileHash, ::GetTypeHash(FImportParameters::bGRemoveDuplicatedTriangle));
+				GeomFileHash = HashCombine(GeomFileHash, ::GetTypeHash(FImportParameters::GStitchingTolerance));
+				GeomFileHash = HashCombine(GeomFileHash, ::GetTypeHash(FImportParameters::bGStitchingForceSew));
+				GeomFileHash = HashCombine(GeomFileHash, ::GetTypeHash(FImportParameters::bGStitchingRemoveThinFaces));
+				GeomFileHash = HashCombine(GeomFileHash, ::GetTypeHash(FImportParameters::bGStitchingRemoveDuplicatedFaces));
+				GeomFileHash = HashCombine(GeomFileHash, ::GetTypeHash(FImportParameters::GStitchingForceFactor));
+			}
 		}
 		return GeomFileHash;
 	}
@@ -117,7 +124,12 @@ public:
 	{
 		if (IsCacheDefined())
 		{
-			return CADLibrary::BuildCadCachePath(*CachePath, FileDescription.GetDescriptorHash());
+			uint32 CADFileHash = FileDescription.GetDescriptorHash();
+			if (FileDescription.GetFileFormat() == ECADFormat::JT && FImportParameters::bGPreferJtFileEmbeddedTessellation)
+			{
+				CADFileHash = HashCombine(CADFileHash, 1);
+			}
+			return CADLibrary::BuildCadCachePath(*CachePath, CADFileHash);
 		}
 		return FString();
 	}
