[NO NON-PII COMMIT MESSAGE]

diff --git a/Engine/Plugins/Enterprise/DatasmithCADImporter/Source/CADInterfaces/Private/TechSoftBridge.cpp b/Engine/Plugins/Enterprise/DatasmithCADImporter/Source/CADInterfaces/Private/TechSoftBridge.cpp
--- a/Engine/Plugins/Enterprise/DatasmithCADImporter/Source/CADInterfaces/Private/TechSoftBridge.cpp
+++ b/Engine/Plugins/Enterprise/DatasmithCADImporter/Source/CADInterfaces/Private/TechSoftBridge.cpp
@@ -215,18 +215,19 @@ UE::CADKernel::FBody* FTechSoftBridge::GetBody(A3DRiBrepModel* A3DBRepModel)
 	return nullptr;
 }
 
-UE::CADKernel::FBody* FTechSoftBridge::AddBody(A3DRiBrepModel* A3DBRepModel, TMap<FString, FString> MetaData, const double InBodyScale)
+UE::CADKernel::FBody* FTechSoftBridge::AddBody(A3DRiBrepModel* A3DBRepModel, FArchiveCADObject& ArchiveBody)
 {
 #ifdef CADKERNEL_DEV
 	UE::CADKernel::FCADFileReport::Get().BodyCount++;
 #endif
 
 	//UE::CADKernel working unit is mm
-	BodyScale = InBodyScale * 10.;
+	BodyScale = ArchiveBody.Unit * 10.;
 
 	FArchiveCADObject BRepMetaData;
 	Parser.ExtractMetaData(A3DBRepModel, BRepMetaData);
-	FString* Name = MetaData.Find(TEXT("Name"));
+
+	FString* Name = ArchiveBody.MetaData.Find(TEXT("Name"));
 	if (Name != nullptr)
 	{
 		BRepMetaData.MetaData.FindOrAdd(TEXT("Name")) = *Name;
@@ -240,6 +241,8 @@ UE::CADKernel::FBody* FTechSoftBridge::AddBody(A3DRiBrepModel* A3DBRepModel, TMa
 	TSharedRef<UE::CADKernel::FBody> Body = UE::CADKernel::FEntity::MakeShared<UE::CADKernel::FBody>();
 	AddMetadata(BRepMetaData, *Body);
 
+	Body->SetDisplayData(ArchiveBody.ColorUId, ArchiveBody.MaterialUId);
+
 	TUniqueTSObj<A3DRiBrepModelData> BRepModelData(A3DBRepModel);
 	if (BRepModelData.IsValid())
 	{
@@ -325,6 +328,8 @@ void FTechSoftBridge::TraverseShell(const A3DTopoShell* A3DShell, TSharedRef<UE:
 #ifdef CADKERNEL_DEV
 	UE::CADKernel::FCADFileReport::Get().ShellCount++;
 #endif
+
+	Shell->SetDisplayData(*Body);
 	AddMetadata(MetaData, *Shell);
 
 	TUniqueTSObj<A3DTopoShellData> ShellData(A3DShell);
@@ -932,6 +937,8 @@ void FTechSoftBridge::AddFace(const A3DTopoFace* A3DFace, UE::CADKernel::EOrient
 	}
 
 	AddMetadata(MetaData, *Face);
+	Face->SetDisplayDataIfUndefined(*Shell);
+
 	Face->SetHostId(ShellIndex);
 	Shell->Add(Face, Orientation);
 }
@@ -1415,8 +1422,7 @@ void FTechSoftBridge::AddMetadata(FArchiveCADObject& MetaData, UE::CADKernel::FT
 	{
 		Entity.SetName(*Name);
 	}
-	Entity.SetColorId(MetaData.ColorUId);
-	Entity.SetMaterialId(MetaData.MaterialUId);
+	Entity.SetDisplayData(MetaData.ColorUId, MetaData.MaterialUId);
 }
 
 }
diff --git a/Engine/Plugins/Enterprise/DatasmithCADImporter/Source/CADInterfaces/Private/TechSoftBridge.h b/Engine/Plugins/Enterprise/DatasmithCADImporter/Source/CADInterfaces/Private/TechSoftBridge.h
--- a/Engine/Plugins/Enterprise/DatasmithCADImporter/Source/CADInterfaces/Private/TechSoftBridge.h
+++ b/Engine/Plugins/Enterprise/DatasmithCADImporter/Source/CADInterfaces/Private/TechSoftBridge.h
@@ -183,7 +183,7 @@ private:
 
 public:
 	FTechSoftBridge(FTechSoftFileParser& InParser, UE::CADKernel::FSession& InSession);
-	UE::CADKernel::FBody* AddBody(A3DRiBrepModel* A3DBRepModel, TMap<FString, FString> MetaData, const double InBodyScale);
+	UE::CADKernel::FBody* AddBody(A3DRiBrepModel* A3DBRepModel, FArchiveCADObject& ArchiveBody);
 	UE::CADKernel::FBody* GetBody(A3DRiBrepModel* A3DBRepModel);
 	const A3DRiBrepModel* GetA3DBody(UE::CADKernel::FBody* BRepModel);
 
diff --git a/Engine/Plugins/Enterprise/DatasmithCADImporter/Source/CADInterfaces/Private/TechSoftFileParserCADKernelTessellator.cpp b/Engine/Plugins/Enterprise/DatasmithCADImporter/Source/CADInterfaces/Private/TechSoftFileParserCADKernelTessellator.cpp
--- a/Engine/Plugins/Enterprise/DatasmithCADImporter/Source/CADInterfaces/Private/TechSoftFileParserCADKernelTessellator.cpp
+++ b/Engine/Plugins/Enterprise/DatasmithCADImporter/Source/CADInterfaces/Private/TechSoftFileParserCADKernelTessellator.cpp
@@ -111,7 +111,7 @@ void FTechSoftFileParserCADKernelTessellator::SewAndMesh(TArray<A3DRiRepresentat
 		}
 		FArchiveBody& Body = SceneGraph.GetBody(*BodyIndex);
 
-		FBody* CADKernelBody = TechSoftBridge.AddBody(Representation, Body.MetaData, Body.Unit);
+		FBody* CADKernelBody = TechSoftBridge.AddBody(Representation, Body);
 		if (CADKernelBody == nullptr)
 		{
 			const EFailureReason FailureReason = TechSoftBridge.GetFailureReason();
@@ -252,7 +252,7 @@ void FTechSoftFileParserCADKernelTessellator::GenerateBodyMesh(A3DRiRepresentati
 
 	FTechSoftBridge TechSoftBridge(*this, CADKernelSession);
 
-	FBody* CADKernelBody = TechSoftBridge.AddBody(Representation, ArchiveBody.MetaData, ArchiveBody.Unit);
+	FBody* CADKernelBody = TechSoftBridge.AddBody(Representation, ArchiveBody);
 	if (CADKernelBody == nullptr)
 	{
 		ArchiveBody.Delete();
@@ -264,6 +264,8 @@ void FTechSoftFileParserCADKernelTessellator::GenerateBodyMesh(A3DRiRepresentati
 		return;
 	}
 
+	CADKernelBody->SetDisplayData(ArchiveBody.ColorUId, ArchiveBody.MaterialUId);
+
 	if (CADFileData.GetImportParameters().GetStitchingTechnique() == StitchingHeal)
 	{
 		UE::CADKernel::FTopomakerOptions TopomakerOptions((UE::CADKernel::ESewOption)SewOption::GetFromImportParameters(), GeometricTolerance, FImportParameters::GStitchingForceFactor);
@@ -282,8 +284,6 @@ void FTechSoftFileParserCADKernelTessellator::MeshAndGetTessellation(UE::CADKern
 	using namespace UE::CADKernel;
 
 	FBodyMesh& BodyMesh = CADFileData.AddBodyMesh(ArchiveBody.Id, ArchiveBody);
-	ArchiveBody.ColorFaceSet = BodyMesh.ColorSet;
-	ArchiveBody.MaterialFaceSet = BodyMesh.MaterialSet;
 
 	// Save Body in CADKernelArchive file for re-tessellation
 	if (CADFileData.IsCacheDefined())
diff --git a/Engine/Plugins/Enterprise/DatasmithCADImporter/Source/CADLibrary/Private/CADKernelTools.cpp b/Engine/Plugins/Enterprise/DatasmithCADImporter/Source/CADLibrary/Private/CADKernelTools.cpp
--- a/Engine/Plugins/Enterprise/DatasmithCADImporter/Source/CADLibrary/Private/CADKernelTools.cpp
+++ b/Engine/Plugins/Enterprise/DatasmithCADImporter/Source/CADLibrary/Private/CADKernelTools.cpp
@@ -111,8 +111,8 @@ bool FillMesh(FMeshConversionContext& Context, UE::CADKernel::FModelMesh& ModelM
 	for (const FFaceMesh* FaceMesh : ModelMesh.GetFaceMeshes())
 	{
 		const FTopologicalFace& Face = (const FTopologicalFace&) FaceMesh->GetGeometricEntity();
-		// we assume that face has only color
-		MaterialToPolygonGroupMapping.Add(Face.GetColorId(), INDEX_NONE);
+		uint32 MaterialId = Face.GetMaterialId() ? Face.GetMaterialId() : Face.GetColorId();
+		MaterialToPolygonGroupMapping.Add(MaterialId, INDEX_NONE);
 	}
 
 	// Add to the mesh, a polygon groups per material
diff --git a/Engine/Source/Runtime/Datasmith/CADKernel/Public/CADKernel/Topo/TopologicalShapeEntity.h b/Engine/Source/Runtime/Datasmith/CADKernel/Public/CADKernel/Topo/TopologicalShapeEntity.h
--- a/Engine/Source/Runtime/Datasmith/CADKernel/Public/CADKernel/Topo/TopologicalShapeEntity.h
+++ b/Engine/Source/Runtime/Datasmith/CADKernel/Public/CADKernel/Topo/TopologicalShapeEntity.h
@@ -136,6 +136,35 @@ public:
 		return Dictionary.GetMaterialId();
 	}
 
+	void SetDisplayData(const uint32& InColorId, const uint32& InMaterialId)
+	{
+		if (InColorId)
+		{
+			Dictionary.SetColorId(InColorId);
+		}
+		if(InMaterialId)
+		{
+			Dictionary.SetMaterialId(InMaterialId);
+		}
+	}
+
+	void SetDisplayData(const FTopologicalShapeEntity& DisplayData)
+	{
+		SetDisplayData(DisplayData.GetColorId(), DisplayData.GetMaterialId());
+	}
+
+	void SetDisplayDataIfUndefined(const FTopologicalShapeEntity& DefaultDisplayData)
+	{
+		if (Dictionary.GetColorId() == 0 && DefaultDisplayData.GetColorId())
+		{
+			Dictionary.SetColorId(DefaultDisplayData.GetColorId());
+		}
+		if (Dictionary.GetMaterialId() == 0 && DefaultDisplayData.GetMaterialId())
+		{
+			Dictionary.SetMaterialId(DefaultDisplayData.GetMaterialId());
+		}
+	}
+
 	void SetPatchId(int32 InPatchId)
 	{
 		Dictionary.SetPatchId(InPatchId);
