Some tools/formats allows to override some occurrences property (visiblily, position, color...) of the instances of a reference (e.g. only the 4th occurrence of C is hidden in the graph below)

      A
      | - B
      |   | - C (show)
      |   | - C (show)
      |
      | - B
      |   | - C (show)
      |   | - C (hide)

diff --git a/Engine/Plugins/Enterprise/DatasmithCADImporter/Source/CADInterfaces/Private/CADSceneGraph.cpp b/Engine/Plugins/Enterprise/DatasmithCADImporter/Source/CADInterfaces/Private/CADSceneGraph.cpp
--- a/Engine/Plugins/Enterprise/DatasmithCADImporter/Source/CADInterfaces/Private/CADSceneGraph.cpp
+++ b/Engine/Plugins/Enterprise/DatasmithCADImporter/Source/CADInterfaces/Private/CADSceneGraph.cpp
@@ -15,11 +15,8 @@ FArchive& operator<<(FArchive& Ar, FArchiveCADObject& Object)
 	Ar << Object.Label;
 	Ar << Object.MetaData;
 	Ar << Object.TransformMatrix;
-	Ar << Object.ColorUId;
-	Ar << Object.MaterialUId;
-	Ar << Object.bIsRemoved;
-	Ar << Object.Inheritance;
 	Ar << Object.Unit;
+	Ar << (FArchiveGraphicProperties&) Object;
 	return Ar;
 }
 
@@ -36,7 +33,7 @@ bool FArchiveCADObject::SetNameWithAttributeValue(const TCHAR* Key)
 
 FArchive& operator<<(FArchive& Ar, FArchiveInstance& Instance) 
 {
-	Ar << (FArchiveCADObject&) Instance;
+	Ar << (FArchiveWithOverridenChildren&) Instance;
 	Ar << Instance.ReferenceNodeId;
 	Ar << Instance.bIsExternalReference;
 	return Ar;
@@ -56,7 +53,20 @@ FArchive& operator<<(FArchive& Ar, FArchiveUnloadedReference& Unloaded)
 	return Ar;
 }
 
-FArchive& operator<<(FArchive& Ar, FArchiveBody& Body) 
+FArchive& operator<<(FArchive& Ar, FArchiveOverrideOccurrence& OverrideOccurrence)
+{
+	Ar << (FArchiveWithOverridenChildren&) OverrideOccurrence;
+	return Ar;
+}
+
+FArchive& operator<<(FArchive& Ar, FArchiveWithOverridenChildren& WithOverridenChildren)
+{
+	Ar << (FArchiveCADObject&) WithOverridenChildren;
+	Ar << WithOverridenChildren.OverridenChildren;
+	return Ar;
+}
+
+FArchive& operator<<(FArchive& Ar, FArchiveBody& Body)
 {
 	Ar << (FArchiveCADObject&) Body;
 	Ar << Body.MaterialFaceSet;
@@ -86,6 +96,7 @@ FArchive& operator<<(FArchive& Ar, FArchiveMaterial& Material)
 
 FArchive& operator<<(FArchive& Ar, FArchiveSceneGraph& SceneGraph)
 {
+	Ar << SceneGraph.SceneGraphId;
 	Ar << SceneGraph.CADFileName;
 	Ar << SceneGraph.ArchiveFileName;
 	Ar << SceneGraph.FullPath;
@@ -97,6 +108,7 @@ FArchive& operator<<(FArchive& Ar, FArchiveSceneGraph& SceneGraph)
 	Ar << SceneGraph.Instances;
 	Ar << SceneGraph.References;
 	Ar << SceneGraph.UnloadedReferences;
+	Ar << SceneGraph.OverrideOccurrences;
 	Ar << SceneGraph.Bodies;
 
 	Ar << SceneGraph.CADIdToIndex;
@@ -135,7 +147,7 @@ FArchiveInstance& FArchiveSceneGraph::AddInstance(const FArchiveCADObject& Paren
 	return Instances[Index];
 }
 
-FArchiveInstance& FArchiveSceneGraph::GetInstance(FCadId CadId)
+FArchiveInstance& FArchiveSceneGraph::GetInstance(const FCadId CadId)
 {
 	ensure(CADIdToIndex.IsValidIndex(CadId));
 	const int32 Index = CADIdToIndex[CadId];
@@ -199,7 +211,7 @@ FArchiveReference& FArchiveSceneGraph::AddReference(FArchiveInstance& Parent)
 	return References[Index];
 }
 
-FArchiveReference& FArchiveSceneGraph::AddOccurence(FArchiveReference& Parent)
+FArchiveReference& FArchiveSceneGraph::AddOccurrence(FArchiveReference& Parent)
 {
 	FArchiveInstance& Instance = AddInstance(Parent);
 	Parent.AddChild(Instance.Id);
@@ -208,14 +220,14 @@ FArchiveReference& FArchiveSceneGraph::AddOccurence(FArchiveReference& Parent)
 	return Reference;
 }
 
-void FArchiveSceneGraph::RemoveLastOccurence()
+void FArchiveSceneGraph::RemoveLastOccurrence()
 {
 	RemoveLastReference();
 	RemoveLastInstance();
 }
 
 
-FArchiveReference& FArchiveSceneGraph::GetReference(FCadId CadId)
+FArchiveReference& FArchiveSceneGraph::GetReference(const FCadId CadId)
 {
 	ensure(CADIdToIndex.IsValidIndex(CadId));
 	const int32 Index = CADIdToIndex[CadId];
@@ -223,7 +235,7 @@ FArchiveReference& FArchiveSceneGraph::GetReference(FCadId CadId)
 	return References[Index];
 }
 
-bool FArchiveSceneGraph::IsAReference(FCadId CadId) const
+bool FArchiveSceneGraph::IsAReference(const FCadId CadId) const
 {
 	if (!CADIdToIndex.IsValidIndex(CadId))
 	{
@@ -248,6 +260,11 @@ void FArchiveSceneGraph::RemoveLastReference()
 	CADIdToIndex.SetNum(CADIdToIndex.Num() - 1, false);
 }
 
+void FArchiveWithOverridenChildren::AddOverridenChild(const FCadId ChildId)
+{
+	OverridenChildren.Add(ChildId);
+}
+
 void FArchiveReference::AddChild(const FCadId ChildId)
 {
 	Children.Add(ChildId);
@@ -309,7 +326,7 @@ FArchiveUnloadedReference& FArchiveSceneGraph::AddUnloadedReference(FArchiveInst
 	return UnloadedReferences[Index];
 }
 
-FArchiveUnloadedReference& FArchiveSceneGraph::GetUnloadedReference(FCadId CadId)
+FArchiveUnloadedReference& FArchiveSceneGraph::GetUnloadedReference(const FCadId CadId)
 {
 	ensure(CADIdToIndex.IsValidIndex(CadId));
 	const int32 Index = CADIdToIndex[CadId];
@@ -317,7 +334,27 @@ FArchiveUnloadedReference& FArchiveSceneGraph::GetUnloadedReference(FCadId CadId
 	return UnloadedReferences[Index];
 }
 
-bool FArchiveSceneGraph::IsAUnloadedReference(FCadId CadId) const
+FArchiveOverrideOccurrence& FArchiveSceneGraph::AddOverrideOccurrence(FArchiveWithOverridenChildren& Parent)
+{
+	ensure(OverrideOccurrences.Num() < OverrideOccurrences.Max());
+
+	const int32 ReferenceId = LastEntityId++;
+
+	int32 Index = OverrideOccurrences.Emplace(ReferenceId, Parent);
+	CADIdToIndex.Add(Index);
+
+	return OverrideOccurrences[Index];
+}
+
+FArchiveOverrideOccurrence& FArchiveSceneGraph::GetOverrideOccurrence(const FCadId CadId)
+{
+	ensure(CADIdToIndex.IsValidIndex(CadId));
+	const int32 Index = CADIdToIndex[CadId];
+	ensure(OverrideOccurrences.IsValidIndex(Index));
+	return OverrideOccurrences[Index];
+}
+
+bool FArchiveSceneGraph::IsAUnloadedReference(const FCadId CadId) const
 {
 	if (!CADIdToIndex.IsValidIndex(CadId))
 	{
@@ -353,7 +390,7 @@ FArchiveBody& FArchiveSceneGraph::AddBody(FArchiveReference& Parent, EMesher InM
 	return Bodies[Index];
 }
 
-FArchiveBody& FArchiveSceneGraph::GetBody(FCadId CadId)
+FArchiveBody& FArchiveSceneGraph::GetBody(const FCadId CadId)
 {
 	ensure(CADIdToIndex.IsValidIndex(CadId));
 	const int32 Index = CADIdToIndex[CadId];
@@ -361,7 +398,7 @@ FArchiveBody& FArchiveSceneGraph::GetBody(FCadId CadId)
 	return Bodies[Index];
 }
 
-bool FArchiveSceneGraph::IsABody(FCadId CadId) const
+bool FArchiveSceneGraph::IsABody(const FCadId CadId) const
 {
 	if (!CADIdToIndex.IsValidIndex(CadId))
 	{
@@ -387,11 +424,6 @@ void FArchiveSceneGraph::RemoveLastBody()
 
 void FArchiveSceneGraph::AddExternalReferenceFile(const FArchiveUnloadedReference& Reference)
 {
-	if (Reference.ExternalFile.GetFileName().IsEmpty())
-	{
-		return;
-	}
-
 	const int32 Index = CADIdToIndex[Reference.Id];
 	ensure(ExternalReferenceFiles.Add(Reference.ExternalFile) == Index);
 }
diff --git a/Engine/Plugins/Enterprise/DatasmithCADImporter/Source/CADInterfaces/Private/TechSoftFileParser.cpp b/Engine/Plugins/Enterprise/DatasmithCADImporter/Source/CADInterfaces/Private/TechSoftFileParser.cpp
--- a/Engine/Plugins/Enterprise/DatasmithCADImporter/Source/CADInterfaces/Private/TechSoftFileParser.cpp
+++ b/Engine/Plugins/Enterprise/DatasmithCADImporter/Source/CADInterfaces/Private/TechSoftFileParser.cpp
@@ -220,7 +220,7 @@ void SetIOOption(A3DImport& Importer)
 	Importer.m_sLoadData.m_sGeneral.m_bReadWireframes = A3D_FALSE;
 	Importer.m_sLoadData.m_sGeneral.m_bReadPmis = A3D_FALSE;
 	Importer.m_sLoadData.m_sGeneral.m_bReadAttributes = A3D_TRUE;
-	Importer.m_sLoadData.m_sGeneral.m_bReadHiddenObjects = A3D_FALSE;
+	Importer.m_sLoadData.m_sGeneral.m_bReadHiddenObjects = A3D_TRUE;
 	Importer.m_sLoadData.m_sGeneral.m_bReadConstructionAndReferences = A3D_FALSE;
 	Importer.m_sLoadData.m_sGeneral.m_bReadActiveFilter = A3D_FALSE;
 	Importer.m_sLoadData.m_sGeneral.m_eReadingMode2D3D = kA3DRead_3D;
@@ -639,6 +639,7 @@ void FTechSoftFileParser::CountUnderModel()
 		else
 		{
 			CountUnderOccurrence(ModelFileData->m_ppPOccurrences[Index]);
+			CountUnderOverrideOccurrence(ModelFileData->m_ppPOccurrences[Index]);
 		}
 	}
 
@@ -734,7 +735,7 @@ void FTechSoftFileParser::TraverseConfigurationSet(const A3DAsmProductOccurrence
 
 	if (ConfigurationToLoad.IsEmpty())
 	{
-		// no default configuration, traverse the first occurence
+		// no default configuration, traverse the first occurrence
 		for (unsigned int Index = 0; Index < ConfigurationSetData->m_uiPOccurrencesSize; ++Index)
 		{
 			ConfigurationData.FillFrom(ConfigurationSetData->m_ppPOccurrences[Index]);
@@ -790,6 +791,7 @@ void FTechSoftFileParser::CountUnderConfigurationSet(const A3DAsmProductOccurren
 			if (bIsConfigurationToLoad)
 			{
 				CountUnderOccurrence(ConfigurationSetData->m_ppPOccurrences[Index]);
+				CountUnderOverrideOccurrence(ConfigurationSetData->m_ppPOccurrences[Index]);
 				return;
 			}
 		}
@@ -809,6 +811,7 @@ void FTechSoftFileParser::CountUnderConfigurationSet(const A3DAsmProductOccurren
 			if (ConfigurationData->m_uiProductFlags & A3D_PRODUCT_FLAG_CONFIG)
 			{
 				CountUnderOccurrence(ConfigurationSetData->m_ppPOccurrences[Index]);
+				CountUnderOverrideOccurrence(ConfigurationSetData->m_ppPOccurrences[Index]);
 			}
 		}
 	}
@@ -868,7 +871,7 @@ void FTechSoftFileParser::ProcessUnloadedReference(const FArchiveInstance& Insta
 
 void FTechSoftFileParser::TraverseOccurrence(const A3DAsmProductOccurrence* OccurrencePtr, FArchiveReference& ParentReference)
 {
-	// first product occurence with m_pPart != nullptr || m_uiPOccurrencesSize > 0
+	// first product occurrence with m_pPart != nullptr || m_uiPOccurrencesSize > 0
 	const A3DAsmProductOccurrence* CachedOccurrencePtr = OccurrencePtr;
 	TUniqueTSObj<A3DAsmProductOccurrenceData> OccurrenceData(OccurrencePtr);
 	if (!OccurrenceData.IsValid())
@@ -889,7 +892,7 @@ void FTechSoftFileParser::TraverseOccurrence(const A3DAsmProductOccurrence* Occu
 	ExtractMetaData(OccurrencePtr, Instance);
 	Instance.DefineGraphicsPropertiesFromNoOverwrite(ParentReference);
 
-	if (Instance.bIsRemoved || !Instance.bShow)
+	if (Instance.bIsRemoved)
 	{
 		SceneGraph.RemoveLastInstance();
 		return;
@@ -978,6 +981,48 @@ void FTechSoftFileParser::TraverseOccurrence(const A3DAsmProductOccurrence* Occu
 	{
 		ReferenceCache.Add(ReferencePtr, Instance.ReferenceNodeId);
 	}
+
+	for (uint32 Index = 0; Index < OccurrenceData->m_uiPOccurrencesSize; ++Index)
+	{
+		ExtractOverrideOccurrenceSubtree(OccurrenceData->m_ppPOccurrences[Index], Instance);
+	}
+}
+
+void FTechSoftFileParser::ExtractOverrideOccurrenceSubtree(const A3DAsmProductOccurrence* OccurrencePtr, FArchiveWithOverridenChildren& Parent)
+{
+	TUniqueTSObj<A3DAsmProductOccurrenceData> OccurrenceData(OccurrencePtr);
+
+	FArchiveOverrideOccurrence& OverrideOccurrence = SceneGraph.AddOverrideOccurrence(Parent);
+	ExtractMetaData(OccurrencePtr, OverrideOccurrence);
+
+	Parent.AddOverridenChild(OverrideOccurrence.Id);
+
+	FArchiveReference Reference;
+	BuildInstanceName(OverrideOccurrence, Reference);
+
+	A3DMiscTransformation* Transform = OccurrenceData->m_pLocation;
+	ExtractTransformation(Transform, OverrideOccurrence);
+
+	for (uint32 Index = 0; Index < OccurrenceData->m_uiPOccurrencesSize; ++Index)
+	{
+		ExtractOverrideOccurrenceSubtree(OccurrenceData->m_ppPOccurrences[Index], OverrideOccurrence);
+	}
+}
+
+void FTechSoftFileParser::CountUnderOverrideOccurrence(const A3DAsmProductOccurrence* Occurrence)
+{
+	TUniqueTSObj<A3DAsmProductOccurrenceData> OccurrenceData(Occurrence);
+	if (Occurrence && OccurrenceData.IsValid())
+	{
+		ComponentCount[EComponentType::OverriddeOccurence]++;
+
+		uint32 ChildrenCount = OccurrenceData->m_uiPOccurrencesSize;
+		A3DAsmProductOccurrence** Children = OccurrenceData->m_ppPOccurrences;
+		for (uint32 Index = 0; Index < ChildrenCount; ++Index)
+		{
+			CountUnderOverrideOccurrence(Children[Index]);
+		}
+	}
 }
 
 void FTechSoftFileParser::ProcessReference(const A3DAsmProductOccurrence* OccurrencePtr, FArchiveInstance& Instance, FArchiveReference& Reference)
@@ -1083,6 +1128,7 @@ void FTechSoftFileParser::CountUnderOccurrence(const A3DAsmProductOccurrence* Oc
 		for (uint32 Index = 0; Index < ChildrenCount; ++Index)
 		{
 			CountUnderOccurrence(Children[Index]);
+			CountUnderOverrideOccurrence(Children[Index]);
 		}
 	}
 }
@@ -1262,7 +1308,7 @@ void FTechSoftFileParser::TraverseRepresentationSet(const A3DRiSet* Representati
 	}
 
 	FCadId RepresentationSetId = 0;
-	FArchiveReference& RepresentationSet = SceneGraph.AddOccurence(Parent);
+	FArchiveReference& RepresentationSet = SceneGraph.AddOccurrence(Parent);
 	ExtractMetaData(RepresentationSetPtr, RepresentationSet);
 	RepresentationSet.DefineGraphicsPropertiesFromNoOverwrite(Parent);
 	BuildRepresentationSetName(RepresentationSet, Parent);
@@ -1270,7 +1316,7 @@ void FTechSoftFileParser::TraverseRepresentationSet(const A3DRiSet* Representati
 	if (RepresentationSet.bIsRemoved || !RepresentationSet.bShow)
 	{
 		Parent.RemoveLastChild();
-		SceneGraph.RemoveLastOccurence();
+		SceneGraph.RemoveLastOccurrence();
 		return;
 	}
 
@@ -1425,7 +1471,7 @@ void FTechSoftFileParser::BuildReferenceName(FArchiveCADObject& ReferenceData)
 	}
 }
 
-void FTechSoftFileParser::BuildInstanceName(FArchiveInstance& InstanceData, const FArchiveReference& Parent)
+void FTechSoftFileParser::BuildInstanceName(FArchiveCADObject& InstanceData, const FArchiveReference& Parent)
 {
 	TMap<FString, FString>& MetaData = InstanceData.MetaData;
 
diff --git a/Engine/Plugins/Enterprise/DatasmithCADImporter/Source/CADInterfaces/Private/TechSoftFileParser.h b/Engine/Plugins/Enterprise/DatasmithCADImporter/Source/CADInterfaces/Private/TechSoftFileParser.h
--- a/Engine/Plugins/Enterprise/DatasmithCADImporter/Source/CADInterfaces/Private/TechSoftFileParser.h
+++ b/Engine/Plugins/Enterprise/DatasmithCADImporter/Source/CADInterfaces/Private/TechSoftFileParser.h
@@ -24,6 +24,7 @@ class FArchiveMaterial;
 class FArchiveReference;
 class FArchiveSceneGraph;
 class FArchiveUnloadedReference;
+class FArchiveWithOverridenChildren;
 class FCADFileData;
 
 class FTechSoftFileParser : public ICADFileParser
@@ -79,6 +80,7 @@ private:
 	void CountUnderPartDefinition(const A3DAsmPartDefinition* PartDefinition);
 	void CountUnderRepresentationItem(const A3DRiRepresentationItem* RepresentationItem);
 	void CountUnderRepresentationSet(const A3DRiSet* RepresentationSet);
+	void CountUnderOverrideOccurrence(const A3DAsmProductOccurrence* Occurrence);
 
 	void ReserveCADFileData();
 
@@ -112,14 +114,19 @@ private:
 	// MetaData
 	void ExtractSpecificMetaData(const A3DAsmProductOccurrence* Occurrence, FArchiveCADObject& OutMetaData);
 
-	void BuildInstanceName(FArchiveInstance& MetaData, const FArchiveReference& Parent);
+	void BuildInstanceName(FArchiveCADObject& MetaData, const FArchiveReference& Parent);
 	void BuildReferenceName(FArchiveCADObject& Reference);
 	void BuildPartName(FArchiveCADObject& Part);
 	void BuildBodyName(FArchiveBody& MetaData, const FArchiveReference& Parent);
 	void BuildRepresentationSetName(FArchiveCADObject& MetaData, const FArchiveReference& Parent);
-	// Graphic properties
+
 	void ExtractGraphicProperties(const A3DGraphics* Graphics, FArchiveCADObject& OutMetaData);
 
+	/**
+	 * Recursive method to extract the subtree of override occurrences of the children tree of an instance 
+	 */
+	void ExtractOverrideOccurrenceSubtree(const A3DAsmProductOccurrence* OccurrencePtr, FArchiveWithOverridenChildren& Parent);
+
 	/**
 	 * ColorName and MaterialName have to be initialized before.
 	 * This method update the value ColorName or MaterialName accordingly of the GraphStyleData type (material or color)
diff --git a/Engine/Plugins/Enterprise/DatasmithCADImporter/Source/CADInterfaces/Private/TechSoftFileParserCADKernelTessellator.cpp b/Engine/Plugins/Enterprise/DatasmithCADImporter/Source/CADInterfaces/Private/TechSoftFileParserCADKernelTessellator.cpp
--- a/Engine/Plugins/Enterprise/DatasmithCADImporter/Source/CADInterfaces/Private/TechSoftFileParserCADKernelTessellator.cpp
+++ b/Engine/Plugins/Enterprise/DatasmithCADImporter/Source/CADInterfaces/Private/TechSoftFileParserCADKernelTessellator.cpp
@@ -54,13 +54,13 @@ void FTechSoftFileParserCADKernelTessellator::SewAndGenerateBodyMeshes()
 	TRACE_CPUPROFILER_EVENT_SCOPE(FTechSoftFileParserCADKernelTessellator::SewAndGenerateBodyMeshes);
 
 	int32 RepresentationCount = RepresentationItemsCache.Num();
-	TMap<FCadId, TArray<A3DRiRepresentationItem*>> OccurenceIdToRepresentations;
+	TMap<FCadId, TArray<A3DRiRepresentationItem*>> OccurrenceIdToRepresentations;
 	for (TPair<A3DRiRepresentationItem*, FCadId>& Entry : RepresentationItemsCache)
 	{
 		A3DRiRepresentationItem* RepresentationItemPtr = Entry.Key;
 		FArchiveBody& Body = SceneGraph.GetBody(Entry.Value);
 
-		TArray<A3DRiRepresentationItem*>& Representations = OccurenceIdToRepresentations.FindOrAdd(Body.ParentId);
+		TArray<A3DRiRepresentationItem*>& Representations = OccurrenceIdToRepresentations.FindOrAdd(Body.ParentId);
 		if (Representations.Max() == 0)
 		{
 			Representations.Reserve(RepresentationCount);
@@ -69,7 +69,7 @@ void FTechSoftFileParserCADKernelTessellator::SewAndGenerateBodyMeshes()
 		Representations.Add(RepresentationItemPtr);
 	}
 
-	for (TPair<FCadId, TArray<A3DRiRepresentationItem*>>& Representations : OccurenceIdToRepresentations)
+	for (TPair<FCadId, TArray<A3DRiRepresentationItem*>>& Representations : OccurrenceIdToRepresentations)
 	{
 		SewAndMesh(Representations.Value);
 		if (bConvertionFailed)
diff --git a/Engine/Plugins/Enterprise/DatasmithCADImporter/Source/CADInterfaces/Public/CADFileData.h b/Engine/Plugins/Enterprise/DatasmithCADImporter/Source/CADInterfaces/Public/CADFileData.h
--- a/Engine/Plugins/Enterprise/DatasmithCADImporter/Source/CADInterfaces/Public/CADFileData.h
+++ b/Engine/Plugins/Enterprise/DatasmithCADImporter/Source/CADInterfaces/Public/CADFileData.h
@@ -73,6 +73,7 @@ public:
 
 	void SetArchiveNames()
 	{
+		SceneGraphArchive.SceneGraphId = FileDescription.GetDescriptorHash();
 		SceneGraphArchive.ArchiveFileName = FString::Printf(TEXT("UEx%08x"), GetSceneFileHash());
 		MeshArchiveFile = FString::Printf(TEXT("UEx%08x"), GetGeomFileHash());
 	}
diff --git a/Engine/Plugins/Enterprise/DatasmithCADImporter/Source/CADInterfaces/Public/CADSceneGraph.h b/Engine/Plugins/Enterprise/DatasmithCADImporter/Source/CADInterfaces/Public/CADSceneGraph.h
--- a/Engine/Plugins/Enterprise/DatasmithCADImporter/Source/CADInterfaces/Public/CADSceneGraph.h
+++ b/Engine/Plugins/Enterprise/DatasmithCADImporter/Source/CADInterfaces/Public/CADSceneGraph.h
@@ -38,8 +38,6 @@ public:
 	FString Label;
 	TMap<FString, FString> MetaData;
 	FMatrix TransformMatrix = FMatrix::Identity;
-
-	// not serialized
 	double Unit = 1;
 
 	bool IsNameDefined() const
@@ -51,12 +49,57 @@ public:
 
 };
 
-class CADINTERFACES_API FArchiveInstance : public FArchiveCADObject
+class CADINTERFACES_API FArchiveWithOverridenChildren : public FArchiveCADObject
+{
+public:
+	FArchiveWithOverridenChildren() = default;
+	FArchiveWithOverridenChildren(FCadId Id, const FArchiveCADObject& Parent)
+		: FArchiveCADObject(Id, Parent)
+	{
+	}
+
+	friend FArchive& operator<<(FArchive& Ar, FArchiveWithOverridenChildren& C);
+
+	void AddOverridenChild(const FCadId ChildId);
+
+public:
+	TArray<FCadId> OverridenChildren;
+};
+
+/**
+ * This class save the override data of an occurrence. 
+ * Indeed, some tools/formats allows to override some occurrences of the instances of a reference (e.g. only the 4th occurrence of C is hidden in the graph below)
+ * 
+ *     A
+ *     | - B
+ *     |   | - C (show)
+ *     |   | - C (show)
+ *     | 
+ *     | - B
+ *     |   | - C (show)
+ *     |   | - C (hide)
+ */
+class CADINTERFACES_API FArchiveOverrideOccurrence : public FArchiveWithOverridenChildren
+{
+public:
+	FArchiveOverrideOccurrence() = default;
+	FArchiveOverrideOccurrence(FCadId Id, const FArchiveCADObject& Parent)
+		: FArchiveWithOverridenChildren(Id, Parent)
+	{
+	}
+
+	friend FArchive& operator<<(FArchive& Ar, FArchiveOverrideOccurrence& C);
+};
+
+/**
+ * An instance is link to its reference but It also save its tree of override occurrence children
+ */
+class CADINTERFACES_API FArchiveInstance : public FArchiveWithOverridenChildren
 {
 public:
 	FArchiveInstance() = default;
 	FArchiveInstance(FCadId Id, const FArchiveCADObject& Parent)
-		: FArchiveCADObject(Id, Parent)
+		: FArchiveWithOverridenChildren(Id, Parent)
 	{
 	}
 
@@ -179,6 +222,7 @@ public:
 	void DeserializeMockUpFile(const TCHAR* Filename);
 
 public:
+	uint32 SceneGraphId = 0;
 	FString CADFileName;
 	FString ArchiveFileName;
 	FString FullPath;
@@ -191,6 +235,7 @@ public:
 	TArray<FArchiveUnloadedReference> UnloadedReferences;
 	TArray<FFileDescriptor> ExternalReferenceFiles;
 	TArray<FArchiveInstance> Instances;
+	TArray<FArchiveOverrideOccurrence> OverrideOccurrences;
 
 	TArray<int32> CADIdToIndex;
 	FCadId LastEntityId = 1;
@@ -206,33 +251,38 @@ public:
 		ExternalReferenceFiles.Reserve(ComponentCount[EComponentType::Reference]);
 		Bodies.Reserve(ComponentCount[EComponentType::Body]);
 
-		CADIdToIndex.Reserve(ComponentCount[EComponentType::Instance] + ComponentCount[EComponentType::Reference] + ComponentCount[EComponentType::Body] + 1);
+		OverrideOccurrences.Reserve(ComponentCount[EComponentType::OverriddeOccurence]);
+
+		CADIdToIndex.Reserve(ComponentCount[EComponentType::OverriddeOccurence] + ComponentCount[EComponentType::Instance] + ComponentCount[EComponentType::Reference] + ComponentCount[EComponentType::Body] + 1);
 		CADIdToIndex.Add(0);
 	}
 
 	FArchiveInstance& AddInstance(const FArchiveCADObject& Parent);
-	FArchiveInstance& GetInstance(FCadId CadId);
+	FArchiveInstance& GetInstance(const FCadId CadId);
 	void RemoveLastInstance();
-	bool IsAInstance(FCadId CadId) const;
+	bool IsAInstance(const FCadId CadId) const;
 
 	FArchiveReference& AddReference(FArchiveUnloadedReference&);
 	FArchiveReference& AddReference(FArchiveInstance& Parent);
-	FArchiveReference& GetReference(FCadId CadId);
+	FArchiveReference& GetReference(const FCadId CadId);
 	void RemoveLastReference();
 	bool IsAReference(FCadId CadId) const;
 
-	FArchiveReference& AddOccurence(FArchiveReference& Parent);
-	void RemoveLastOccurence();
+	FArchiveReference& AddOccurrence(FArchiveReference& Parent);
+	void RemoveLastOccurrence();
 
 	FArchiveUnloadedReference& AddUnloadedReference(FArchiveInstance& Parent);
-	FArchiveUnloadedReference& GetUnloadedReference(FCadId CadId);
+	FArchiveUnloadedReference& GetUnloadedReference(const FCadId CadId);
 	void RemoveLastUnloadedReference();
-	bool IsAUnloadedReference(FCadId CadId) const;
+	bool IsAUnloadedReference(const FCadId CadId) const;
+
+	FArchiveOverrideOccurrence& AddOverrideOccurrence(FArchiveWithOverridenChildren& Parent);
+	FArchiveOverrideOccurrence& GetOverrideOccurrence(const FCadId CadId);
 
 	FArchiveBody& AddBody(FArchiveReference& Parent, EMesher InMesher);
-	FArchiveBody& GetBody(FCadId CadId);
+	FArchiveBody& GetBody(const FCadId CadId);
 	void RemoveLastBody();
-	bool IsABody(FCadId CadId) const;
+	bool IsABody(const FCadId CadId) const;
 
 	void AddExternalReferenceFile(const FArchiveUnloadedReference& Reference);
 
diff --git a/Engine/Plugins/Enterprise/DatasmithCADImporter/Source/CADTools/Private/CADData.cpp b/Engine/Plugins/Enterprise/DatasmithCADImporter/Source/CADTools/Private/CADData.cpp
--- a/Engine/Plugins/Enterprise/DatasmithCADImporter/Source/CADTools/Private/CADData.cpp
+++ b/Engine/Plugins/Enterprise/DatasmithCADImporter/Source/CADTools/Private/CADData.cpp
@@ -51,6 +51,16 @@ FMaterialUId BuildMaterialUId(const FCADMaterial& Material)
 	return FMath::Abs((int32)MaterialUId);
 }
 
+FArchive& operator<<(FArchive& Ar, FArchiveGraphicProperties& Object)
+{
+	Ar << Object.ColorUId;
+	Ar << Object.MaterialUId;
+	Ar << Object.bIsRemoved;
+	Ar << Object.bShow;
+	Ar << Object.Inheritance;
+	return Ar;
+}
+
 FArchive& operator<<(FArchive& Ar, FCADMaterial& Material)
 {
 	Ar << Material.MaterialName;
@@ -242,13 +252,15 @@ void FFileDescriptor::SetFileFormat(const FString& Extension)
 	}
 	else if (Extension == TEXT("iam"))
 	{
+		// Importer.m_sLoadData.m_sIncremental not supported
 		Format = ECADFormat::INVENTOR;
-		bCanReferenceOtherFiles = true;
+		bCanReferenceOtherFiles = false;
 	}
 	else if (Extension == TEXT("jt"))
 	{
+		// Importer.m_sLoadData.m_sIncremental not supported
 		Format = ECADFormat::JT;
-		bCanReferenceOtherFiles = true;
+		bCanReferenceOtherFiles = false;
 	}
 	else if (Extension == TEXT("model"))
 	{
@@ -302,6 +314,7 @@ void FFileDescriptor::SetFileFormat(const FString& Extension)
 	}
 	else if (Extension == TEXT("3dxml") || Extension == TEXT("3drep"))
 	{
+		// Importer.m_sLoadData.m_sIncremental not supported
 		Format = ECADFormat::CATIA_3DXML;
 		bCanReferenceOtherFiles = false;
 	}
@@ -323,7 +336,7 @@ void FFileDescriptor::SetFileFormat(const FString& Extension)
 	else if (Extension == TEXT("ifc") || Extension == TEXT("ifczip"))
 	{
 		Format = ECADFormat::IFC;
-		bCanReferenceOtherFiles = false;
+		bCanReferenceOtherFiles = true;
 	}
 	else if (Extension == TEXT("dgn"))
 	{
diff --git a/Engine/Plugins/Enterprise/DatasmithCADImporter/Source/CADTools/Private/CADToolsModule.cpp b/Engine/Plugins/Enterprise/DatasmithCADImporter/Source/CADTools/Private/CADToolsModule.cpp
--- a/Engine/Plugins/Enterprise/DatasmithCADImporter/Source/CADTools/Private/CADToolsModule.cpp
+++ b/Engine/Plugins/Enterprise/DatasmithCADImporter/Source/CADTools/Private/CADToolsModule.cpp
@@ -18,7 +18,7 @@ bool FCADToolsModule::IsAvailable()
 
 uint32 FCADToolsModule::GetCacheVersion()
 {
-	return 504;
+	return 530;
 }
 
 IMPLEMENT_MODULE(FCADToolsModule, CADTools);
diff --git a/Engine/Plugins/Enterprise/DatasmithCADImporter/Source/CADTools/Public/CADData.h b/Engine/Plugins/Enterprise/DatasmithCADImporter/Source/CADTools/Public/CADData.h
--- a/Engine/Plugins/Enterprise/DatasmithCADImporter/Source/CADTools/Public/CADData.h
+++ b/Engine/Plugins/Enterprise/DatasmithCADImporter/Source/CADTools/Public/CADData.h
@@ -55,6 +55,7 @@ enum EComponentType : uint8
 {
 	Reference = 0,
 	Instance,
+	OverriddeOccurence,
 	Body,
 	LastType
 };
@@ -96,11 +97,7 @@ struct CADTOOLS_API FObjectDisplayDataId
 class CADTOOLS_API FArchiveGraphicProperties
 {
 public:
-	FArchiveGraphicProperties()
-		: ColorUId(0)
-		, MaterialUId(0)
-	{
-	}
+	FArchiveGraphicProperties() = default;
 
 	FArchiveGraphicProperties(const FArchiveGraphicProperties& Parent)
 		: ColorUId(Parent.ColorUId)
@@ -110,6 +107,8 @@ public:
 
 	virtual ~FArchiveGraphicProperties() = default;
 
+	friend CADTOOLS_API FArchive& operator<<(FArchive& Ar, FArchiveGraphicProperties& C);
+
 public:
 	FMaterialUId ColorUId = 0;
 	FMaterialUId MaterialUId = 0;
diff --git a/Engine/Plugins/Enterprise/DatasmithCADImporter/Source/DatasmithCADTranslator/Private/DatasmithSceneGraphBuilder.cpp b/Engine/Plugins/Enterprise/DatasmithCADImporter/Source/DatasmithCADTranslator/Private/DatasmithSceneGraphBuilder.cpp
--- a/Engine/Plugins/Enterprise/DatasmithCADImporter/Source/DatasmithCADTranslator/Private/DatasmithSceneGraphBuilder.cpp
+++ b/Engine/Plugins/Enterprise/DatasmithCADImporter/Source/DatasmithCADTranslator/Private/DatasmithSceneGraphBuilder.cpp
@@ -17,7 +17,7 @@ namespace DatasmithSceneGraphBuilderImpl
  * For more details:
  * @see https://docs.techsoft3d.com/exchange/latest/build/managing_attribute_inheritance.html
  */
-void SpreadGraphicProperties(const CADLibrary::FArchiveCADObject& Component, ActorData& ComponentData)
+void SpreadGraphicProperties(const CADLibrary::FArchiveCADObject& Component, FActorData& ComponentData)
 {
 	if (!ComponentData.ColorUId)
 	{
@@ -131,6 +131,7 @@ bool FDatasmithSceneGraphBuilder::Build()
 	{
 		return false;
 	}
+
 	AncestorSceneGraphHash.Add(RootHash);
 
 	return FDatasmithSceneBaseGraphBuilder::Build();
@@ -176,10 +177,10 @@ void FDatasmithSceneGraphBuilder::FillAnchorActor(const TSharedRef<IDatasmithAct
 		return;
 	}
 
-	ActorData Data(TEXT(""));
+	FActorData Data(TEXT(""));
 
 	const FCadId RootId = 1;
-	ActorData ParentData(ActorElement->GetName());
+	FActorData ParentData(ActorElement->GetName());
 	CADLibrary::FArchiveReference& Reference = SceneGraph->GetReference(RootId);
 
 	CADLibrary::FArchiveInstance EmptyInstance;
@@ -189,7 +190,7 @@ void FDatasmithSceneGraphBuilder::FillAnchorActor(const TSharedRef<IDatasmithAct
 
 	AddMetaData(ActorElement, EmptyInstance, Reference);
 
-	ActorData ComponentData(*ActorUUID, ParentData);
+	FActorData ComponentData(*ActorUUID, ParentData);
 	//DatasmithSceneGraphBuilderImpl::GetMainMaterial(Reference.MetaData, ComponentData, bMaterialPropagationIsTopDown);
 
 	AddChildren(ActorElement, Reference, ComponentData);
@@ -225,7 +226,7 @@ bool FDatasmithSceneBaseGraphBuilder::Build()
 {
 	const FCadId RootId = 1;
 
-	ActorData Data(TEXT(""));
+	FActorData Data(TEXT(""));
 	CADLibrary::FArchiveReference& Reference = SceneGraph->GetReference(RootId);
 	TSharedPtr<IDatasmithActorElement> RootActor = BuildReference(Reference, Data);
 	DatasmithScene->AddActor(RootActor);
@@ -298,13 +299,54 @@ bool FDatasmithSceneBaseGraphBuilder::Build()
 	return true;
 }
 
-TSharedPtr<IDatasmithActorElement>  FDatasmithSceneBaseGraphBuilder::BuildInstance(FCadId InstanceId, const ActorData& ParentData)
+void FDatasmithSceneBaseGraphBuilder::OverrideInstance(CADLibrary::FArchiveOverrideOccurrence* OverrideOccurrence, CADLibrary::FArchiveInstance& Instance, FActorData& InstanceData)
+{
+	if (OverrideOccurrence == nullptr)
+	{
+		return;
+	}
+
+	if (OverrideOccurrence->OverridenChildren.Num())
+	{
+		InstanceData.OverridenChildren = MoveTemp(OverrideOccurrence->OverridenChildren);
+	}
+
+	if (OverrideOccurrence->TransformMatrix != FMatrix::Identity)
+	{
+		Instance.TransformMatrix = OverrideOccurrence->TransformMatrix;
+	}
+
+	Instance.bShow = OverrideOccurrence->bShow;
+}
+
+TSharedPtr<IDatasmithActorElement>  FDatasmithSceneBaseGraphBuilder::BuildInstance(FCadId InstanceId, const FActorData& ParentData)
 {
 	CADLibrary::FArchiveReference* Reference = nullptr;
 	CADLibrary::FArchiveReference EmptyReference;
 
 	CADLibrary::FArchiveInstance& Instance = SceneGraph->GetInstance(InstanceId);
 
+	CADLibrary::FArchiveOverrideOccurrence* OverrideInstancePtr = nullptr;
+	uint32 OverridenSceneGraphId = 0;
+
+	if(ParentData.OverridenChildren.Num() > 0)
+	{
+		CADLibrary::FArchiveSceneGraph* ParentSceneGraph = CADFileToSceneGraphArchive.FindRef(ParentData.SceneGraphId);
+		if(ParentSceneGraph)
+		{
+			for (const FCadId& OverrideOccurrenceId : ParentData.OverridenChildren)
+			{
+				CADLibrary::FArchiveOverrideOccurrence& OverrideOccurrence = ParentSceneGraph->GetOverrideOccurrence(OverrideOccurrenceId);
+				if (OverrideOccurrence.Label == Instance.Label)
+				{
+					OverrideInstancePtr = &OverrideOccurrence;
+					OverridenSceneGraphId = ParentData.SceneGraphId;
+					break;
+				}
+			}
+		}
+	}
+
 	CADLibrary::FArchiveSceneGraph* InstanceSceneGraph = SceneGraph;
 	if (Instance.bIsExternalReference)
 	{
@@ -357,14 +399,31 @@ TSharedPtr<IDatasmithActorElement>  FDatasmithSceneBaseGraphBuilder::BuildInstan
 	{
 		AddMetaData(Actor, Instance, *Reference);
 
-		ActorData InstanceData(*ActorUUID, ParentData);
+		FActorData InstanceData(*ActorUUID, ParentData);
 
 		DatasmithSceneGraphBuilderImpl::SpreadGraphicProperties(Instance, InstanceData);
 		DatasmithSceneGraphBuilderImpl::SpreadGraphicProperties(*Reference, InstanceData);
+		if (OverrideInstancePtr && OverrideInstancePtr->OverridenChildren.Num())
+		{
+			InstanceData.OverridenChildren = OverrideInstancePtr->OverridenChildren;
+			InstanceData.SceneGraphId = OverridenSceneGraphId;
+		}
+		else if (Instance.OverridenChildren.Num())
+		{
+			InstanceData.OverridenChildren = MoveTemp(Instance.OverridenChildren);
+			InstanceData.SceneGraphId = InstanceSceneGraph->SceneGraphId;
+		}
 
-		AddChildren(Actor, *Reference, InstanceData);
-
-		DatasmithSceneGraphBuilderImpl::AddTransformToActor(Instance, Actor, ImportParameters);
+		OverrideInstance(OverrideInstancePtr, Instance, InstanceData);
+		if(Instance.bShow)
+		{
+			AddChildren(Actor, *Reference, InstanceData);
+			DatasmithSceneGraphBuilderImpl::AddTransformToActor(Instance, Actor, ImportParameters);
+		}
+		else
+		{
+			Actor.Reset();
+		}
 	}
 
 	if (SceneGraph != InstanceSceneGraph)
@@ -386,7 +445,7 @@ TSharedPtr<IDatasmithActorElement>  FDatasmithSceneBaseGraphBuilder::CreateActor
 	return TSharedPtr<IDatasmithActorElement>();
 }
 
-TSharedPtr<IDatasmithActorElement> FDatasmithSceneBaseGraphBuilder::BuildReference(CADLibrary::FArchiveReference& Reference, const ActorData& ParentData)
+TSharedPtr<IDatasmithActorElement> FDatasmithSceneBaseGraphBuilder::BuildReference(CADLibrary::FArchiveReference& Reference, const FActorData& ParentData)
 {
 	TMap<FString, FString> InstanceNodeMetaDataMap;
 
@@ -403,7 +462,7 @@ TSharedPtr<IDatasmithActorElement> FDatasmithSceneBaseGraphBuilder::BuildReferen
 
 	AddMetaData(Actor, EmptyInstance, Reference);
 
-	ActorData ReferenceData(*ActorUUID, ParentData);
+	FActorData ReferenceData(*ActorUUID, ParentData);
 	DatasmithSceneGraphBuilderImpl::SpreadGraphicProperties(Reference, ReferenceData);
 
 	AddChildren(Actor, Reference, ReferenceData);
@@ -413,7 +472,7 @@ TSharedPtr<IDatasmithActorElement> FDatasmithSceneBaseGraphBuilder::BuildReferen
 	return Actor;
 }
 
-TSharedPtr<IDatasmithActorElement> FDatasmithSceneBaseGraphBuilder::BuildBody(FCadId BodyId, const ActorData& ParentData)
+TSharedPtr<IDatasmithActorElement> FDatasmithSceneBaseGraphBuilder::BuildBody(FCadId BodyId, const FActorData& ParentData)
 {
 	CADLibrary::FArchiveInstance EmptyInstance;
 	CADLibrary::FArchiveBody& Body = SceneGraph->GetBody(BodyId);
@@ -428,7 +487,7 @@ TSharedPtr<IDatasmithActorElement> FDatasmithSceneBaseGraphBuilder::BuildBody(FC
 	DatasmithSceneGraphBuilderImpl::GetNodeUuidAndLabel(EmptyInstance, Body, ParentData.Uuid, BodyUUID, BodyLabel);
 
 	// Apply materials on the current part
-	ActorData BodyData(*BodyUUID, ParentData);
+	FActorData BodyData(*BodyUUID, ParentData);
 	DatasmithSceneGraphBuilderImpl::SpreadGraphicProperties(Body, BodyData);
 	FMaterialUId MaterialUId = (BodyData.Inheritance == CADLibrary::ECADGraphicPropertyInheritance::Unset) ? -1 : BodyData.MaterialUId ? BodyData.MaterialUId : BodyData.ColorUId;
 
@@ -633,7 +692,7 @@ bool FDatasmithSceneBaseGraphBuilder::DoesActorHaveChildrenOrIsAStaticMesh(const
 	return false;
 }
 
-void FDatasmithSceneBaseGraphBuilder::AddChildren(TSharedPtr<IDatasmithActorElement> Actor, const CADLibrary::FArchiveReference& Reference, const ActorData& ParentData)
+void FDatasmithSceneBaseGraphBuilder::AddChildren(TSharedPtr<IDatasmithActorElement> Actor, const CADLibrary::FArchiveReference& Reference, const FActorData& ParentData)
 {
 	for (const FCadId& ChildId : Reference.Children)
 	{
diff --git a/Engine/Plugins/Enterprise/DatasmithCADImporter/Source/DatasmithCADTranslator/Public/DatasmithSceneGraphBuilder.h b/Engine/Plugins/Enterprise/DatasmithCADImporter/Source/DatasmithCADTranslator/Public/DatasmithSceneGraphBuilder.h
--- a/Engine/Plugins/Enterprise/DatasmithCADImporter/Source/DatasmithCADTranslator/Public/DatasmithSceneGraphBuilder.h
+++ b/Engine/Plugins/Enterprise/DatasmithCADImporter/Source/DatasmithCADTranslator/Public/DatasmithSceneGraphBuilder.h
@@ -23,10 +23,10 @@ class IDatasmithMeshElement;
 class IDatasmithScene;
 class IDatasmithUEPbrMaterialElement;
 
-class ActorData  //#ueent_CAD
+class FActorData
 {
 public:
-	ActorData(const TCHAR* NodeUuid, const ActorData& ParentData)
+	FActorData(const TCHAR* NodeUuid, const FActorData& ParentData)
 		: Uuid(NodeUuid)
 		, Inheritance(ParentData.Inheritance)
 		, MaterialUId(ParentData.MaterialUId)
@@ -34,7 +34,7 @@ public:
 	{
 	}
 
-	ActorData(const TCHAR* NodeUuid)
+	FActorData(const TCHAR* NodeUuid)
 		: Uuid(NodeUuid)
 		, Inheritance(CADLibrary::ECADGraphicPropertyInheritance::Unset)
 		, MaterialUId(0)
@@ -46,10 +46,10 @@ public:
 	CADLibrary::ECADGraphicPropertyInheritance Inheritance = CADLibrary::ECADGraphicPropertyInheritance::Unset;
 	FMaterialUId MaterialUId;
 	FMaterialUId ColorUId;
-};
-
-
 
+	uint32 SceneGraphId = 0;
+	TArray<FCadId> OverridenChildren;
+};
 
 class DATASMITHCADTRANSLATOR_API FDatasmithSceneBaseGraphBuilder
 {
@@ -66,12 +66,14 @@ public:
 	virtual bool Build();
 
 protected:
-	TSharedPtr<IDatasmithActorElement> BuildInstance(FCadId InstanceId, const ActorData& ParentData);
-	TSharedPtr<IDatasmithActorElement> BuildReference(CADLibrary::FArchiveReference& Reference, const ActorData& ParentData);
-	TSharedPtr<IDatasmithActorElement> BuildBody(FCadId BodyId, const ActorData& ParentData);
+	TSharedPtr<IDatasmithActorElement> BuildInstance(FCadId InstanceId, const FActorData& ParentData);
+	TSharedPtr<IDatasmithActorElement> BuildReference(CADLibrary::FArchiveReference& Reference, const FActorData& ParentData);
+	TSharedPtr<IDatasmithActorElement> BuildBody(FCadId BodyId, const FActorData& ParentData);
+
+	void OverrideInstance(CADLibrary::FArchiveOverrideOccurrence* OverrideOccurrence, CADLibrary::FArchiveInstance& Instance, FActorData& InstanceData);
 
 	void AddMetaData(TSharedPtr<IDatasmithActorElement> ActorElement, const CADLibrary::FArchiveCADObject& Instance, const CADLibrary::FArchiveCADObject& Reference);
-	void AddChildren(TSharedPtr<IDatasmithActorElement> Actor, const CADLibrary::FArchiveReference& Reference, const ActorData& ParentData);
+	void AddChildren(TSharedPtr<IDatasmithActorElement> Actor, const CADLibrary::FArchiveReference& Reference, const FActorData& ParentData);
 	bool DoesActorHaveChildrenOrIsAStaticMesh(const TSharedPtr< IDatasmithActorElement >& ActorElement);
 
 	TSharedPtr<IDatasmithUEPbrMaterialElement> GetDefaultMaterial();
