[NO NON-PII COMMIT MESSAGE]

diff --git a/Engine/Plugins/Enterprise/DatasmithCADImporter/Source/CADInterfaces/Private/TechSoftFileParser.cpp b/Engine/Plugins/Enterprise/DatasmithCADImporter/Source/CADInterfaces/Private/TechSoftFileParser.cpp
--- a/Engine/Plugins/Enterprise/DatasmithCADImporter/Source/CADInterfaces/Private/TechSoftFileParser.cpp
+++ b/Engine/Plugins/Enterprise/DatasmithCADImporter/Source/CADInterfaces/Private/TechSoftFileParser.cpp
@@ -493,7 +493,15 @@ void FTechSoftFileParser::GenerateBodyMeshes()
 	{
 		A3DRiRepresentationItem* RepresentationItemPtr = Entry.Key;
 		FArchiveBody& Body = SceneGraph.GetBody(Entry.Value);
-		GenerateBodyMesh(RepresentationItemPtr, Body);
+		if (!Body.bIsFromCad)
+		{
+			FTechSoftFileParser::GenerateBodyMesh(RepresentationItemPtr, Body);
+		}
+		else
+		{
+			GenerateBodyMesh(RepresentationItemPtr, Body);
+		}
+
 		if (bConvertionFailed)
 		{
 			return;
diff --git a/Engine/Plugins/Enterprise/DatasmithCADImporter/Source/CADInterfaces/Private/TechSoftFileParserCADKernelTessellator.cpp b/Engine/Plugins/Enterprise/DatasmithCADImporter/Source/CADInterfaces/Private/TechSoftFileParserCADKernelTessellator.cpp
--- a/Engine/Plugins/Enterprise/DatasmithCADImporter/Source/CADInterfaces/Private/TechSoftFileParserCADKernelTessellator.cpp
+++ b/Engine/Plugins/Enterprise/DatasmithCADImporter/Source/CADInterfaces/Private/TechSoftFileParserCADKernelTessellator.cpp
@@ -59,6 +59,11 @@ void FTechSoftFileParserCADKernelTessellator::SewAndGenerateBodyMeshes()
 	{
 		A3DRiRepresentationItem* RepresentationItemPtr = Entry.Key;
 		FArchiveBody& Body = SceneGraph.GetBody(Entry.Value);
+		if (!Body.bIsFromCad)
+		{
+			FTechSoftFileParser::GenerateBodyMesh(RepresentationItemPtr, Body);
+			continue;
+		}
 
 		TArray<A3DRiRepresentationItem*>& Representations = OccurrenceIdToRepresentations.FindOrAdd(Body.ParentId);
 		if (Representations.Max() == 0)
