Set Techsoft as default tessellator for python test

diff --git a/Engine/Plugins/Enterprise/DatasmithCADImporter/Source/CADTools/Public/CADOptions.h b/Engine/Plugins/Enterprise/DatasmithCADImporter/Source/CADTools/Public/CADOptions.h
--- a/Engine/Plugins/Enterprise/DatasmithCADImporter/Source/CADTools/Public/CADOptions.h
+++ b/Engine/Plugins/Enterprise/DatasmithCADImporter/Source/CADTools/Public/CADOptions.h
@@ -46,7 +46,7 @@ namespace CADLibrary
 
 	ENUM_CLASS_FLAGS(ESewOption);
 
-	class FImportParameters
+	class CADTOOLS_API FImportParameters
 	{
 	private:
 		double ChordTolerance;
@@ -57,23 +57,23 @@ namespace CADLibrary
 		FDatasmithUtils::EModelCoordSystem ModelCoordSys;
 		
 	public:
-		CADTOOLS_API static bool bGDisableCADKernelTessellation;
-		CADTOOLS_API static bool bGEnableCADCache;
-		CADTOOLS_API static bool bGEnableTimeControl;
-		CADTOOLS_API static bool bGOverwriteCache;
-		CADTOOLS_API static bool bGPreferJtFileEmbeddedTessellation;
-		CADTOOLS_API static bool bGSewMeshIfNeeded;
-		CADTOOLS_API static bool bGRemoveDuplicatedTriangle;
-		CADTOOLS_API static float GStitchingTolerance;
-		CADTOOLS_API static bool bGStitchingForceSew;
-		CADTOOLS_API static bool bGStitchingRemoveThinFaces;
-		CADTOOLS_API static bool bGStitchingRemoveDuplicatedFaces;
-		CADTOOLS_API static bool bGActivateThinZoneMeshing;
-		CADTOOLS_API static float GStitchingForceFactor;
-		CADTOOLS_API static float GUnitScale;
-		CADTOOLS_API static float GMeshingParameterFactor;
-		CADTOOLS_API static int32 GMaxMaterialCountPerMesh;
-		CADTOOLS_API static bool bValidationProcess;
+		static bool bGDisableCADKernelTessellation;
+		static bool bGEnableCADCache;
+		static bool bGEnableTimeControl;
+		static bool bGOverwriteCache;
+		static bool bGPreferJtFileEmbeddedTessellation;
+		static bool bGSewMeshIfNeeded;
+		static bool bGRemoveDuplicatedTriangle;
+		static float GStitchingTolerance;
+		static bool bGStitchingForceSew;
+		static bool bGStitchingRemoveThinFaces;
+		static bool bGStitchingRemoveDuplicatedFaces;
+		static bool bGActivateThinZoneMeshing;
+		static float GStitchingForceFactor;
+		static float GUnitScale;
+		static float GMeshingParameterFactor;
+		static int32 GMaxMaterialCountPerMesh;
+		static bool bValidationProcess;
 
 	public:
 		FImportParameters(FDatasmithUtils::EModelCoordSystem NewCoordinateSystem = FDatasmithUtils::EModelCoordSystem::ZUp_RightHanded)
@@ -86,7 +86,7 @@ namespace CADLibrary
 		{
 		}
 	
-		FImportParameters(FImportParameters InParamneters, EMesher InMesher)
+		FImportParameters(const FImportParameters& InParamneters, EMesher InMesher)
 			: ChordTolerance(InParamneters.ChordTolerance)
 			, MaxEdgeLength(InParamneters.MaxEdgeLength)
 			, MaxNormalAngle(InParamneters.MaxNormalAngle)
