[NO NON-PII COMMIT MESSAGE]

diff --git a/Engine/Plugins/Enterprise/DatasmithCADImporter/Source/CADInterfaces/Private/TechSoftFileParser.cpp b/Engine/Plugins/Enterprise/DatasmithCADImporter/Source/CADInterfaces/Private/TechSoftFileParser.cpp
--- a/Engine/Plugins/Enterprise/DatasmithCADImporter/Source/CADInterfaces/Private/TechSoftFileParser.cpp
+++ b/Engine/Plugins/Enterprise/DatasmithCADImporter/Source/CADInterfaces/Private/TechSoftFileParser.cpp
@@ -1887,7 +1887,7 @@ void FTechSoftFileParser::ExtractGeneralTransformation(const A3DMiscTransformati
 		if (Scale.Equals(FVector3d::OneVector, KINDA_SMALL_NUMBER))
 		{
 			const double TranslationScale = Component.Unit * FImportParameters::GUnitScale;
-			for (Index = 0; Index < 3; ++Index, ++Index)
+			for (Index = 0; Index < 3; ++Index)
 			{
 				Matrix.M[3][Index] *= TranslationScale;
 			}
@@ -1895,17 +1895,17 @@ void FTechSoftFileParser::ExtractGeneralTransformation(const A3DMiscTransformati
 		}
 		else
 		{
-		FVector3d Translation = Transform.GetTranslation();
-		Translation *= FImportParameters::GUnitScale;
+			FVector3d Translation = Transform.GetTranslation();
+			Translation *= FImportParameters::GUnitScale;
 
-		double UniformScale = TechSoftFileParserImpl::ExtractUniformScale(Scale);
+			double UniformScale = TechSoftFileParserImpl::ExtractUniformScale(Scale);
 			Component.Unit *= UniformScale;
 
-		FQuat4d Rotation = Transform.GetRotation();
+			FQuat4d Rotation = Transform.GetRotation();
 
-		FTransform3d NewTransform;
-		NewTransform.SetScale3D(Scale);
-		NewTransform.SetRotation(Rotation);
+			FTransform3d NewTransform;
+			NewTransform.SetScale3D(Scale);
+			NewTransform.SetRotation(Rotation);
 
 			Component.TransformMatrix = NewTransform.ToMatrixWithScale();
 			Component.TransformMatrix.SetOrigin(Translation);
