[NO NON-PII COMMIT MESSAGE]

diff --git a/Engine/Plugins/Enterprise/DatasmithCADImporter/Source/CADInterfaces/Private/TechSoftFileParser.cpp b/Engine/Plugins/Enterprise/DatasmithCADImporter/Source/CADInterfaces/Private/TechSoftFileParser.cpp
--- a/Engine/Plugins/Enterprise/DatasmithCADImporter/Source/CADInterfaces/Private/TechSoftFileParser.cpp
+++ b/Engine/Plugins/Enterprise/DatasmithCADImporter/Source/CADInterfaces/Private/TechSoftFileParser.cpp
@@ -1896,7 +1896,7 @@ void FTechSoftFileParser::ExtractGeneralTransformation(const A3DMiscTransformati
 		else
 		{
 			FVector3d Translation = Transform.GetTranslation();
-			Translation *= FImportParameters::GUnitScale;
+			Translation *= Component.Unit * FImportParameters::GUnitScale;
 
 			double UniformScale = TechSoftFileParserImpl::ExtractUniformScale(Scale);
 			Component.Unit *= UniformScale;
