The Alias orientation flag was taken into account twice in case of BRep body

diff --git a/Engine/Plugins/Enterprise/DatasmithCADImporter/Source/DatasmithWireTranslator/Private/DatasmithWireTranslator.cpp b/Engine/Plugins/Enterprise/DatasmithCADImporter/Source/DatasmithWireTranslator/Private/DatasmithWireTranslator.cpp
--- a/Engine/Plugins/Enterprise/DatasmithCADImporter/Source/DatasmithWireTranslator/Private/DatasmithWireTranslator.cpp
+++ b/Engine/Plugins/Enterprise/DatasmithCADImporter/Source/DatasmithWireTranslator/Private/DatasmithWireTranslator.cpp
@@ -2268,7 +2268,7 @@ TOptional<FMeshDescription> FWireTranslatorImpl::GetMeshDescription(TSharedRef<I
 	{
 		boolean bAlOrientation;
 		DagNode.getSurfaceOrientation(bAlOrientation);
-		MeshParameters.bNeedSwapOrientation = (bool)bAlOrientation;
+		MeshParameters.bNeedSwapOrientation = (objectType == kMeshNodeType) ? (bool)bAlOrientation : false; 
 
 		AlLayer* LayerPtr = DagNode.layer();
 		if (AlIsValid(LayerPtr))
