Fixed explicit LexToString template parameter which should have been deduced, causing a compile error.

diff --git a/Engine/Plugins/Enterprise/DatasmithCADImporter/Source/CADLibrary/Private/CADKernelTools.cpp b/Engine/Plugins/Enterprise/DatasmithCADImporter/Source/CADLibrary/Private/CADKernelTools.cpp
--- a/Engine/Plugins/Enterprise/DatasmithCADImporter/Source/CADLibrary/Private/CADKernelTools.cpp
+++ b/Engine/Plugins/Enterprise/DatasmithCADImporter/Source/CADLibrary/Private/CADKernelTools.cpp
@@ -119,7 +119,7 @@ bool FillMesh(FMeshConversionContext& Context, UE::CADKernel::FModelMesh& ModelM
 	for (auto& Material : MaterialToPolygonGroupMapping)
 	{
 		uint32 MaterialHash = Material.Key;
-		FName ImportedSlotName = *LexToString<uint32>(MaterialHash);
+		FName ImportedSlotName = *LexToString(MaterialHash);
 
 		FPolygonGroupID PolyGroupID = MeshDescription.CreatePolygonGroup();
 		PolygonGroupImportedMaterialSlotNames[PolyGroupID] = ImportedSlotName;
