Fixed explicit LexToString template parameter which should have been deduced, causing a compile error.

diff --git a/Engine/Plugins/Enterprise/DatasmithCADImporter/Source/CADLibrary/Private/MeshDescriptionHelper.cpp b/Engine/Plugins/Enterprise/DatasmithCADImporter/Source/CADLibrary/Private/MeshDescriptionHelper.cpp
--- a/Engine/Plugins/Enterprise/DatasmithCADImporter/Source/CADLibrary/Private/MeshDescriptionHelper.cpp
+++ b/Engine/Plugins/Enterprise/DatasmithCADImporter/Source/CADLibrary/Private/MeshDescriptionHelper.cpp
@@ -318,7 +318,7 @@ bool FillMesh(const FMeshConversionContext& MeshConversionContext, FBodyMesh& Bo
 		if (PolyGroupIndex < FImportParameters::GMaxMaterialCountPerMesh)
 		{
 			uint32 MaterialHash = Material.Key;
-			FName ImportedSlotName = *LexToString<uint32>(MaterialHash);
+			FName ImportedSlotName = *LexToString(MaterialHash);
 			PolyGroupID = MeshDescription.CreatePolygonGroup();
 			PolygonGroupImportedMaterialSlotNames[PolyGroupID] = ImportedSlotName;
 		}
