CIS fix.

diff --git a/Engine/Plugins/Enterprise/DatasmithCADImporter/Source/DatasmithWireTranslator/Private/AliasModelToTechSoftConverter.cpp b/Engine/Plugins/Enterprise/DatasmithCADImporter/Source/DatasmithWireTranslator/Private/AliasModelToTechSoftConverter.cpp
--- a/Engine/Plugins/Enterprise/DatasmithCADImporter/Source/DatasmithWireTranslator/Private/AliasModelToTechSoftConverter.cpp
+++ b/Engine/Plugins/Enterprise/DatasmithCADImporter/Source/DatasmithWireTranslator/Private/AliasModelToTechSoftConverter.cpp
@@ -150,7 +150,7 @@ A3DCrvBase* FAliasModelToTechSoftConverter::CreateCurve(const AlTrimCurve& Alias
 	int ControlPointCount = AliasTrimCurve.numberOfCVs();
 	int32 KnotCount = AliasTrimCurve.realNumberOfKnots();
 
-	using AlPoint = double[3];
+	using AlPoint = TStaticArray<double, 3>;
 
 	TArray<AlPoint> AliasPoles;
 	TArray<double> AliasNodalVector;
@@ -159,7 +159,7 @@ A3DCrvBase* FAliasModelToTechSoftConverter::CreateCurve(const AlTrimCurve& Alias
 
 	// Notice that each CV has three coordinates - the three coordinates describe 2D parameter space, with a homogeneous coordinate.
 	// Each control point is u, v and w, where u and v are parameter space and w is the homogeneous coordinate.
-	AliasTrimCurve.CVsUVPosition(AliasNodalVector.GetData(), AliasPoles.GetData());
+	AliasTrimCurve.CVsUVPosition(AliasNodalVector.GetData(), (double(*)[3])AliasPoles.GetData());
 
 	AliasNodalVector.SetNum(KnotCount);
 	AliasTrimCurve.realKnotVector(AliasNodalVector.GetData());
