Fixed static mesh build issue when mesh is not manifold

#rnx

diff --git a/Engine/Plugins/Enterprise/DatasmithCADImporter/Source/CADLibrary/Private/MeshDescriptionHelper.cpp b/Engine/Plugins/Enterprise/DatasmithCADImporter/Source/CADLibrary/Private/MeshDescriptionHelper.cpp
--- a/Engine/Plugins/Enterprise/DatasmithCADImporter/Source/CADLibrary/Private/MeshDescriptionHelper.cpp
+++ b/Engine/Plugins/Enterprise/DatasmithCADImporter/Source/CADLibrary/Private/MeshDescriptionHelper.cpp
@@ -556,6 +556,10 @@ bool ConvertBodyMeshToMeshDescription(const FMeshConversionContext& MeshConversi
 	{
 		return false;
 	}
+	if (!Body.bIsFromCad)
+	{
+		MeshOperator::FixNonManifoldMesh(MeshDescription);
+	}
 
 	// Workaround SDHE-19725: Compute any null normals.
 	MeshOperator::RecomputeNullNormal(MeshDescription);
diff --git a/Engine/Plugins/Enterprise/DatasmithCADImporter/Source/CADLibrary/Private/MeshOperator.cpp b/Engine/Plugins/Enterprise/DatasmithCADImporter/Source/CADLibrary/Private/MeshOperator.cpp
--- a/Engine/Plugins/Enterprise/DatasmithCADImporter/Source/CADLibrary/Private/MeshOperator.cpp
+++ b/Engine/Plugins/Enterprise/DatasmithCADImporter/Source/CADLibrary/Private/MeshOperator.cpp
@@ -7,6 +7,7 @@
 #include "Containers/Array.h"
 #include "Containers/Queue.h"
 #include "DynamicMesh/DynamicMesh3.h"
+#include "DynamicMesh/NonManifoldMappingSupport.h"
 #include "DynamicMesh/Operations/MergeCoincidentMeshEdges.h"
 #include "DynamicMeshToMeshDescription.h"
 #include "MeshAttributes.h"
@@ -383,4 +384,61 @@ void MeshOperator::ResolveTJunctions(FMeshDescription& MeshDescription, double T
 	}
 
 	MeshDescriptionDataCache.RestoreMaterialSlotNames(MeshDescription);
+}
+
+
+void MeshOperator::FixNonManifoldMesh(FMeshDescription& MeshDescription)
+{
+	using namespace UE::Geometry;
+
+	FDynamicMesh3 DynamicMesh(UE::Geometry::EMeshComponents::All);
+	DynamicMesh.EnableAttributes();
+
+	TFunction<void()> ToDynamicMesh = [&MeshDescription, &DynamicMesh]()
+	{
+		DynamicMesh.Clear();
+
+		FMeshDescriptionToDynamicMesh ConverterToDynamicMesh;
+		ConverterToDynamicMesh.Convert(&MeshDescription, DynamicMesh);
+	};
+
+	TFunction<void()> ToMeshDescription = [&MeshDescription, &DynamicMesh]()
+	{
+		CADLibrary::FMeshDescriptionDataCache MeshDescriptionDataCache(MeshDescription);
+
+		{
+			FConversionToMeshDescriptionOptions ConversionOptions;
+			ConversionOptions.bSetPolyGroups = true;
+			ConversionOptions.bUpdatePositions = true;
+			ConversionOptions.bUpdateNormals = true;
+			ConversionOptions.bUpdateTangents = false;
+			ConversionOptions.bUpdateUVs = true;
+			ConversionOptions.bUpdateVtxColors = false;
+			ConversionOptions.bTransformVtxColorsSRGBToLinear = false;
+
+			FDynamicMeshToMeshDescription ConverterToMeshDescription(ConversionOptions);
+			ConverterToMeshDescription.Convert(&DynamicMesh, MeshDescription);
+		}
+
+		MeshDescriptionDataCache.RestoreMaterialSlotNames(MeshDescription);
+	};
+
+	ToDynamicMesh();
+
+	// Nothing to do if mesh is manifold (all edges are shared by at most 2 triangles)
+	if (DynamicMesh.CheckValidity(FDynamicMesh3::FValidityOptions(false, true), EValidityCheckFailMode::ReturnOnly))
+	{
+		return;
+	}
+
+	// Although the conversion from MeshDescription to DynamicMesh3 handles non-manifold edges
+	// Some edges may be shared by more than 3 triangles. To solve this, we must perform
+	// the conversion back and forth until all non-manifold edges have been eliminated.
+	do 
+	{
+		ToMeshDescription();
+		ToDynamicMesh();
+	} while (DynamicMesh.CheckValidity(FDynamicMesh3::FValidityOptions(false, true), EValidityCheckFailMode::ReturnOnly));
+
+	ToMeshDescription();
 }
\ No newline at end of file
diff --git a/Engine/Plugins/Enterprise/DatasmithCADImporter/Source/CADLibrary/Private/MeshOperator.h b/Engine/Plugins/Enterprise/DatasmithCADImporter/Source/CADLibrary/Private/MeshOperator.h
--- a/Engine/Plugins/Enterprise/DatasmithCADImporter/Source/CADLibrary/Private/MeshOperator.h
+++ b/Engine/Plugins/Enterprise/DatasmithCADImporter/Source/CADLibrary/Private/MeshOperator.h
@@ -9,6 +9,8 @@
 namespace MeshOperator
 {
 
+void FixNonManifoldMesh(FMeshDescription& MeshDescription);
+
 bool OrientMesh(FMeshDescription& MeshDescription);
 
 void RecomputeNullNormal(FMeshDescription& MeshDescription);
