Fix C4702 warnings
#rnx

diff --git a/Engine/Plugins/Enterprise/DatasmithCADImporter/Source/CADLibrary/Private/MeshOperator.cpp b/Engine/Plugins/Enterprise/DatasmithCADImporter/Source/CADLibrary/Private/MeshOperator.cpp
--- a/Engine/Plugins/Enterprise/DatasmithCADImporter/Source/CADLibrary/Private/MeshOperator.cpp
+++ b/Engine/Plugins/Enterprise/DatasmithCADImporter/Source/CADLibrary/Private/MeshOperator.cpp
@@ -342,14 +342,7 @@ void MeshOperator::ResolveTJunctions(FMeshDescription& MeshDescription, double T
 	}
 
 	// Check if there are boundary edges
-	int32 BoundaryEdgeCount = 0;
-	for (int32 eid : DynamicMesh.BoundaryEdgeIndicesItr())
-	{
-		BoundaryEdgeCount++;
-		break;
-	}
-
-	if (BoundaryEdgeCount == 0)
+	if (auto It(DynamicMesh.BoundaryEdgeIndicesItr()); It.begin() == It.end())
 	{
 		return;
 	}
@@ -364,7 +357,6 @@ void MeshOperator::ResolveTJunctions(FMeshDescription& MeshDescription, double T
 		MergeCoincidentEdges.MergeVertexTolerance = Tolerance;
 		MergeCoincidentEdges.MergeSearchTolerance = 2 * Tolerance;
 		MergeCoincidentEdges.Apply();
-		BoundaryEdgeCount = MergeCoincidentEdges.FinalNumBoundaryEdges;
 	}
 
 	CADLibrary::FMeshDescriptionDataCache MeshDescriptionDataCache(MeshDescription);
diff --git a/Engine/Source/Runtime/Datasmith/CADKernel/Private/CADKernel/Mesh/Structure/LoopCleaner.cpp b/Engine/Source/Runtime/Datasmith/CADKernel/Private/CADKernel/Mesh/Structure/LoopCleaner.cpp
--- a/Engine/Source/Runtime/Datasmith/CADKernel/Private/CADKernel/Mesh/Structure/LoopCleaner.cpp
+++ b/Engine/Source/Runtime/Datasmith/CADKernel/Private/CADKernel/Mesh/Structure/LoopCleaner.cpp
@@ -740,19 +740,13 @@ bool FLoopCleaner::RemoveSelfIntersectionsOfLoop()
 		if (IntersectionCount == 1)
 		{
 			NextIntersectionIndex = IntersectionIndex + 1;
-			for (; NextIntersectionIndex < Intersections.Num(); ++NextIntersectionIndex)
+			if (Intersections.IsValidIndex(NextIntersectionIndex) && Intersections[NextIntersectionIndex].Key <= Intersection.Value)
 			{
-				if (Intersections[NextIntersectionIndex].Key > Intersection.Value)
-				{
-					break;
-				}
-
 				bIntersectionForward = false;
 				++IntersectionCount;
-#ifdef DEBUG_REMOVE_LOOP_INTERSECTIONS		
+#ifdef DEBUG_REMOVE_LOOP_INTERSECTIONS
 				DisplayIntersection(Intersections[NextIntersectionIndex]);
 #endif
-				break;
 			}
 		}
 
