Miscellaeous improvements in import of wire files

#rnx

diff --git a/Engine/Plugins/Enterprise/DatasmithCADImporter/Source/CADKernelSurface/Private/CADKernelSurfaceExtension.cpp b/Engine/Plugins/Enterprise/DatasmithCADImporter/Source/CADKernelSurface/Private/CADKernelSurfaceExtension.cpp
--- a/Engine/Plugins/Enterprise/DatasmithCADImporter/Source/CADKernelSurface/Private/CADKernelSurfaceExtension.cpp
+++ b/Engine/Plugins/Enterprise/DatasmithCADImporter/Source/CADKernelSurface/Private/CADKernelSurfaceExtension.cpp
@@ -64,8 +64,7 @@ bool UCADKernelParametricSurfaceData::Tessellate(UStaticMesh& StaticMesh, const
 		CADKernelSession->AddDatabase(RawData);
 
 		FModel& CADKernelModel = CADKernelSession->GetModel();
-		TArray<TSharedPtr<FBody>> CADKernelBodies = CADKernelModel.GetBodies();
-		if (CADKernelBodies.Num() != 1)
+		if (CADKernelModel.GetBodies().IsEmpty())
 		{
 			return bSuccessfulTessellation;
 		}
diff --git a/Engine/Plugins/Enterprise/DatasmithCADImporter/Source/CADLibrary/Private/CADKernelTools.cpp b/Engine/Plugins/Enterprise/DatasmithCADImporter/Source/CADLibrary/Private/CADKernelTools.cpp
--- a/Engine/Plugins/Enterprise/DatasmithCADImporter/Source/CADLibrary/Private/CADKernelTools.cpp
+++ b/Engine/Plugins/Enterprise/DatasmithCADImporter/Source/CADLibrary/Private/CADKernelTools.cpp
@@ -271,7 +271,7 @@ bool FCADKernelTools::Tessellate(UE::CADKernel::FTopologicalShapeEntity& CADTopo
 	TSharedRef<FModelMesh> CADKernelModelMesh = FEntity::MakeShared<FModelMesh>();
 
 	const double GeometricTolerance = FImportParameters::GStitchingTolerance * 10; // cm to mm
-	const bool bActivateThinZoneMeshing = FImportParameters::bGActivateThinZoneMeshing;
+	const bool bActivateThinZoneMeshing = TessellationContext.bSkipThinZoneMeshing ? false : FImportParameters::bGActivateThinZoneMeshing;
 	FMesher Mesher(*CADKernelModelMesh, GeometricTolerance, bActivateThinZoneMeshing);
 
 	DefineMeshCriteria(*CADKernelModelMesh, TessellationContext.ImportParameters, GeometricTolerance);
diff --git a/Engine/Plugins/Enterprise/DatasmithCADImporter/Source/CADLibrary/Public/MeshDescriptionHelper.h b/Engine/Plugins/Enterprise/DatasmithCADImporter/Source/CADLibrary/Public/MeshDescriptionHelper.h
--- a/Engine/Plugins/Enterprise/DatasmithCADImporter/Source/CADLibrary/Public/MeshDescriptionHelper.h
+++ b/Engine/Plugins/Enterprise/DatasmithCADImporter/Source/CADLibrary/Public/MeshDescriptionHelper.h
@@ -28,6 +28,7 @@ struct FMeshConversionContext
 	TArray<int32> VertexIds;
 	TArray<int32> SymmetricVertexIds;
 	TSet<int32> PatchesToMesh;
+	bool bSkipThinZoneMeshing = false;
 
 	FMeshConversionContext(const FImportParameters& InImportParams, const FMeshParameters& InMeshParameters)
 		: ImportParameters(InImportParams)
diff --git a/Engine/Plugins/Enterprise/DatasmithCADImporter/Source/DatasmithWireTranslator/DatasmithWireTranslator2020.Build.cs b/Engine/Plugins/Enterprise/DatasmithCADImporter/Source/DatasmithWireTranslator/DatasmithWireTranslator2020.Build.cs
--- a/Engine/Plugins/Enterprise/DatasmithCADImporter/Source/DatasmithWireTranslator/DatasmithWireTranslator2020.Build.cs
+++ b/Engine/Plugins/Enterprise/DatasmithCADImporter/Source/DatasmithWireTranslator/DatasmithWireTranslator2020.Build.cs
@@ -67,6 +67,7 @@ public class DatasmithWireTranslator2020 : DatasmithWireTranslatorBase
 	public DatasmithWireTranslator2020(ReadOnlyTargetRules Target) 
 		: base(Target)
 	{
+		PublicDefinitions.Add("IS_MAIN_MODULE");
 	}
 
 	public override string GetAliasVersion()
diff --git a/Engine/Plugins/Enterprise/DatasmithCADImporter/Source/DatasmithWireTranslator/Private/AliasModelToCADKernelConverter.cpp b/Engine/Plugins/Enterprise/DatasmithCADImporter/Source/DatasmithWireTranslator/Private/AliasModelToCADKernelConverter.cpp
--- a/Engine/Plugins/Enterprise/DatasmithCADImporter/Source/DatasmithWireTranslator/Private/AliasModelToCADKernelConverter.cpp
+++ b/Engine/Plugins/Enterprise/DatasmithCADImporter/Source/DatasmithWireTranslator/Private/AliasModelToCADKernelConverter.cpp
@@ -2,8 +2,11 @@
 
 #include "AliasModelToCADKernelConverter.h"
 
+#include "OpenModelUtils.h"
+
 #include "CADData.h"
 #include "CADKernelTools.h"
+#include "HAL/ConsoleManager.h"
 #include "Hal/PlatformMemory.h"
 #include "MeshDescriptionHelper.h"
 
@@ -25,6 +28,10 @@
 
 #ifdef USE_OPENMODEL
 
+#if PLATFORM_WINDOWS
+#include "Windows/AllowWindowsPlatformTypes.h"
+#endif
+
 // Alias API wrappes object in AlObjects. This is an abstract base class which holds a reference to an anonymous data structure.
 // The only way to compare two AlObjects is to compare their data structure. That is the reason why private fields are made public. 
 #define private public
@@ -39,6 +46,10 @@
 #include "AlTrimRegion.h"
 #include "AlTM.h"
 
+#if PLATFORM_WINDOWS
+#include "Windows/HideWindowsPlatformTypes.h"
+#endif
+
 namespace UE_DATASMITHWIRETRANSLATOR_NAMESPACE
 {
 
@@ -143,7 +154,7 @@ TSharedPtr<UE::CADKernel::FTopologicalEdge> FAliasModelToCADKernelConverter::Add
 	}
 
 	// Only TrimCurve with twin need to be in the map used in LinkEdgesLoop 
-	TUniquePtr<AlTrimCurve> TwinCurve(AliasTrimCurve.getTwinCurve());
+	TAlObjectPtr<AlTrimCurve> TwinCurve(AliasTrimCurve.getTwinCurve());
 	if (TwinCurve.IsValid())
 	{
 		AlEdge2CADKernelEdge.Add(AliasTrimCurve.fSpline, Edge);
@@ -157,7 +168,7 @@ TSharedPtr<UE::CADKernel::FTopologicalLoop> FAliasModelToCADKernelConverter::Add
 	TArray<TSharedPtr<UE::CADKernel::FTopologicalEdge>> Edges;
 	TArray<UE::CADKernel::EOrientation> Directions;
 
-	for (TUniquePtr<AlTrimCurve> TrimCurve(TrimBoundary.firstCurve()); TrimCurve.IsValid(); TrimCurve = TUniquePtr<AlTrimCurve>(TrimCurve->nextCurve()))
+	for (TAlObjectPtr<AlTrimCurve> TrimCurve(TrimBoundary.firstCurve()); TrimCurve.IsValid(); TrimCurve = TAlObjectPtr<AlTrimCurve>(TrimCurve->nextCurve()))
 	{
 		TSharedPtr<UE::CADKernel::FTopologicalEdge> Edge = AddEdge(*TrimCurve, CarrierSurface);
 		if (Edge.IsValid())
@@ -178,7 +189,7 @@ TSharedPtr<UE::CADKernel::FTopologicalLoop> FAliasModelToCADKernelConverter::Add
 
 void FAliasModelToCADKernelConverter::LinkEdgesLoop(const AlTrimBoundary& TrimBoundary, UE::CADKernel::FTopologicalLoop& Loop)
 {
-	for (TUniquePtr<AlTrimCurve> TrimCurve(TrimBoundary.firstCurve()); TrimCurve.IsValid(); TrimCurve = TUniquePtr<AlTrimCurve>(TrimCurve->nextCurve()))
+	for (TAlObjectPtr<AlTrimCurve> TrimCurve(TrimBoundary.firstCurve()); TrimCurve.IsValid(); TrimCurve = TAlObjectPtr<AlTrimCurve>(TrimCurve->nextCurve()))
 	{
 		TSharedPtr<UE::CADKernel::FTopologicalEdge>* Edge = AlEdge2CADKernelEdge.Find(TrimCurve->fSpline);
 		if (!Edge || !Edge->IsValid() || (*Edge)->IsDeleted() || (*Edge)->IsDegenerated())
@@ -189,7 +200,7 @@ void FAliasModelToCADKernelConverter::LinkEdgesLoop(const AlTrimBoundary& TrimBo
 		ensure(&Loop == (*Edge)->GetLoop());
 
 		// Link edges
-		TUniquePtr<AlTrimCurve> TwinCurve(TrimCurve->getTwinCurve());
+		TAlObjectPtr<AlTrimCurve> TwinCurve(TrimCurve->getTwinCurve());
 		if (TwinCurve.IsValid())
 		{
 			if (TSharedPtr<UE::CADKernel::FTopologicalEdge>* TwinEdge = AlEdge2CADKernelEdge.Find(TwinCurve->fSpline))
@@ -213,7 +224,7 @@ TSharedPtr<UE::CADKernel::FTopologicalFace> FAliasModelToCADKernelConverter::Add
 
 	bool bIsExternal = true;
 	TArray<TSharedPtr<UE::CADKernel::FTopologicalLoop>> Loops;
-	for (TUniquePtr<AlTrimBoundary> TrimBoundary(TrimRegion.firstBoundary()); TrimBoundary.IsValid(); TrimBoundary = TUniquePtr<AlTrimBoundary>(TrimBoundary->nextBoundary()))
+	for (TAlObjectPtr<AlTrimBoundary> TrimBoundary(TrimRegion.firstBoundary()); TrimBoundary.IsValid(); TrimBoundary = TAlObjectPtr<AlTrimBoundary>(TrimBoundary->nextBoundary()))
 	{
 		TSharedPtr<UE::CADKernel::FTopologicalLoop> Loop = AddLoop(*TrimBoundary, Surface, bIsExternal);
 		if (Loop.IsValid())
@@ -250,10 +261,10 @@ TSharedPtr<UE::CADKernel::FTopologicalFace> FAliasModelToCADKernelConverter::Add
 
 void FAliasModelToCADKernelConverter::AddFace(const AlSurface& Surface, EAliasObjectReference InObjectReference, const AlMatrix4x4& InAlMatrix, bool bInOrientation, TSharedRef<UE::CADKernel::FShell>& Shell)
 {
-	TUniquePtr<AlTrimRegion> TrimRegion(Surface.firstTrimRegion());
+	TAlObjectPtr<AlTrimRegion> TrimRegion(Surface.firstTrimRegion());
 	if (TrimRegion.IsValid())
 	{
-		for (; TrimRegion.IsValid(); TrimRegion = TUniquePtr<AlTrimRegion>(TrimRegion->nextRegion()))
+		for (; TrimRegion.IsValid(); TrimRegion = TAlObjectPtr<AlTrimRegion>(TrimRegion->nextRegion()))
 		{
 			TSharedPtr<UE::CADKernel::FTopologicalFace> Face = AddTrimRegion(*TrimRegion, InObjectReference, InAlMatrix, bInOrientation);
 			if (Face.IsValid())
@@ -275,7 +286,7 @@ void FAliasModelToCADKernelConverter::AddFace(const AlSurface& Surface, EAliasOb
 
 void FAliasModelToCADKernelConverter::AddShell(const AlShell& InShell, EAliasObjectReference InObjectReference, const AlMatrix4x4& InAlMatrix, bool bInOrientation, TSharedRef<UE::CADKernel::FShell>& CADKernelShell)
 {
-	for(TUniquePtr<AlTrimRegion> TrimRegion(InShell.firstTrimRegion()); TrimRegion.IsValid(); TrimRegion = TUniquePtr<AlTrimRegion>(TrimRegion->nextRegion()))
+	for(TAlObjectPtr<AlTrimRegion> TrimRegion(InShell.firstTrimRegion()); TrimRegion.IsValid(); TrimRegion = TAlObjectPtr<AlTrimRegion>(TrimRegion->nextRegion()))
 	{
 		TSharedPtr<UE::CADKernel::FTopologicalFace> Face = AddTrimRegion(*TrimRegion, InObjectReference, InAlMatrix, bInOrientation);
 		if (Face.IsValid())
@@ -317,7 +328,7 @@ bool FAliasModelToCADKernelConverter::AddBRep(AlDagNode& DagNode, const FColor&
 			AlShell* ShellPtr = ShellNode->shell();
 			if(AlIsValid(ShellPtr))
 			{
-				TUniquePtr<AlShell> AliasShell(ShellPtr);
+				TAlObjectPtr<AlShell> AliasShell(ShellPtr);
 				AddShell(*ShellPtr, InObjectReference, AlMatrix, bOrientation, CADKernelShell);
 			}
 		}
@@ -331,7 +342,7 @@ bool FAliasModelToCADKernelConverter::AddBRep(AlDagNode& DagNode, const FColor&
 			AlSurface* SurfacePtr = SurfaceNode->surface();
 			if (AlIsValid(SurfacePtr))
 			{
-				TUniquePtr<AlSurface> AliasSurface(SurfacePtr);
+				TAlObjectPtr<AlSurface> AliasSurface(SurfacePtr);
 				AddFace(*SurfacePtr, InObjectReference, AlMatrix, bOrientation, CADKernelShell);
 			}
 		}
@@ -347,6 +358,45 @@ bool FAliasModelToCADKernelConverter::AddBRep(AlDagNode& DagNode, const FColor&
 	return CADKernelShell->FaceCount() > 0;
 }
 
+bool FAliasModelToCADKernelConverter::Tessellate(const CADLibrary::FMeshParameters& InMeshParameters, FMeshDescription& OutMeshDescription)
+{
+	UE::CADKernel::FModel& Model = CADKernelSession.GetModel();
+
+	CADLibrary::FMeshConversionContext Context(ImportParameters, InMeshParameters);
+
+	const IConsoleVariable* ConsoleVariable = IConsoleManager::Get().FindConsoleVariable(TEXT("ds.WireTranslator.SkipThinZoneMeshing"));
+	Context.bSkipThinZoneMeshing = ConsoleVariable ? ConsoleVariable->GetBool() : false;
+
+	return CADLibrary::FCADKernelTools::Tessellate(Model, Context, OutMeshDescription);
+}
+
+bool FAliasModelToCADKernelConverter::RepairTopology()
+{
+	using namespace CADLibrary;
+	// Apply stitching if applicable
+	if (ImportParameters.GetStitchingTechnique() != StitchingNone)
+	{
+		const double StitchingTolerance = FImportParameters::GStitchingTolerance * 10.; //CM to MM
+		UE::CADKernel::ESewOption SewOptionValue = (UE::CADKernel::ESewOption)SewOption::GetFromImportParameters();
+
+		const IConsoleVariable* ConsoleVariable = IConsoleManager::Get().FindConsoleVariable(TEXT("ds.WireTranslator.SkipThinZoneMeshing"));
+		bool bSkipThinZoneMeshing = ConsoleVariable ? ConsoleVariable->GetBool() : false;
+		if (bSkipThinZoneMeshing)
+		{
+			SewOptionValue = UE::CADKernel::ESewOption((uint8)SewOptionValue & (uint8)~ESewOption::RemoveThinFaces);
+		}
+
+		UE::CADKernel::FTopomakerOptions TopomakerOptions(SewOptionValue, StitchingTolerance, FImportParameters::GStitchingForceFactor);
+
+		UE::CADKernel::FTopomaker Topomaker(CADKernelSession, TopomakerOptions);
+		Topomaker.Sew();
+		Topomaker.SplitIntoConnectedShells();
+		Topomaker.OrientShells();
+	}
+
+	return true;
+}
+
 }
 
 #endif
\ No newline at end of file
diff --git a/Engine/Plugins/Enterprise/DatasmithCADImporter/Source/DatasmithWireTranslator/Private/AliasModelToCADKernelConverter.h b/Engine/Plugins/Enterprise/DatasmithCADImporter/Source/DatasmithWireTranslator/Private/AliasModelToCADKernelConverter.h
--- a/Engine/Plugins/Enterprise/DatasmithCADImporter/Source/DatasmithWireTranslator/Private/AliasModelToCADKernelConverter.h
+++ b/Engine/Plugins/Enterprise/DatasmithCADImporter/Source/DatasmithWireTranslator/Private/AliasModelToCADKernelConverter.h
@@ -42,6 +42,10 @@ public:
 
 	virtual bool AddBRep(AlDagNode& DagNode, const FColor& Color, EAliasObjectReference ObjectReference) override;
 
+	virtual bool Tessellate(const CADLibrary::FMeshParameters& InMeshParameters, FMeshDescription& OutMeshDescription) override;
+
+	virtual bool RepairTopology() override;
+
 protected:
 	TSharedPtr<UE::CADKernel::FTopologicalEdge> AddEdge(const AlTrimCurve& TrimCurve, TSharedPtr<UE::CADKernel::FSurface>& CarrierSurface);
 
diff --git a/Engine/Plugins/Enterprise/DatasmithCADImporter/Source/DatasmithWireTranslator/Private/AliasModelToTechSoftConverter.cpp b/Engine/Plugins/Enterprise/DatasmithCADImporter/Source/DatasmithWireTranslator/Private/AliasModelToTechSoftConverter.cpp
--- a/Engine/Plugins/Enterprise/DatasmithCADImporter/Source/DatasmithWireTranslator/Private/AliasModelToTechSoftConverter.cpp
+++ b/Engine/Plugins/Enterprise/DatasmithCADImporter/Source/DatasmithWireTranslator/Private/AliasModelToTechSoftConverter.cpp
@@ -2,6 +2,8 @@
 
 #include "AliasModelToTechSoftConverter.h"
 
+#include "OpenModelUtils.h"
+
 #include "Hal/PlatformMemory.h"
 #include "Math/Color.h"
 #include "MeshDescriptionHelper.h"
@@ -12,6 +14,10 @@
 
 #ifdef USE_OPENMODEL
 
+#if PLATFORM_WINDOWS
+#include "Windows/AllowWindowsPlatformTypes.h"
+#endif
+
 // Alias API wrappes object in AlObjects. This is an abstract base class which holds a reference to an anonymous data structure.
 // The only way to compare two AlObjects is to compare their data structure. That is the reason why private fields are made public. 
 #define private public
@@ -26,6 +32,10 @@
 #include "AlTrimRegion.h"
 #include "AlTM.h"
 
+#if PLATFORM_WINDOWS
+#include "Windows/HideWindowsPlatformTypes.h"
+#endif
+
 namespace UE_DATASMITHWIRETRANSLATOR_NAMESPACE
 {
 
@@ -240,7 +250,7 @@ A3DTopoCoEdge* FAliasModelToTechSoftConverter::CreateEdge(const AlTrimCurve& Tri
 	A3DTopoCoEdge* CoEdgePtr = CADLibrary::TechSoftInterface::CreateTopoCoEdge(*CoEdgeData);
 
 	// Only TrimCurve with twin need to be in the map used in LinkEdgesLoop 
-	TUniquePtr<AlTrimCurve> TwinCurve(TrimCurve.getTwinCurve());
+	TAlObjectPtr<AlTrimCurve> TwinCurve(TrimCurve.getTwinCurve());
 	if (TwinCurve.IsValid())
 	{
 		AlEdgeToTSCoEdge.Add(TrimCurve.fSpline, CoEdgePtr);
@@ -253,7 +263,7 @@ A3DTopoLoop* FAliasModelToTechSoftConverter::CreateTopoLoop(const AlTrimBoundary
 	TArray<A3DTopoCoEdge*> Edges;
 	Edges.Reserve(20);
 
-	for (TUniquePtr<AlTrimCurve> TrimCurve(TrimBoundary.firstCurve()); TrimCurve.IsValid(); TrimCurve = TUniquePtr<AlTrimCurve>(TrimCurve->nextCurve()))
+	for (TAlObjectPtr<AlTrimCurve> TrimCurve(TrimBoundary.firstCurve()); TrimCurve.IsValid(); TrimCurve = TAlObjectPtr<AlTrimCurve>(TrimCurve->nextCurve()))
 	{
 		A3DTopoCoEdge* Edge = CreateEdge(*TrimCurve);
 		if (Edge != nullptr)
@@ -278,7 +288,7 @@ A3DTopoLoop* FAliasModelToTechSoftConverter::CreateTopoLoop(const AlTrimBoundary
 
 void FAliasModelToTechSoftConverter::LinkEdgesLoop(const AlTrimBoundary& TrimBoundary)
 {
-	for (TUniquePtr<AlTrimCurve> TrimCurve(TrimBoundary.firstCurve()); TrimCurve.IsValid(); TrimCurve = TUniquePtr<AlTrimCurve>(TrimCurve->nextCurve()))
+	for (TAlObjectPtr<AlTrimCurve> TrimCurve(TrimBoundary.firstCurve()); TrimCurve.IsValid(); TrimCurve = TAlObjectPtr<AlTrimCurve>(TrimCurve->nextCurve()))
 	{
 		A3DTopoCoEdge** Edge = AlEdgeToTSCoEdge.Find(TrimCurve->fSpline);
 		if (Edge == nullptr)
@@ -287,7 +297,7 @@ void FAliasModelToTechSoftConverter::LinkEdgesLoop(const AlTrimBoundary& TrimBou
 		}
 
 		// Link edges
-		TUniquePtr<AlTrimCurve> TwinCurve(TrimCurve->getTwinCurve());
+		TAlObjectPtr<AlTrimCurve> TwinCurve(TrimCurve->getTwinCurve());
 		if (TwinCurve.IsValid())
 		{
 			if (A3DTopoCoEdge** TwinEdge = AlEdgeToTSCoEdge.Find(TwinCurve->fSpline))
@@ -313,7 +323,7 @@ A3DTopoFace* FAliasModelToTechSoftConverter::AddTrimRegion(const AlTrimRegion& I
 	TArray<A3DTopoLoop*> Loops;
 	Loops.Reserve(5);
 
-	for (TUniquePtr<AlTrimBoundary> TrimBoundary(InTrimRegion.firstBoundary()); TrimBoundary.IsValid(); TrimBoundary = TUniquePtr<AlTrimBoundary>(TrimBoundary->nextBoundary()))
+	for (TAlObjectPtr<AlTrimBoundary> TrimBoundary(InTrimRegion.firstBoundary()); TrimBoundary.IsValid(); TrimBoundary = TAlObjectPtr<AlTrimBoundary>(TrimBoundary->nextBoundary()))
 	{
 		A3DTopoLoop* Loop = CreateTopoLoop(*TrimBoundary);
 		if (Loop != nullptr)
@@ -370,10 +380,10 @@ bool FAliasModelToTechSoftConverter::AddBRep(AlDagNode& DagNode, const FColor& C
 	{
 		if (AlShellNode* ShellNode = DagNode.asShellNodePtr())
 		{
-			TUniquePtr<AlShell> AliasShell(ShellNode->shell());
+			TAlObjectPtr<AlShell> AliasShell(ShellNode->shell());
 			if(AlIsValid(AliasShell.Get()))
 			{
-				for (TUniquePtr<AlTrimRegion> TrimRegion(AliasShell->firstTrimRegion()); TrimRegion.IsValid(); TrimRegion = TUniquePtr<AlTrimRegion>(TrimRegion->nextRegion()))
+				for (TAlObjectPtr<AlTrimRegion> TrimRegion(AliasShell->firstTrimRegion()); TrimRegion.IsValid(); TrimRegion = TAlObjectPtr<AlTrimRegion>(TrimRegion->nextRegion()))
 				{
 					A3DTopoFace* TSFace = AddTrimRegion(*TrimRegion, Color, InObjectReference, AlMatrix);
 					if (TSFace != nullptr)
@@ -390,13 +400,13 @@ bool FAliasModelToTechSoftConverter::AddBRep(AlDagNode& DagNode, const FColor& C
 	{
 		if (AlSurfaceNode* SurfaceNode = DagNode.asSurfaceNodePtr())
 		{
-			TUniquePtr<AlSurface> AliasSurface(SurfaceNode->surface());
+			TAlObjectPtr<AlSurface> AliasSurface(SurfaceNode->surface());
 			if (AlIsValid(AliasSurface.Get()))
 			{
-				TUniquePtr<AlTrimRegion> TrimRegion(AliasSurface->firstTrimRegion());
+				TAlObjectPtr<AlTrimRegion> TrimRegion(AliasSurface->firstTrimRegion());
 				if (TrimRegion.IsValid())
 				{
-					for (; TrimRegion.IsValid(); TrimRegion = TUniquePtr<AlTrimRegion>(TrimRegion->nextRegion()))
+					for (; TrimRegion.IsValid(); TrimRegion = TAlObjectPtr<AlTrimRegion>(TrimRegion->nextRegion()))
 					{
 						A3DTopoFace* TSFace = AddTrimRegion(*TrimRegion, Color, InObjectReference, AlMatrix);
 						if (TSFace != nullptr)
diff --git a/Engine/Plugins/Enterprise/DatasmithCADImporter/Source/DatasmithWireTranslator/Private/DatasmithWireTranslator.cpp b/Engine/Plugins/Enterprise/DatasmithCADImporter/Source/DatasmithWireTranslator/Private/DatasmithWireTranslator.cpp
--- a/Engine/Plugins/Enterprise/DatasmithCADImporter/Source/DatasmithWireTranslator/Private/DatasmithWireTranslator.cpp
+++ b/Engine/Plugins/Enterprise/DatasmithCADImporter/Source/DatasmithWireTranslator/Private/DatasmithWireTranslator.cpp
@@ -9,11 +9,11 @@
 #include "DatasmithTranslator.h"
 #include "DatasmithUtils.h"
 #include "DatasmithWireTranslatorModule.h"
+#include "GenericPlatform/GenericPlatformTLS.h"
 #include "HAL/ConsoleManager.h"
 #include "IDatasmithSceneElements.h"
 #include "Misc/FileHelper.h"
 #include "Misc/Paths.h"
-#include "OpenModelUtils.h"
 #include "Utility/DatasmithMeshHelper.h"
 
 #include "StaticMeshDescription.h"
@@ -59,6 +59,10 @@
 #include <AlUniverse.h>
 #endif
 
+#if PLATFORM_WINDOWS
+#include "Windows/HideWindowsPlatformTypes.h"
+#endif
+
 DEFINE_LOG_CATEGORY_STATIC(LogDatasmithWireTranslator, Log, All);
 
 #define LOCTEXT_NAMESPACE "DatasmithWireTranslator"
@@ -66,6 +70,7 @@ DEFINE_LOG_CATEGORY_STATIC(LogDatasmithWireTranslator, Log, All);
 #define WRONG_VERSION_TEXT "Unsupported version of Alias detected. Please upgrade to Alias 2021.3 (or later version)."
 #define CAD_INTERFACE_UNAVAILABLE "CAD Interface module is unavailable. Meshing will be done by Alias."
 
+#pragma optimize("", off)
 #ifdef OPEN_MODEL_2023_0
 
 static bool bGAliasSewByMaterial = false;
@@ -86,17 +91,36 @@ ECVF_Default);
 
 #endif 
 
+#ifdef IS_MAIN_MODULE
+static TAutoConsoleVariable<bool> CVarAliasThreadSafe(
+	TEXT("ds.WireTranslator.ThreadSafe"),
+	false,
+	TEXT("If true, the translator will be called in more than one thread. Default false.\n"),
+	ECVF_Default);
+
+static TAutoConsoleVariable<bool> CVarAliasUseNative(
+	TEXT("ds.WireTranslator.UseNative"),
+	false,
+	TEXT("If true, the AlaisStudio tessellator will be used. Default false.\n"),
+	ECVF_Default);
+
+static TAutoConsoleVariable<bool> CVarAliasSkipThinZoneMeshing(
+	TEXT("ds.WireTranslator.SkipThinZoneMeshing"),
+	true,
+	TEXT("If true, the global thin zone cvar will be honored. Default true.\n"),
+	ECVF_Default);
+#endif
 
 namespace UE_DATASMITHWIRETRANSLATOR_NAMESPACE
 {
 
+#ifdef USE_OPENMODEL
+
 static bool bGSewByMaterial = false;
 static bool bGLayersAsActors = false;
 
 static const FColor DefaultColor = FColor(200, 200, 200);
 
-#ifdef USE_OPENMODEL
-
 const uint64 LibAliasNext_Version = 0xffffffffffffffffull;
 const uint64 LibAlias2023_1_0_Version = 8162778619576619;
 const uint64 LibAlias2023_0_0_Version = 8162774324609149;
@@ -299,18 +323,23 @@ public:
 		DatasmithScene->SetExporterSDKVersion(*AliasSdkVersion);
 		DatasmithScene->SetProductVersion(*AliasProductVersion);
 
-		CADLibrary::FImportParameters ImportParameters;
-		if (CADLibrary::FImportParameters::bGDisableCADKernelTessellation)
+		const IConsoleVariable* ConsoleVariable = IConsoleManager::Get().FindConsoleVariable(TEXT("ds.Wiretranslator.UseNative"));
+		bAliasUseNative = ConsoleVariable ? ConsoleVariable->GetBool() : false;
+		if (!bAliasUseNative)
 		{
-			TSharedRef<FAliasModelToTechSoftConverter> AliasToTechSoftConverter = MakeShared<FAliasModelToTechSoftConverter>(ImportParameters);
-			CADModelConverter = AliasToTechSoftConverter;
-			AliasBRepConverter = AliasToTechSoftConverter;
-		}
-		else
-		{
-			TSharedRef<FAliasModelToCADKernelConverter> AliasToCADKernelConverter = MakeShared<FAliasModelToCADKernelConverter>(ImportParameters);
-			CADModelConverter = AliasToCADKernelConverter;
-			AliasBRepConverter = AliasToCADKernelConverter;
+			CADLibrary::FImportParameters ImportParameters;
+			if (CADLibrary::FImportParameters::bGDisableCADKernelTessellation)
+			{
+				TSharedRef<FAliasModelToTechSoftConverter> AliasToTechSoftConverter = MakeShared<FAliasModelToTechSoftConverter>(ImportParameters);
+				CADModelConverter = AliasToTechSoftConverter;
+				AliasBRepConverter = AliasToTechSoftConverter;
+			}
+			else
+			{
+				TSharedRef<FAliasModelToCADKernelConverter> AliasToCADKernelConverter = MakeShared<FAliasModelToCADKernelConverter>(ImportParameters);
+				CADModelConverter = AliasToCADKernelConverter;
+				AliasBRepConverter = AliasToCADKernelConverter;
+			}
 		}
 
 		IConsoleVariable* CVarDSAliasSewByMaterialEnabled = IConsoleManager::Get().FindConsoleVariable(TEXT("ds.CADTranslator.Alias.SewByMaterial"));
@@ -388,6 +417,8 @@ private:
 
 	void AddNodeInBodyGroup(TSharedPtr<AlDagNode>& DagNode, const FString& ShaderName, TMap<uint32, TSharedPtr<BodyData>>& ShellToProcess, bool bIsAPatch, uint32 MaxSize);
 
+	TSharedPtr<CADLibrary::ICADModelConverter> GetModelConverter(IAliasBRepConverter*& BRepConverter) const;
+
 	TOptional<FMeshDescription> MeshDagNodeWithExternalMesher(AlDagNode& DagNode, TSharedRef<IDatasmithMeshElement> MeshElement, CADLibrary::FMeshParameters& MeshParameters);
 	TOptional<FMeshDescription> MeshDagNodeWithExternalMesher(TSharedRef<BodyData> DagNode, TSharedRef<IDatasmithMeshElement> MeshElement, CADLibrary::FMeshParameters& MeshParameters);
 
@@ -438,10 +469,10 @@ private:
 		}
 	}
 
-	void AddAlBlinnParameters(const TUniquePtr<AlShader>& Shader, TSharedRef<IDatasmithUEPbrMaterialElement> MaterialElement);
-	void AddAlLambertParameters(const TUniquePtr<AlShader>& Shader, TSharedRef<IDatasmithUEPbrMaterialElement> MaterialElement);
-	void AddAlLightSourceParameters(const TUniquePtr<AlShader>& Shader, TSharedRef<IDatasmithUEPbrMaterialElement> MaterialElement);
-	void AddAlPhongParameters(const TUniquePtr<AlShader>& Shader, TSharedRef<IDatasmithUEPbrMaterialElement> MaterialElement);
+	void AddAlBlinnParameters(const TAlObjectPtr<AlShader>& Shader, TSharedRef<IDatasmithUEPbrMaterialElement> MaterialElement);
+	void AddAlLambertParameters(const TAlObjectPtr<AlShader>& Shader, TSharedRef<IDatasmithUEPbrMaterialElement> MaterialElement);
+	void AddAlLightSourceParameters(const TAlObjectPtr<AlShader>& Shader, TSharedRef<IDatasmithUEPbrMaterialElement> MaterialElement);
+	void AddAlPhongParameters(const TAlObjectPtr<AlShader>& Shader, TSharedRef<IDatasmithUEPbrMaterialElement> MaterialElement);
 
 private:
 	TSharedRef<IDatasmithScene> DatasmithScene;
@@ -475,12 +506,16 @@ private:
 	TSharedPtr<CADLibrary::ICADModelConverter> CADModelConverter;
 	TSharedPtr<IAliasBRepConverter> AliasBRepConverter;
 
+	bool bAliasUseNative;
 };
 
 void FWireTranslatorImpl::SetTessellationOptions(const FDatasmithTessellationOptions& Options)
 {
 	TessellationOptions = Options;
-	CADModelConverter->SetImportParameters(Options.ChordTolerance, Options.MaxEdgeLength, Options.NormalTolerance, (CADLibrary::EStitchingTechnique)Options.StitchingTechnique);
+	if (CADModelConverter)
+	{
+		CADModelConverter->SetImportParameters(Options.ChordTolerance, Options.MaxEdgeLength, Options.NormalTolerance, (CADLibrary::EStitchingTechnique)Options.StitchingTechnique);
+	}
 	SceneFileHash = HashCombine(Options.GetHash(), GetSceneFileHash(SceneFullPath, SceneName));
 }
 
@@ -508,7 +543,7 @@ bool FWireTranslatorImpl::Read()
 	return GetDagLeaves();
 }
 
-void FWireTranslatorImpl::AddAlBlinnParameters(const TUniquePtr<AlShader>& Shader, TSharedRef<IDatasmithUEPbrMaterialElement> MaterialElement)
+void FWireTranslatorImpl::AddAlBlinnParameters(const TAlObjectPtr<AlShader>& Shader, TSharedRef<IDatasmithUEPbrMaterialElement> MaterialElement)
 {
 	// Default values for a Blinn material
 	FColor Color(145, 148, 153);
@@ -792,7 +827,7 @@ void FWireTranslatorImpl::AddAlBlinnParameters(const TUniquePtr<AlShader>& Shade
 
 }
 
-void FWireTranslatorImpl::AddAlLambertParameters(const TUniquePtr<AlShader>& Shader, TSharedRef<IDatasmithUEPbrMaterialElement> MaterialElement)
+void FWireTranslatorImpl::AddAlLambertParameters(const TAlObjectPtr<AlShader>& Shader, TSharedRef<IDatasmithUEPbrMaterialElement> MaterialElement)
 {
 	// Default values for a Lambert material
 	FColor Color(145, 148, 153);
@@ -951,7 +986,7 @@ void FWireTranslatorImpl::AddAlLambertParameters(const TUniquePtr<AlShader>& Sha
 
 }
 
-void FWireTranslatorImpl::AddAlLightSourceParameters(const TUniquePtr<AlShader>& Shader, TSharedRef<IDatasmithUEPbrMaterialElement> MaterialElement)
+void FWireTranslatorImpl::AddAlLightSourceParameters(const TAlObjectPtr<AlShader>& Shader, TSharedRef<IDatasmithUEPbrMaterialElement> MaterialElement)
 {
 	// Default values for a LightSource material
 	FColor Color(145, 148, 153);
@@ -1077,7 +1112,7 @@ void FWireTranslatorImpl::AddAlLightSourceParameters(const TUniquePtr<AlShader>&
 
 }
 
-void FWireTranslatorImpl::AddAlPhongParameters(const TUniquePtr<AlShader>& Shader, TSharedRef<IDatasmithUEPbrMaterialElement> MaterialElement)
+void FWireTranslatorImpl::AddAlPhongParameters(const TAlObjectPtr<AlShader>& Shader, TSharedRef<IDatasmithUEPbrMaterialElement> MaterialElement)
 {
 	// Default values for a Phong material
 	FColor Color(145, 148, 153);
@@ -1333,7 +1368,7 @@ void FWireTranslatorImpl::AddAlPhongParameters(const TUniquePtr<AlShader>& Shade
 // Make material
 bool FWireTranslatorImpl::GetShader()
 {
-	for (TUniquePtr<AlShader> Shader(AlUniverse::firstShader()); Shader.IsValid(); Shader = TUniquePtr<AlShader>(AlUniverse::nextShader(Shader.Get())))
+	for (TAlObjectPtr<AlShader> Shader(AlUniverse::firstShader()); Shader.IsValid(); Shader = TAlObjectPtr<AlShader>(AlUniverse::nextShader(Shader.Get())))
 	{
 		const FString ShaderName = UTF8_TO_TCHAR(Shader->name());
 		const FString ShaderModelName = Shader->shadingModel();
@@ -1427,7 +1462,7 @@ void FWireTranslatorImpl::FDagNodeInfo::ExtractDagNodeLayer(const AlDagNode& InD
 	AlLayer* LayerPtr = InDagNode.layer();
 	if (LayerPtr)
 	{
-		TUniquePtr<AlLayer> Layer(LayerPtr);
+		TAlObjectPtr<AlLayer> Layer(LayerPtr);
 		LayerName = UTF8_TO_TCHAR(Layer->name());
 		bLayerIsVisible = !Layer->invisible();
 		if (ActorElement.IsValid())
@@ -1788,14 +1823,14 @@ uint32 GetBodyGroupUuid(const FString& ShaderName, const FString& LayerName, boo
 	return Uuid;
 }
 
-uint32 GetPatchCount(const TUniquePtr<AlShell>& Shell)
+uint32 GetPatchCount(const TAlObjectPtr<AlShell>& Shell)
 {
 	uint32 PatchCount = 0;
-	TUniquePtr<AlTrimRegion> TrimRegion(Shell->firstTrimRegion());
+	TAlObjectPtr<AlTrimRegion> TrimRegion(Shell->firstTrimRegion());
 	while (TrimRegion.IsValid())
 	{
 		PatchCount++;
-		TrimRegion = TUniquePtr<AlTrimRegion>(TrimRegion->nextRegion());
+		TrimRegion = TAlObjectPtr<AlTrimRegion>(TrimRegion->nextRegion());
 	}
 	return PatchCount;
 }
@@ -1807,7 +1842,7 @@ void FWireTranslatorImpl::AddNodeInBodyGroup(TSharedPtr<AlDagNode>& DagNode, con
 	AlLayer* LayerPtr = DagNode->layer();
 	if (AlIsValid(LayerPtr))
 	{
-		const TUniquePtr<AlLayer> Layer(LayerPtr);
+		const TAlObjectPtr<AlLayer> Layer(LayerPtr);
 		LayerName = UTF8_TO_TCHAR(Layer->name());
 		if (Layer->invisible())
 		{
@@ -1878,10 +1913,10 @@ bool FWireTranslatorImpl::RecurseDagForLeaves(const TSharedPtr<AlDagNode>& First
 				break;
 			}
 
-			TUniquePtr<AlShell> Shell(ShellPtr);
+			TAlObjectPtr<AlShell> Shell(ShellPtr);
 			uint32 NbPatch = GetPatchCount(Shell);
 
-			TUniquePtr<AlShader> Shader(Shell->firstShader());
+			TAlObjectPtr<AlShader> Shader(Shell->firstShader());
 			if (Shader.IsValid())
 			{
 				ShaderName = Shader->name();
@@ -1901,10 +1936,10 @@ bool FWireTranslatorImpl::RecurseDagForLeaves(const TSharedPtr<AlDagNode>& First
 		case kSurfaceNodeType:
 		{
 			AlSurfaceNode* SurfaceNode = DagNode->asSurfaceNodePtr();
-			TUniquePtr<AlSurface> Surface(SurfaceNode->surface());
+			TAlObjectPtr<AlSurface> Surface(SurfaceNode->surface());
 			if (Surface.IsValid())
 			{
-				TUniquePtr<AlShader> Shader(Surface->firstShader());
+				TAlObjectPtr<AlShader> Shader(Surface->firstShader());
 				if (Shader.IsValid())
 				{
 					ShaderName = Shader->name();
@@ -1917,10 +1952,10 @@ bool FWireTranslatorImpl::RecurseDagForLeaves(const TSharedPtr<AlDagNode>& First
 		case kMeshNodeType:
 		{
 			AlMeshNode* MeshNode = DagNode->asMeshNodePtr();
-			TUniquePtr<AlMesh> Mesh(MeshNode->mesh());
+			TAlObjectPtr<AlMesh> Mesh(MeshNode->mesh());
 			if (Mesh.IsValid())
 			{
-				TUniquePtr<AlShader> Shader(Mesh->firstShader());
+				TAlObjectPtr<AlShader> Shader(Mesh->firstShader());
 				if (Shader.IsValid())
 				{
 					ShaderName = Shader->name();
@@ -1975,10 +2010,10 @@ void FWireTranslatorImpl::DagForLeavesNoMerge(const TSharedPtr<AlDagNode>& DagNo
 	case kShellNodeType:
 	{
 		AlShellNode* ShellNode = DagNode->asShellNodePtr();
-		TUniquePtr<AlShell> Shell(ShellNode->shell());
+		TAlObjectPtr<AlShell> Shell(ShellNode->shell());
 		if (Shell.IsValid())
 		{
-			TUniquePtr<AlShader> Shader(Shell->firstShader());
+			TAlObjectPtr<AlShader> Shader(Shell->firstShader());
 			if (Shader.IsValid())
 			{
 				ShaderName = Shader->name();
@@ -1991,10 +2026,10 @@ void FWireTranslatorImpl::DagForLeavesNoMerge(const TSharedPtr<AlDagNode>& DagNo
 	case kSurfaceNodeType:
 	{
 		AlSurfaceNode* SurfaceNode = DagNode->asSurfaceNodePtr();
-		TUniquePtr<AlSurface> Surface(SurfaceNode->surface());
+		TAlObjectPtr<AlSurface> Surface(SurfaceNode->surface());
 		if (Surface.IsValid())
 		{
-			TUniquePtr<AlShader> Shader(Surface->firstShader());
+			TAlObjectPtr<AlShader> Shader(Surface->firstShader());
 			if (Shader.IsValid())
 			{
 				ShaderName = Shader->name();
@@ -2007,10 +2042,10 @@ void FWireTranslatorImpl::DagForLeavesNoMerge(const TSharedPtr<AlDagNode>& DagNo
 	case kMeshNodeType:
 	{
 		AlMeshNode* MeshNode = DagNode->asMeshNodePtr();
-		TUniquePtr<AlMesh> Mesh(MeshNode->mesh());
+		TAlObjectPtr<AlMesh> Mesh(MeshNode->mesh());
 		if (Mesh.IsValid())
 		{
-			TUniquePtr<AlShader> Shader(Mesh->firstShader());
+			TAlObjectPtr<AlShader> Shader(Mesh->firstShader());
 			if (Shader.IsValid())
 			{
 				ShaderName = Shader->name();
@@ -2035,10 +2070,36 @@ void FWireTranslatorImpl::DagForLeavesNoMerge(const TSharedPtr<AlDagNode>& DagNo
 	}
 }
 
+TSharedPtr<CADLibrary::ICADModelConverter> FWireTranslatorImpl::GetModelConverter(IAliasBRepConverter*& BRepConverter) const
+{
+	TSharedPtr<CADLibrary::ICADModelConverter> ModelConverter;
+
+	CADLibrary::FImportParameters ImportParameters;
+	if (CADLibrary::FImportParameters::bGDisableCADKernelTessellation)
+	{
+		TSharedRef<FAliasModelToTechSoftConverter> AliasToTechSoftConverter = MakeShared<FAliasModelToTechSoftConverter>(ImportParameters);
+		ModelConverter = AliasToTechSoftConverter;
+		BRepConverter = &AliasToTechSoftConverter.Get();
+	}
+	else
+	{
+		TSharedRef<FAliasModelToCADKernelConverter> AliasToCADKernelConverter = MakeShared<FAliasModelToCADKernelConverter>(ImportParameters);
+		ModelConverter = AliasToCADKernelConverter;
+		BRepConverter = &AliasToCADKernelConverter.Get();
+	}
+
+	ModelConverter->SetImportParameters(TessellationOptions.ChordTolerance, TessellationOptions.MaxEdgeLength, TessellationOptions.NormalTolerance, (CADLibrary::EStitchingTechnique)TessellationOptions.StitchingTechnique);
+
+	return ModelConverter;
+}
+
 TOptional<FMeshDescription> FWireTranslatorImpl::MeshDagNodeWithExternalMesher(AlDagNode& DagNode, TSharedRef<IDatasmithMeshElement> MeshElement, CADLibrary::FMeshParameters& MeshParameters)
 {
 	// Wire unit is cm
-	CADModelConverter->InitializeProcess();
+	IAliasBRepConverter* BRepConverter = nullptr;
+	TSharedPtr<CADLibrary::ICADModelConverter> ModelConverter = GetModelConverter(BRepConverter);
+
+	ModelConverter->InitializeProcess();
 
 	EAliasObjectReference ObjectReference = EAliasObjectReference::LocalReference;
 
@@ -2059,19 +2120,19 @@ TOptional<FMeshDescription> FWireTranslatorImpl::MeshDagNodeWithExternalMesher(A
 		ColorPtr = &DefaultColor;
 	}
 
-	if (!AliasBRepConverter->AddBRep(DagNode, *ColorPtr, ObjectReference))
+	if (!BRepConverter->AddBRep(DagNode, *ColorPtr, ObjectReference))
 	{
 		return TOptional<FMeshDescription>();
 	}
 
-	CADModelConverter->RepairTopology();
+	ModelConverter->RepairTopology();
 
 	CADModelConverter->SaveModel(*OutputPath, MeshElement);
 
 	FMeshDescription MeshDescription;
 	DatasmithMeshHelper::PrepareAttributeForStaticMesh(MeshDescription);
 
-	bool bRet = CADModelConverter->Tessellate(MeshParameters, MeshDescription);
+	bool bRet = ModelConverter->Tessellate(MeshParameters, MeshDescription);
 	if (!bRet)
 	{
 		const TCHAR* StaticMeshLable = MeshElement->GetLabel();
@@ -2085,7 +2146,10 @@ TOptional<FMeshDescription> FWireTranslatorImpl::MeshDagNodeWithExternalMesher(A
 TOptional<FMeshDescription> FWireTranslatorImpl::MeshDagNodeWithExternalMesher(TSharedRef<BodyData> Body, TSharedRef<IDatasmithMeshElement> MeshElement, CADLibrary::FMeshParameters& MeshParameters)
 {
 	// Wire unit is cm
-	CADModelConverter->InitializeProcess();
+	IAliasBRepConverter* BRepConverter = nullptr;
+	TSharedPtr<CADLibrary::ICADModelConverter> ModelConverter = GetModelConverter(BRepConverter);
+
+	ModelConverter->InitializeProcess();
 
 	EAliasObjectReference ObjectReference = EAliasObjectReference::LocalReference;
 	if (MeshParameters.bIsSymmetric)
@@ -2101,24 +2165,24 @@ TOptional<FMeshDescription> FWireTranslatorImpl::MeshDagNodeWithExternalMesher(T
 
 	for (const TPair<TSharedPtr<AlDagNode>, FColor>& DagNode : Body->GetShells())
 	{
-		AliasBRepConverter->AddBRep(*DagNode.Key, DagNode.Value, ObjectReference);
+		BRepConverter->AddBRep(*DagNode.Key, DagNode.Value, ObjectReference);
 	}
 
-	CADModelConverter->RepairTopology();
+	ModelConverter->RepairTopology();
 
-	CADModelConverter->SaveModel(*OutputPath, MeshElement);
+	ModelConverter->SaveModel(*OutputPath, MeshElement);
 
 	FMeshDescription MeshDescription;
 	DatasmithMeshHelper::PrepareAttributeForStaticMesh(MeshDescription);
 
-	CADModelConverter->Tessellate(MeshParameters, MeshDescription);
+	ModelConverter->Tessellate(MeshParameters, MeshDescription);
 
 	return MoveTemp(MeshDescription);
 }
 
 TOptional<FMeshDescription> FWireTranslatorImpl::GetMeshOfShellNode(AlDagNode& DagNode, TSharedRef<IDatasmithMeshElement> MeshElement, CADLibrary::FMeshParameters& MeshParameters)
 {
-	if (CADModelConverter->IsSessionValid())
+	if (CADModelConverter.IsValid() && CADModelConverter->IsSessionValid())
 	{
 		TOptional<FMeshDescription> UEMesh = MeshDagNodeWithExternalMesher(DagNode, MeshElement, MeshParameters);
 		return UEMesh;
@@ -2165,7 +2229,7 @@ TOptional<FMeshDescription> FWireTranslatorImpl::GetMeshOfMeshBody(TSharedRef<Bo
 		{
 			continue;
 		}
-		TUniquePtr<AlMesh> Mesh(MeshPtr);
+		TAlObjectPtr<AlMesh> Mesh(MeshPtr);
 
 		AlMatrix4x4 AlMatrix;
 		DagNode.Key->localTransformationMatrix(AlMatrix);
@@ -2200,7 +2264,7 @@ TOptional<FMeshDescription> FWireTranslatorImpl::GetMeshOfNodeMesh(AlDagNode& Te
 		return TOptional<FMeshDescription>();
 	}
 
-	TUniquePtr<AlMesh> SharedMesh(Mesh);
+	TAlObjectPtr<AlMesh> SharedMesh(Mesh);
 	if (AlMeshInvGlobalMatrix != nullptr)
 	{
 		SharedMesh->transform(*AlMeshInvGlobalMatrix);
@@ -2221,7 +2285,7 @@ TOptional<FMeshDescription> FWireTranslatorImpl::GetMeshDescription(TSharedRef<I
 	AlLayer* LayerPtr = DagNode->layer();
 	if (AlIsValid(LayerPtr))
 	{
-		const TUniquePtr<AlLayer> Layer(LayerPtr);
+		const TAlObjectPtr<AlLayer> Layer(LayerPtr);
 		if (LayerPtr->isSymmetric())
 		{
 			MeshParameters.bIsSymmetric = true;
@@ -2273,7 +2337,7 @@ TOptional<FMeshDescription> FWireTranslatorImpl::GetMeshDescription(TSharedRef<I
 		AlLayer* LayerPtr = DagNode.layer();
 		if (AlIsValid(LayerPtr))
 		{
-			const TUniquePtr<AlLayer> Layer(LayerPtr);
+			const TAlObjectPtr<AlLayer> Layer(LayerPtr);
 			if (LayerPtr->isSymmetric())
 			{
 				MeshParameters.bIsSymmetric = true;
@@ -2342,14 +2406,46 @@ TOptional<FMeshDescription> FWireTranslatorImpl::ImportMesh(AlMesh& InMesh, TSha
 	return MoveTemp(MeshDescription);
 }
 
+#define TRACK_MESH 1
+
 bool FWireTranslatorImpl::LoadStaticMesh(const TSharedRef<IDatasmithMeshElement> MeshElement, FDatasmithMeshElementPayload& OutMeshPayload, const FDatasmithTessellationOptions& InTessellationOptions)
 {
+#if TRACK_MESH
+	static const FString MeshName(TEXT("shell_30181"));
+	bool bSkip = true;
+	if (MeshName.Equals(MeshElement->GetLabel()))
+	{
+		bSkip = false;
+	}
+	if (bSkip)
+	{
+		return false;
+	}
+#endif
 	CADLibrary::FMeshParameters MeshParameters;
+
+	uint64 StartTime = FPlatformTime::Cycles64();
+
 	if (TOptional<FMeshDescription> Mesh = GetMeshDescription(MeshElement, MeshParameters))
 	{
 		OutMeshPayload.LodMeshes.Add(MoveTemp(Mesh.GetValue()));
-		CADModelConverter->AddSurfaceDataForMesh(MeshElement->GetFile(), MeshParameters, InTessellationOptions, OutMeshPayload);
+		const TCHAR* MeshFilename = MeshElement->GetFile();
+		if (!bAliasUseNative && FPaths::FileExists(MeshFilename))
+		{
+			CADModelConverter->AddSurfaceDataForMesh(MeshFilename, MeshParameters, InTessellationOptions, OutMeshPayload);
+
+			// Remove the file because it is temporary since caching is disabled.
+			if (!CADLibrary::FImportParameters::bGEnableCADCache)
+			{
+				IFileManager::Get().Delete(MeshFilename);
+			}
+		}
 	}
+
+	double ElapsedSeconds = FPlatformTime::ToSeconds64(FPlatformTime::Cycles64() - StartTime);
+
+	UE_LOG(LogDatasmithWireTranslator, Display, TEXT("Mesh description of %s retrieved in %.3f s"), MeshElement->GetLabel(), ElapsedSeconds);
+
 	return OutMeshPayload.LodMeshes.Num() > 0;
 }
 
@@ -2385,6 +2481,10 @@ void FDatasmithWireTranslator::Initialize(FDatasmithTranslatorCapabilities& OutC
 					UE_LOG(LogDatasmithWireTranslator, Warning, TEXT(WRONG_VERSION_TEXT)); return true;
 				};
 				static const bool bIsDispaly = DisplayMessage();
+
+				const IConsoleVariable* ConsoleVariable = IConsoleManager::Get().FindConsoleVariable(TEXT("ds.Wiretranslator.ThreadSafe"));
+				OutCapabilities.bParallelLoadStaticMeshSupported = ConsoleVariable ? ConsoleVariable->GetBool() : false;
+
 				OutCapabilities.bIsEnabled = false;
 				return;
 			}
@@ -2394,6 +2494,10 @@ void FDatasmithWireTranslator::Initialize(FDatasmithTranslatorCapabilities& OutC
 			{
 				const FString AliasVersion = FString::Printf(TEXT("AliasStudio %s model files"), *AliasSdkVersion);
 				OutCapabilities.SupportedFileFormats.Add(FFileFormatInfo{ TEXT("wire"), *AliasVersion });
+				
+				const IConsoleVariable* ConsoleVariable = IConsoleManager::Get().FindConsoleVariable(TEXT("ds.Wiretranslator.ThreadSafe"));
+				OutCapabilities.bParallelLoadStaticMeshSupported = ConsoleVariable ? ConsoleVariable->GetBool() : false;
+				
 				OutCapabilities.bIsEnabled = true;
 				return;
 			}
@@ -2476,6 +2580,7 @@ bool FDatasmithWireTranslator::LoadScene(TSharedRef<IDatasmithScene> OutScene)
 
 void FDatasmithWireTranslator::UnloadScene()
 {
+	Translator.Reset();
 }
 
 bool FDatasmithWireTranslator::LoadStaticMesh(const TSharedRef<IDatasmithMeshElement> MeshElement, FDatasmithMeshElementPayload& OutMeshPayload)
@@ -2504,7 +2609,3 @@ void FDatasmithWireTranslator::SetSceneImportOptions(const TArray<TObjectPtr<UDa
 } // namespace
 
 #undef LOCTEXT_NAMESPACE // "DatasmithWireTranslator"
-
-#if PLATFORM_WINDOWS
-#include "Windows/HideWindowsPlatformTypes.h"
-#endif
\ No newline at end of file
diff --git a/Engine/Plugins/Enterprise/DatasmithCADImporter/Source/DatasmithWireTranslator/Private/DatasmithWireTranslator.h b/Engine/Plugins/Enterprise/DatasmithCADImporter/Source/DatasmithWireTranslator/Private/DatasmithWireTranslator.h
--- a/Engine/Plugins/Enterprise/DatasmithCADImporter/Source/DatasmithWireTranslator/Private/DatasmithWireTranslator.h
+++ b/Engine/Plugins/Enterprise/DatasmithCADImporter/Source/DatasmithWireTranslator/Private/DatasmithWireTranslator.h
@@ -4,6 +4,7 @@
 
 #include "CoreMinimal.h"
 
+#include "OpenModelUtils.h"
 #include "ParametricSurfaceTranslator.h"
 #include "UObject/ObjectMacros.h"
 
diff --git a/Engine/Plugins/Enterprise/DatasmithCADImporter/Source/DatasmithWireTranslator/Private/OpenModelUtils.h b/Engine/Plugins/Enterprise/DatasmithCADImporter/Source/DatasmithWireTranslator/Private/OpenModelUtils.h
--- a/Engine/Plugins/Enterprise/DatasmithCADImporter/Source/DatasmithWireTranslator/Private/OpenModelUtils.h
+++ b/Engine/Plugins/Enterprise/DatasmithCADImporter/Source/DatasmithWireTranslator/Private/OpenModelUtils.h
@@ -29,7 +29,32 @@ typedef double AlMatrix4x4[4][4];
 
 namespace UE_DATASMITHWIRETRANSLATOR_NAMESPACE
 {
+	template<typename T>
+	class TAlObjectPtr
+	{
+	public:
+		TAlObjectPtr(T* AlObject) : Ptr(AlObject) {}
+
+		bool IsValid() { return Ptr && AlIsValid(Ptr); }
+
+		T* operator->() const
+		{
+			return Ptr;
+		}
+
+		T& operator*() const
+		{
+			return *Ptr;
+		}
+
+		T* Get() const
+		{
+			return Ptr;
+		}
 
+	private:
+		T* Ptr = nullptr;
+	};
 enum class ETesselatorType : uint8
 {
 	Fast,
