Rearchitected WireTranaslator plugin - Part #1

#rnx

diff --git a/Engine/Plugins/Enterprise/DatasmithCADImporter/DatasmithCADImporter.uplugin b/Engine/Plugins/Enterprise/DatasmithCADImporter/DatasmithCADImporter.uplugin
--- a/Engine/Plugins/Enterprise/DatasmithCADImporter/DatasmithCADImporter.uplugin
+++ b/Engine/Plugins/Enterprise/DatasmithCADImporter/DatasmithCADImporter.uplugin
@@ -71,7 +71,7 @@
 			]
 		},
 		{
-			"Name": "DatasmithWireTranslator2020",
+			"Name": "DatasmithWireTranslator",
 			"Type": "Editor",
 			"LoadingPhase": "Default",
 			"PlatformAllowList": [
@@ -79,7 +79,7 @@
 			]
 		},
 		{
-			"Name": "DatasmithWireTranslator2021_3",
+			"Name": "WireInterface2020",
 			"Type": "Editor",
 			"LoadingPhase": "Default",
 			"PlatformAllowList": [
@@ -87,7 +87,7 @@
 			]
 		},
 		{
-			"Name": "DatasmithWireTranslator2022",
+			"Name": "WireInterface2021_3",
 			"Type": "Editor",
 			"LoadingPhase": "Default",
 			"PlatformAllowList": [
@@ -95,7 +95,7 @@
 			]
 		},
 		{
-			"Name": "DatasmithWireTranslator2022_1",
+			"Name": "WireInterface2022",
 			"Type": "Editor",
 			"LoadingPhase": "Default",
 			"PlatformAllowList": [
@@ -103,7 +103,7 @@
 			]
 		},
 		{
-			"Name": "DatasmithWireTranslator2022_2",
+			"Name": "WireInterface2022_1",
 			"Type": "Editor",
 			"LoadingPhase": "Default",
 			"PlatformAllowList": [
@@ -111,7 +111,7 @@
 			]
 		},
 		{
-			"Name": "DatasmithWireTranslator2023_0",
+			"Name": "WireInterface2022_2",
 			"Type": "Editor",
 			"LoadingPhase": "Default",
 			"PlatformAllowList": [
@@ -119,7 +119,15 @@
 			]
 		},
 		{
-			"Name": "DatasmithWireTranslator2023_1",
+			"Name": "WireInterface2023_0",
+			"Type": "Editor",
+			"LoadingPhase": "Default",
+			"PlatformAllowList": [
+				"Win64"
+			]
+		},
+		{
+			"Name": "WireInterface2023_1",
 			"Type": "Editor",
 			"LoadingPhase": "Default",
 			"PlatformAllowList": [
diff --git a/Engine/Plugins/Enterprise/DatasmithCADImporter/Source/CADKernelSurface/Public/CADModelToCADKernelConverterBase.h b/Engine/Plugins/Enterprise/DatasmithCADImporter/Source/CADKernelSurface/Public/CADModelToCADKernelConverterBase.h
--- a/Engine/Plugins/Enterprise/DatasmithCADImporter/Source/CADKernelSurface/Public/CADModelToCADKernelConverterBase.h
+++ b/Engine/Plugins/Enterprise/DatasmithCADImporter/Source/CADKernelSurface/Public/CADModelToCADKernelConverterBase.h
@@ -54,12 +54,17 @@ public:
 		return true;
 	}
 
-	virtual bool SaveModel(const TCHAR* InFolderPath, TSharedRef<IDatasmithMeshElement>& MeshElement) override
+	virtual bool SaveModel(const TCHAR* InFolderPath, TSharedPtr<IDatasmithMeshElement> MeshElement) override
 	{
-		FString FilePath = FPaths::Combine(InFolderPath, MeshElement->GetName()) + TEXT(".ugeom");
-		CADKernelSession.SaveDatabase(*FilePath);
-		MeshElement->SetFile(*FilePath);
-		return true;
+		if (MeshElement.IsValid())
+		{
+			FString FilePath = FPaths::Combine(InFolderPath, MeshElement->GetName()) + TEXT(".ugeom");
+			CADKernelSession.SaveDatabase(*FilePath);
+			MeshElement->SetFile(*FilePath);
+			return true;
+		}
+
+		return false;
 	}
 
 	virtual bool Tessellate(const CADLibrary::FMeshParameters& InMeshParameters, FMeshDescription& OutMeshDescription) override
diff --git a/Engine/Plugins/Enterprise/DatasmithCADImporter/Source/DatasmithWireTranslator/DatasmithWireTranslator.Build.cs b/Engine/Plugins/Enterprise/DatasmithCADImporter/Source/DatasmithWireTranslator/DatasmithWireTranslator.Build.cs
new file mode 100644
--- /dev/null
+++ b/Engine/Plugins/Enterprise/DatasmithCADImporter/Source/DatasmithWireTranslator/DatasmithWireTranslator.Build.cs
@@ -0,0 +1,42 @@
+// Copyright Epic Games, Inc. All Rights Reserved.
+
+using System.IO;
+using UnrealBuildTool;
+
+public class DatasmithWireTranslator : ModuleRules
+{
+	public DatasmithWireTranslator(ReadOnlyTargetRules Target) : base(Target)
+	{
+		//bUseUnity = false;
+		//OptimizeCode = CodeOptimization.Never;
+		//PCHUsage = ModuleRules.PCHUsageMode.NoPCHs;
+
+		PublicDependencyModuleNames.AddRange(
+			new string[]
+			{
+				"Core",
+				"CoreUObject",
+			}
+		);
+
+		PrivateDependencyModuleNames.AddRange(
+			new string[]
+			{
+				"DatasmithContent",
+				"DatasmithTranslator",
+				"ParametricSurface",
+			}
+		);
+
+		if (Target.bBuildEditor == true)
+		{
+			PrivateDependencyModuleNames.AddRange(
+				new string[]
+				{
+ 					"MessageLog",
+ 					"UnrealEd",
+				}
+			);
+		}
+	}
+}
\ No newline at end of file
diff --git a/Engine/Plugins/Enterprise/DatasmithCADImporter/Source/DatasmithWireTranslator/DatasmithWireTranslator2020.Build.cs b/Engine/Plugins/Enterprise/DatasmithCADImporter/Source/DatasmithWireTranslator/DatasmithWireTranslator2020.Build.cs
deleted file mode 100644
diff --git a/Engine/Plugins/Enterprise/DatasmithCADImporter/Source/DatasmithWireTranslator/DatasmithWireTranslator2021_3.Build.cs b/Engine/Plugins/Enterprise/DatasmithCADImporter/Source/DatasmithWireTranslator/DatasmithWireTranslator2021_3.Build.cs
deleted file mode 100644
diff --git a/Engine/Plugins/Enterprise/DatasmithCADImporter/Source/DatasmithWireTranslator/DatasmithWireTranslator2022.Build.cs b/Engine/Plugins/Enterprise/DatasmithCADImporter/Source/DatasmithWireTranslator/DatasmithWireTranslator2022.Build.cs
deleted file mode 100644
diff --git a/Engine/Plugins/Enterprise/DatasmithCADImporter/Source/DatasmithWireTranslator/DatasmithWireTranslator2022_1.Build.cs b/Engine/Plugins/Enterprise/DatasmithCADImporter/Source/DatasmithWireTranslator/DatasmithWireTranslator2022_1.Build.cs
deleted file mode 100644
diff --git a/Engine/Plugins/Enterprise/DatasmithCADImporter/Source/DatasmithWireTranslator/DatasmithWireTranslator2022_2.Build.cs b/Engine/Plugins/Enterprise/DatasmithCADImporter/Source/DatasmithWireTranslator/DatasmithWireTranslator2022_2.Build.cs
deleted file mode 100644
diff --git a/Engine/Plugins/Enterprise/DatasmithCADImporter/Source/DatasmithWireTranslator/DatasmithWireTranslator2023_0.Build.cs b/Engine/Plugins/Enterprise/DatasmithCADImporter/Source/DatasmithWireTranslator/DatasmithWireTranslator2023_0.Build.cs
deleted file mode 100644
diff --git a/Engine/Plugins/Enterprise/DatasmithCADImporter/Source/DatasmithWireTranslator/DatasmithWireTranslator2023_1.Build.cs b/Engine/Plugins/Enterprise/DatasmithCADImporter/Source/DatasmithWireTranslator/DatasmithWireTranslator2023_1.Build.cs
deleted file mode 100644
diff --git a/Engine/Plugins/Enterprise/DatasmithCADImporter/Source/DatasmithWireTranslator/Private/AliasBrepConverter.h b/Engine/Plugins/Enterprise/DatasmithCADImporter/Source/DatasmithWireTranslator/Private/AliasBrepConverter.h
deleted file mode 100644
diff --git a/Engine/Plugins/Enterprise/DatasmithCADImporter/Source/DatasmithWireTranslator/Private/AliasModelToCADKernelConverter.cpp b/Engine/Plugins/Enterprise/DatasmithCADImporter/Source/DatasmithWireTranslator/Private/AliasModelToCADKernelConverter.cpp
deleted file mode 100644
diff --git a/Engine/Plugins/Enterprise/DatasmithCADImporter/Source/DatasmithWireTranslator/Private/AliasModelToCADKernelConverter.h b/Engine/Plugins/Enterprise/DatasmithCADImporter/Source/DatasmithWireTranslator/Private/AliasModelToCADKernelConverter.h
deleted file mode 100644
diff --git a/Engine/Plugins/Enterprise/DatasmithCADImporter/Source/DatasmithWireTranslator/Private/AliasModelToTechSoftConverter.cpp b/Engine/Plugins/Enterprise/DatasmithCADImporter/Source/DatasmithWireTranslator/Private/AliasModelToTechSoftConverter.cpp
deleted file mode 100644
diff --git a/Engine/Plugins/Enterprise/DatasmithCADImporter/Source/DatasmithWireTranslator/Private/AliasModelToTechSoftConverter.h b/Engine/Plugins/Enterprise/DatasmithCADImporter/Source/DatasmithWireTranslator/Private/AliasModelToTechSoftConverter.h
deleted file mode 100644
diff --git a/Engine/Plugins/Enterprise/DatasmithCADImporter/Source/DatasmithWireTranslator/Private/DatasmithWireTranslator.cpp b/Engine/Plugins/Enterprise/DatasmithCADImporter/Source/DatasmithWireTranslator/Private/DatasmithWireTranslator.cpp
--- a/Engine/Plugins/Enterprise/DatasmithCADImporter/Source/DatasmithWireTranslator/Private/DatasmithWireTranslator.cpp
+++ b/Engine/Plugins/Enterprise/DatasmithCADImporter/Source/DatasmithWireTranslator/Private/DatasmithWireTranslator.cpp
@@ -2,2547 +2,160 @@
 
 #include "DatasmithWireTranslator.h"
 
-#include "CADOptions.h"
-#include "Containers/List.h"
-#include "DatasmithImportOptions.h"
-#include "DatasmithSceneFactory.h"
-#include "DatasmithTranslator.h"
-#include "DatasmithUtils.h"
-#include "DatasmithWireTranslatorModule.h"
-#include "GenericPlatform/GenericPlatformTLS.h"
-#include "HAL/ConsoleManager.h"
-#include "IDatasmithSceneElements.h"
-#include "Misc/FileHelper.h"
-#include "Misc/Paths.h"
-#include "Utility/DatasmithMeshHelper.h"
-
-#include "StaticMeshDescription.h"
-#include "StaticMeshOperations.h"
-
-#if WITH_EDITOR
-#include "Editor.h"
-#include "IMessageLogListing.h"
-#include "Logging/TokenizedMessage.h"
-#include "MessageLogModule.h"
-#include "Modules/ModuleManager.h"
-#endif
-
-#include "AliasBrepConverter.h"
-#include "AliasModelToCADKernelConverter.h"
-#include "AliasModelToTechSoftConverter.h" // requires Techsoft as public dependency
-#include "CADInterfacesModule.h"
-
-#if PLATFORM_WINDOWS
-#include "Windows/AllowWindowsPlatformTypes.h"
-#endif
-
-#ifdef USE_OPENMODEL
-#include <AlChannel.h>
-#include <AlDagNode.h>
-#include <AlGroupNode.h>
-#include <AlLayer.h>
-#include <AlLinkItem.h>
-#include <AlList.h>
-#include <AlMesh.h>
-#include <AlMeshNode.h>
-#include <AlPersistentID.h>
-#include <AlRetrieveOptions.h>
-#include <AlShader.h>
-#include <AlShadingFieldItem.h>
-#include <AlShell.h>
-#include <AlShellNode.h>
-#include <AlSurface.h>
-#include <AlSurfaceNode.h>
-#include <AlTesselate.h>
-#include <AlTrimRegion.h>
-#include <AlTM.h>
-#include <AlUniverse.h>
-#endif
-
-#if PLATFORM_WINDOWS
-#include "Windows/HideWindowsPlatformTypes.h"
-#endif
-
-DEFINE_LOG_CATEGORY_STATIC(LogDatasmithWireTranslator, Log, All);
-
-#define LOCTEXT_NAMESPACE "DatasmithWireTranslator"
-
-#define WRONG_VERSION_TEXT "Unsupported version of Alias detected. Please upgrade to Alias 2021.3 (or later version)."
-#define CAD_INTERFACE_UNAVAILABLE "CAD Interface module is unavailable. Meshing will be done by Alias."
-
-#pragma optimize("", off)
-#ifdef OPEN_MODEL_2023_0
-
-static bool bGAliasSewByMaterial = false;
-FAutoConsoleVariableRef GAliasSewByMaterial(
-	TEXT("ds.CADTranslator.Alias.SewByMaterial"),
-	bGAliasSewByMaterial,
-	TEXT("Enable Sew action merges BReps according to their material i.e. only BReps associated with same material can be merged together.\
-Default is disable\n"),
-ECVF_Default);
-
-static bool bGAliasLayerAsActor = false;
-FAutoConsoleVariableRef GAliasLayerAsActor(
-	TEXT("ds.CADTranslator.Alias.LayersAsActors"),
-	bGAliasLayerAsActor,
-	TEXT("If true, the first level of actors in the outliner is the layers.\
-Default is false\n"),
-ECVF_Default);
-
-#endif 
-
-#ifdef IS_MAIN_MODULE
-static TAutoConsoleVariable<bool> CVarAliasThreadSafe(
-	TEXT("ds.WireTranslator.ThreadSafe"),
-	false,
-	TEXT("If true, the translator will be called in more than one thread. Default false.\n"),
-	ECVF_Default);
-
-static TAutoConsoleVariable<bool> CVarAliasUseNative(
-	TEXT("ds.WireTranslator.UseNative"),
-	false,
-	TEXT("If true, the AlaisStudio tessellator will be used. Default false.\n"),
-	ECVF_Default);
-
-static TAutoConsoleVariable<bool> CVarAliasSkipThinZoneMeshing(
-	TEXT("ds.WireTranslator.SkipThinZoneMeshing"),
-	true,
-	TEXT("If true, the global thin zone cvar will be honored. Default true.\n"),
-	ECVF_Default);
-#endif
-
-namespace UE_DATASMITHWIRETRANSLATOR_NAMESPACE
-{
-
-#ifdef USE_OPENMODEL
-
-static bool bGSewByMaterial = false;
-static bool bGLayersAsActors = false;
-
-static const FColor DefaultColor = FColor(200, 200, 200);
-
-const uint64 LibAliasNext_Version = 0xffffffffffffffffull;
-const uint64 LibAlias2023_1_0_Version = 8162778619576619;
-const uint64 LibAlias2023_0_0_Version = 8162774324609149;
-const uint64 LibAlias2022_2_0_Version = 7881307937833405;
-const uint64 LibAlias2022_1_0_Version = 7881303642865885;
-const uint64 LibAlias2022_0_1_Version = 7881299347964005;
-const uint64 LibAlias2021_3_2_Version = 7599833027117059;
-const uint64 LibAlias2021_3_1_Version = 7599824433840131;
-const uint64 LibAlias2021_3_0_Version = 7599824424206339;
-const uint64 LibAlias2021_Version = 7599824377020416;
-const uint64 LibAlias2020_Version = 7318349414924288;
-const uint64 LibAlias2019_Version = 5000000000000000;
-
-#if defined(OPEN_MODEL_2020)
-const uint64 LibAliasVersionMin = LibAlias2019_Version;
-const uint64 LibAliasVersionMax = LibAlias2021_3_0_Version;
-const FString AliasSdkVersion   = TEXT("2020");
-#elif defined(OPEN_MODEL_2021_3)
-const uint64 LibAliasVersionMin = LibAlias2021_3_0_Version;
-const uint64 LibAliasVersionMax = LibAlias2022_0_1_Version;
-const FString AliasSdkVersion   = TEXT("2021.3");
-#elif defined(OPEN_MODEL_2022)
-const uint64 LibAliasVersionMin = LibAlias2022_0_1_Version;
-const uint64 LibAliasVersionMax = LibAlias2022_1_0_Version;
-const FString AliasSdkVersion   = TEXT("2022");
-#elif defined(OPEN_MODEL_2022_1)
-const uint64 LibAliasVersionMin = LibAlias2022_1_0_Version;
-const uint64 LibAliasVersionMax = LibAlias2022_2_0_Version;
-const FString AliasSdkVersion   = TEXT("2022.1");
-#elif defined(OPEN_MODEL_2022_2)
-const uint64 LibAliasVersionMin = LibAlias2022_2_0_Version;
-const uint64 LibAliasVersionMax = LibAlias2023_0_0_Version;
-const FString AliasSdkVersion   = TEXT("2022.2");
-#elif defined(OPEN_MODEL_2023_0)
-const uint64 LibAliasVersionMin = LibAlias2023_0_0_Version;
-const uint64 LibAliasVersionMax = LibAlias2023_1_0_Version;
-const FString AliasSdkVersion   = TEXT("2023.0");
-#elif defined(OPEN_MODEL_2023_1)      
-const uint64 LibAliasVersionMin = LibAlias2023_1_0_Version;
-const uint64 LibAliasVersionMax = LibAliasNext_Version;
-const FString AliasSdkVersion   = TEXT("2023.1");
-#endif
-
-// Alias material management (to allow sew of BReps of different materials):
-// To be compatible with "Retessellate" function, Alias material management as to be the same than CAD (TechSoft) import. 
-// As a reminder: the name and slot of UE Material from CAD is based on CAD material/color data i.e. RGBA Color components => "UE Material slot" (int32) and "UE Material name" (FString = FString::FromInt("UE Material slot"))
-// UE Material Label is free.
-// 
-// During Retessellate step, Color/Material of each CAD Faces are known, so "UE Material slot" can be deduced.
-// 
-// For Alias import:
-// Alias BRep is exported in CAD modeler (CADKernel, TechSoft, ...)
-// Material is build in UE, 
-// From an Alias Material, a unique Color is generated.
-// This Color is associated to the BRep Shell/face in the CAD modeler
-// The name and slot of the associated UE Material is defined from this color
-// So at the retessellate step, nothing changes from the CAD Retessellate process
-// 
-// The unique Color of an Alias Material is defined as follows:
-// TypeHash(Alias Material Name) => uint24 == 3 uint8 => RGB components of the color  
-
-FColor CreateShaderColorFromShaderName(const FString& ShaderName)
-{
-	const uint32 ShaderHash = FCrc::Strihash_DEPRECATED(*ShaderName);
-	const uint32 Red = (ShaderHash & 0xff000000) >> 24;
-	const uint32 Green = (ShaderHash & 0x00ff0000) >> 16;
-	const uint32 Blue = (ShaderHash & 0x0000ff00) >> 8;
-	return FColor(Red, Green, Blue);
-}
-
-int32 CreateShaderId(const FColor& ShaderColor)
-{
-	return FMath::Abs((int32)GetTypeHash(ShaderColor));
-}
-
-class BodyData
-{
-private:
-	TArray<TPair<TSharedPtr<AlDagNode>, FColor>> Shells;
-	TSet<FString> ShaderNames;
-
-	FString Label;
-	FString LayerName;
-	bool bCadData = true;
-
-public:
-
-	BodyData(const FString& InLayerName, bool bInCadData)
-		: bCadData(bInCadData)
-	{
-		LayerName = InLayerName;
-	};
-
-	bool IsCadData() const
-	{
-		return bCadData;
-	}
-
-	void SetLabel(const FString& InLabel)
-	{
-		Label = InLabel;
-	}
-
-	FColor LastColor;
-	bool bWarningNotLogYet = true;
-	void Add(const TSharedPtr<AlDagNode>& Node, FColor Color, const FString& ShaderName)
-	{
-		using namespace CADLibrary;
-		if (ShaderNames.Num() == FImportParameters::GMaxMaterialCountPerMesh)
-		{
-			if (!ShaderNames.Contains(ShaderName))
-			{
-				UE_LOG(LogDatasmithWireTranslator, Warning, TEXT("The main UE5 rendering systems do not support more than 256 materials per mesh and the limit has been defined to %d materials. Only the first %d materials are kept. The others are replaced by the last one"), FImportParameters::GMaxMaterialCountPerMesh, FImportParameters::GMaxMaterialCountPerMesh);
-				bWarningNotLogYet = false;
-				Color = LastColor;
-			}
-		}
-		else
-		{
-			ShaderNames.Add(ShaderName);
-			LastColor = Color;
-		}
-		Shells.Emplace(Node, Color);
-	}
-
-	void Reserve(int32 MaxElement)
-	{
-		Shells.Reserve(MaxElement);
-		ShaderNames.Reserve(MaxElement);
-	}
-
-	const TSet<FString>& GetShaderNames() const
-	{
-		return ShaderNames;
-	}
-
-	const FString& GetLayerName() const
-	{
-		return LayerName;
-	}
-
-	const TArray<TPair<TSharedPtr<AlDagNode>, FColor>>& GetShells() const
-	{
-		return Shells;
-	}
-
-	// Generates BodyData's unique id from AlDagNode objects
-	uint32 GetUuid(const uint32& ParentUuid)
-	{
-		if (Shells.Num() == 0)
-		{
-			return ParentUuid;
-		}
-
-		if (Shells.Num() > 1)
-		{
-			Shells.Sort([&](const TPair<TSharedPtr<AlDagNode>, FColor>& NodeA, const TPair<TSharedPtr<AlDagNode>, FColor>& NodeB) { return OpenModelUtils::GetAlDagNodeUuid(*NodeA.Key) < OpenModelUtils::GetAlDagNodeUuid(*NodeB.Key); });
-		}
-
-		uint32 BodyUUID = 0;
-		for (const TPair<TSharedPtr<AlDagNode>, FColor>& DagNode : Shells)
-		{
-			BodyUUID = HashCombine(BodyUUID, OpenModelUtils::GetAlDagNodeUuid(*DagNode.Key));
-		}
-
-		return HashCombine(ParentUuid, BodyUUID);
-	}
-};
-
-uint32 GetSceneFileHash(const FString& FullPath, const FString& FileName)
-{
-	FFileStatData FileStatData = IFileManager::Get().GetStatData(*FullPath);
-
-	int64 FileSize = FileStatData.FileSize;
-	FDateTime ModificationTime = FileStatData.ModificationTime;
-
-	uint32 FileHash = GetTypeHash(FileName);
-	FileHash = HashCombine(FileHash, GetTypeHash(FileSize));
-	FileHash = HashCombine(FileHash, GetTypeHash(ModificationTime));
-
-	return FileHash;
-}
-
-class FWireTranslatorImpl
-{
-public:
-	FWireTranslatorImpl(const FString& InSceneFullName, TSharedRef<IDatasmithScene> InScene)
-		: DatasmithScene(InScene)
-		, SceneName(FPaths::GetBaseFilename(InSceneFullName))
-		, CurrentPath(FPaths::GetPath(InSceneFullName))
-		, SceneFullPath(InSceneFullName)
-		, SceneFileHash(0)
-		, AlRootNode(nullptr)
-	{
-
-		const FString AliasProductVersion = FString::Printf(TEXT("Alias %s"), *AliasSdkVersion);
-		DatasmithScene->SetHost(TEXT("Alias"));
-		DatasmithScene->SetVendor(TEXT("Autodesk"));
-		DatasmithScene->SetProductName(TEXT("Alias Tools"));
-		DatasmithScene->SetExporterSDKVersion(*AliasSdkVersion);
-		DatasmithScene->SetProductVersion(*AliasProductVersion);
-
-		const IConsoleVariable* ConsoleVariable = IConsoleManager::Get().FindConsoleVariable(TEXT("ds.Wiretranslator.UseNative"));
-		bAliasUseNative = ConsoleVariable ? ConsoleVariable->GetBool() : false;
-		if (!bAliasUseNative)
-		{
-			CADLibrary::FImportParameters ImportParameters;
-			if (CADLibrary::FImportParameters::bGDisableCADKernelTessellation)
-			{
-				TSharedRef<FAliasModelToTechSoftConverter> AliasToTechSoftConverter = MakeShared<FAliasModelToTechSoftConverter>(ImportParameters);
-				CADModelConverter = AliasToTechSoftConverter;
-				AliasBRepConverter = AliasToTechSoftConverter;
-			}
-			else
-			{
-				TSharedRef<FAliasModelToCADKernelConverter> AliasToCADKernelConverter = MakeShared<FAliasModelToCADKernelConverter>(ImportParameters);
-				CADModelConverter = AliasToCADKernelConverter;
-				AliasBRepConverter = AliasToCADKernelConverter;
-			}
-		}
-
-		IConsoleVariable* CVarDSAliasSewByMaterialEnabled = IConsoleManager::Get().FindConsoleVariable(TEXT("ds.CADTranslator.Alias.SewByMaterial"));
-		if (CVarDSAliasSewByMaterialEnabled)
-		{
-			bGSewByMaterial = CVarDSAliasSewByMaterialEnabled->GetBool();
-		}
-
-		IConsoleVariable* CVarDSAliasLayersAsActorsEnabled = IConsoleManager::Get().FindConsoleVariable(TEXT("ds.CADTranslator.Alias.LayersAsActors"));
-		if (CVarDSAliasLayersAsActorsEnabled)
-		{
-			bGLayersAsActors = CVarDSAliasLayersAsActorsEnabled->GetBool();
-		}
-	}
-
-	~FWireTranslatorImpl()
-	{
-		AlUniverse::deleteAll();
-	}
-
-	bool Read();
-	TOptional<FMeshDescription> GetMeshDescription(TSharedRef<IDatasmithMeshElement> MeshElement, CADLibrary::FMeshParameters& OutMeshParameters);
-	TOptional<FMeshDescription> GetMeshDescription(TSharedRef<IDatasmithMeshElement> MeshElement, CADLibrary::FMeshParameters& MeshParameters, TSharedRef<BodyData> BodyTemp);
-
-	void SetTessellationOptions(const FDatasmithTessellationOptions& Options);
-	void SetOutputPath(const FString& Path) { OutputPath = Path; }
-
-	bool LoadStaticMesh(const TSharedRef<IDatasmithMeshElement> MeshElement, FDatasmithMeshElementPayload& OutMeshPayload, const FDatasmithTessellationOptions& InTessellationOptions);
-
-private:
-
-	struct FDagNodeInfo
-	{
-		uint32 Uuid;  // Use for actor name
-		FString Label;
-		TSharedPtr<IDatasmithActorElement> ActorElement;
-		FString LayerName;
-		bool bLayerIsVisible = true;
-
-		TMap<FString, TSharedPtr<IDatasmithActorElement>> LayerToActorElement;
-		FDagNodeInfo* ParentInfo = nullptr;
-
-		FDagNodeInfo(FDagNodeInfo& InParentInfo)
-			: ParentInfo(&InParentInfo)
-		{}
-
-		FDagNodeInfo()
-			: ParentInfo(nullptr)
-		{}
-
-		void ExtractDagNodeInfo(AlDagNode& CurrentNode);
-		void ExtractDagNodeInfo(BodyData& CurrentNode);
-		void ExtractDagNodeLayer(const AlDagNode& InDagNode);
-	};
-
-	bool GetDagLeaves();
-	bool GetShader();
-	bool RecurseDagForLeaves(const TSharedPtr<AlDagNode>& DagNode, FDagNodeInfo& ParentInfo);
-	void RecurseDagForLeavesNoMerge(const TSharedPtr<AlDagNode>& DagNode, FDagNodeInfo& ParentInfo);
-	void DagForLeavesNoMerge(const TSharedPtr<AlDagNode>& DagNode, FDagNodeInfo& ParentInfo);
-	bool ProcessAlGroupNode(AlDagNode& GroupNode, FDagNodeInfo& ParentInfo);
-	void ProcessAlShellNode(const TSharedPtr<AlDagNode>& ShellNode, FDagNodeInfo& ParentInfo, const FString& ShaderName);
-	void ProcessBodyNode(const TSharedPtr<BodyData>& Body, FDagNodeInfo& ParentInfo);
-	TSharedPtr<IDatasmithMeshElement> FindOrAddMeshElement(const TSharedPtr<BodyData>& Body, const FDagNodeInfo& ParentInfo);
-	TSharedPtr<IDatasmithMeshElement> FindOrAddMeshElement(const TSharedPtr<AlDagNode>& ShellNode, const FDagNodeInfo& ParentInfo, const FString& ShaderName);
-	TSharedPtr<IDatasmithActorElement> FindOrAddParentActor(FDagNodeInfo& ParentInfo, const FString& LayerName);
-	TSharedPtr<IDatasmithActorElement> FindOrAddLayerActor(const FString& LayerName);
-
-	void RecurseDeleteEmptyActor(TSharedPtr<IDatasmithActorElement> Actor);
-
-	TOptional<FMeshDescription> GetMeshOfShellNode(AlDagNode& DagNode, TSharedRef<IDatasmithMeshElement> MeshElement, CADLibrary::FMeshParameters& MeshParameters);
-	TOptional<FMeshDescription> GetMeshOfNodeMesh(AlDagNode& DagNode, TSharedRef<IDatasmithMeshElement> MeshElement, CADLibrary::FMeshParameters& MeshParameters, AlMatrix4x4* AlMeshInvGlobalMatrix = nullptr);
-	TOptional<FMeshDescription> GetMeshOfShellBody(TSharedRef<BodyData> DagNode, TSharedRef<IDatasmithMeshElement> MeshElement, CADLibrary::FMeshParameters& MeshParameters);
-	TOptional<FMeshDescription> GetMeshOfMeshBody(TSharedRef<BodyData> DagNode, TSharedRef<IDatasmithMeshElement> MeshElement, CADLibrary::FMeshParameters& MeshParameters);
-
-	void AddNodeInBodyGroup(TSharedPtr<AlDagNode>& DagNode, const FString& ShaderName, TMap<uint32, TSharedPtr<BodyData>>& ShellToProcess, bool bIsAPatch, uint32 MaxSize);
-
-	TSharedPtr<CADLibrary::ICADModelConverter> GetModelConverter(IAliasBRepConverter*& BRepConverter) const;
-
-	TOptional<FMeshDescription> MeshDagNodeWithExternalMesher(AlDagNode& DagNode, TSharedRef<IDatasmithMeshElement> MeshElement, CADLibrary::FMeshParameters& MeshParameters);
-	TOptional<FMeshDescription> MeshDagNodeWithExternalMesher(TSharedRef<BodyData> DagNode, TSharedRef<IDatasmithMeshElement> MeshElement, CADLibrary::FMeshParameters& MeshParameters);
-
-	TOptional<FMeshDescription> ImportMesh(AlMesh& Mesh, TSharedRef<IDatasmithMeshElement> MeshElement, CADLibrary::FMeshParameters& MeshParameters);
-
-	FORCEINLINE bool IsTransparent(FColor& TransparencyColor)
-	{
-		float Opacity = 1.0f - ((float)(TransparencyColor.R + TransparencyColor.G + TransparencyColor.B)) / 765.0f;
-		return !FMath::IsNearlyEqual(Opacity, 1.0f);
-	}
-
-	FORCEINLINE bool GetCommonParameters(AlShadingFields Field, double Value, FColor& Color, FColor& TransparencyColor, FColor& IncandescenceColor, double GlowIntensity)
-	{
-		switch (Field)
-		{
-		case AlShadingFields::kFLD_SHADING_COMMON_COLOR_R:
-			Color.R = (uint8)Value;
-			return true;
-		case AlShadingFields::kFLD_SHADING_COMMON_COLOR_G:
-			Color.G = (uint8)Value;
-			return true;
-		case  AlShadingFields::kFLD_SHADING_COMMON_COLOR_B:
-			Color.B = (uint8)Value;
-			return true;
-		case AlShadingFields::kFLD_SHADING_COMMON_INCANDESCENCE_R:
-			IncandescenceColor.R = (uint8)Value;
-			return true;
-		case AlShadingFields::kFLD_SHADING_COMMON_INCANDESCENCE_G:
-			IncandescenceColor.G = (uint8)Value;
-			return true;
-		case AlShadingFields::kFLD_SHADING_COMMON_INCANDESCENCE_B:
-			IncandescenceColor.B = (uint8)Value;
-			return true;
-		case  AlShadingFields::kFLD_SHADING_COMMON_TRANSPARENCY_R:
-			TransparencyColor.R = (uint8)Value;
-			return true;
-		case  AlShadingFields::kFLD_SHADING_COMMON_TRANSPARENCY_G:
-			TransparencyColor.G = (uint8)Value;
-			return true;
-		case  AlShadingFields::kFLD_SHADING_COMMON_TRANSPARENCY_B:
-			TransparencyColor.B = (uint8)Value;
-			return true;
-		case AlShadingFields::kFLD_SHADING_COMMON_GLOW_INTENSITY:
-			GlowIntensity = Value;
-			return true;
-		default:
-			return false;
-		}
-	}
-
-	void AddAlBlinnParameters(const TAlObjectPtr<AlShader>& Shader, TSharedRef<IDatasmithUEPbrMaterialElement> MaterialElement);
-	void AddAlLambertParameters(const TAlObjectPtr<AlShader>& Shader, TSharedRef<IDatasmithUEPbrMaterialElement> MaterialElement);
-	void AddAlLightSourceParameters(const TAlObjectPtr<AlShader>& Shader, TSharedRef<IDatasmithUEPbrMaterialElement> MaterialElement);
-	void AddAlPhongParameters(const TAlObjectPtr<AlShader>& Shader, TSharedRef<IDatasmithUEPbrMaterialElement> MaterialElement);
-
-private:
-	TSharedRef<IDatasmithScene> DatasmithScene;
-	FString SceneName;
-	FString CurrentPath;
-	FString OutputPath;
-	FString SceneFullPath;
-
-	// Hash value of the scene file used to check if the file has been modified for re-import
-	uint32 SceneFileHash;
-
-	TSharedPtr<AlDagNode> AlRootNode;
-
-	/** Table of correspondence between mesh identifier and associated Datasmith mesh element */
-	TMap<uint32, TSharedPtr<IDatasmithMeshElement>> ShellUuidToMeshElementMap;
-	TMap<uint32, TSharedPtr<IDatasmithMeshElement>> BodyUuidToMeshElementMap;
-
-	/** Datasmith mesh elements to OpenModel objects */
-	TMap<IDatasmithMeshElement*, TSharedPtr<AlDagNode>> MeshElementToAlDagNodeMap;
-
-	TMap<IDatasmithMeshElement*, TSharedPtr<BodyData>> MeshElementToBodyMap;
-
-	TMap<IDatasmithMeshElement*, FString> MeshElementToShaderName;
-	TMap<FString, FColor> ShaderNameToColor;
-	TMap<FString, TSharedPtr<IDatasmithMaterialIDElement>> ShaderNameToUEMaterialId;
-
-	TMap<FString, TSharedPtr<IDatasmithActorElement>> LayerNameToActor;
-
-	FDatasmithTessellationOptions TessellationOptions;
-
-	TSharedPtr<CADLibrary::ICADModelConverter> CADModelConverter;
-	TSharedPtr<IAliasBRepConverter> AliasBRepConverter;
-
-	bool bAliasUseNative;
-};
-
-void FWireTranslatorImpl::SetTessellationOptions(const FDatasmithTessellationOptions& Options)
-{
-	TessellationOptions = Options;
-	if (CADModelConverter)
-	{
-		CADModelConverter->SetImportParameters(Options.ChordTolerance, Options.MaxEdgeLength, Options.NormalTolerance, (CADLibrary::EStitchingTechnique)Options.StitchingTechnique);
-	}
-	SceneFileHash = HashCombine(Options.GetHash(), GetSceneFileHash(SceneFullPath, SceneName));
-}
-
-bool FWireTranslatorImpl::Read()
-{
-	// Initialize Alias.
-	AlUniverse::initialize();
-
-	if (AlUniverse::retrieve(TCHAR_TO_UTF8(*SceneFullPath)) != sSuccess)
-	{
-		return false;
-	}
-
-	AlRetrieveOptions options;
-	AlUniverse::retrieveOptions(options);
-
-	// Make material
-	if (!GetShader())
-	{
-		return false;
-	}
-
-	// Parse and extract the DAG leaf nodes.
-	// Note that Alias file unit is cm like UE
-	return GetDagLeaves();
-}
-
-void FWireTranslatorImpl::AddAlBlinnParameters(const TAlObjectPtr<AlShader>& Shader, TSharedRef<IDatasmithUEPbrMaterialElement> MaterialElement)
-{
-	// Default values for a Blinn material
-	FColor Color(145, 148, 153);
-	FColor TransparencyColor(0, 0, 0);
-	FColor IncandescenceColor(0, 0, 0);
-	FColor SpecularColor(38, 38, 38);
-	double Diffuse = 1.0;
-	double GlowIntensity = 0.0;
-	double Gloss = 0.8;
-	double Eccentricity = 0.35;
-	double Specularity = 1.0;
-	double Reflectivity = 0.5;
-	double SpecularRolloff = 0.5;
-
-	AlList* List = Shader->fields();
-	for (AlShadingFieldItem* Item = static_cast<AlShadingFieldItem*>(List->first()); Item; Item = Item->nextField())
-	{
-		double Value = 0.0f;
-		statusCode ErrorCode = Shader->parameter(Item->field(), Value);
-		if (ErrorCode != 0)
-		{
-			continue;
-		}
-
-		if (GetCommonParameters(Item->field(), Value, Color, TransparencyColor, IncandescenceColor, GlowIntensity))
-		{
-			continue;
-		}
-
-		switch (Item->field())
-		{
-		case AlShadingFields::kFLD_SHADING_BLINN_DIFFUSE:
-			Diffuse = Value;
-			break;
-		case AlShadingFields::kFLD_SHADING_BLINN_GLOSS_:
-			Gloss = Value;
-			break;
-		case AlShadingFields::kFLD_SHADING_BLINN_SPECULAR_R:
-			SpecularColor.R = (uint8)(255.f * Value);
-			break;
-		case AlShadingFields::kFLD_SHADING_BLINN_SPECULAR_G:
-			SpecularColor.G = (uint8)(255.f * Value);;
-			break;
-		case AlShadingFields::kFLD_SHADING_BLINN_SPECULAR_B:
-			SpecularColor.B = (uint8)(255.f * Value);;
-			break;
-		case AlShadingFields::kFLD_SHADING_BLINN_SPECULARITY_:
-			Specularity = Value;
-			break;
-		case AlShadingFields::kFLD_SHADING_BLINN_SPECULAR_ROLLOFF:
-			SpecularRolloff = Value;
-			break;
-		case AlShadingFields::kFLD_SHADING_BLINN_ECCENTRICITY:
-			Eccentricity = Value;
-			break;
-		case AlShadingFields::kFLD_SHADING_BLINN_REFLECTIVITY:
-			Reflectivity = Value;
-			break;
-		}
-	}
-
-	bool bIsTransparent = IsTransparent(TransparencyColor);
-
-	// Construct parameter expressions
-	IDatasmithMaterialExpressionScalar* DiffuseExpression = MaterialElement->AddMaterialExpression<IDatasmithMaterialExpressionScalar>();
-	DiffuseExpression->GetScalar() = Diffuse;
-	DiffuseExpression->SetName(TEXT("Diffuse"));
-
-	IDatasmithMaterialExpressionScalar* GlossExpression = MaterialElement->AddMaterialExpression<IDatasmithMaterialExpressionScalar>();
-	GlossExpression->GetScalar() = Gloss;
-	GlossExpression->SetName(TEXT("Gloss"));
-
-	IDatasmithMaterialExpressionColor* SpecularColorExpression = MaterialElement->AddMaterialExpression<IDatasmithMaterialExpressionColor>();
-	SpecularColorExpression->SetName(TEXT("SpecularColor"));
-	SpecularColorExpression->GetColor() = FLinearColor::FromSRGBColor(SpecularColor);
-
-	IDatasmithMaterialExpressionScalar* SpecularityExpression = MaterialElement->AddMaterialExpression<IDatasmithMaterialExpressionScalar>();
-	SpecularityExpression->GetScalar() = Specularity * 0.3;
-	SpecularityExpression->SetName(TEXT("Specularity"));
-
-	IDatasmithMaterialExpressionScalar* SpecularRolloffExpression = MaterialElement->AddMaterialExpression<IDatasmithMaterialExpressionScalar>();
-	SpecularRolloffExpression->GetScalar() = SpecularRolloff;
-	SpecularRolloffExpression->SetName(TEXT("SpecularRolloff"));
-
-	IDatasmithMaterialExpressionScalar* EccentricityExpression = MaterialElement->AddMaterialExpression<IDatasmithMaterialExpressionScalar>();
-	EccentricityExpression->GetScalar() = Eccentricity;
-	EccentricityExpression->SetName(TEXT("Eccentricity"));
-
-	IDatasmithMaterialExpressionScalar* ReflectivityExpression = MaterialElement->AddMaterialExpression<IDatasmithMaterialExpressionScalar>();
-	ReflectivityExpression->GetScalar() = Reflectivity;
-	ReflectivityExpression->SetName(TEXT("Reflectivity"));
-
-	IDatasmithMaterialExpressionColor* ColorExpression = MaterialElement->AddMaterialExpression<IDatasmithMaterialExpressionColor>();
-	ColorExpression->SetName(TEXT("Color"));
-	ColorExpression->GetColor() = FLinearColor::FromSRGBColor(Color);
-
-	IDatasmithMaterialExpressionColor* IncandescenceColorExpression = MaterialElement->AddMaterialExpression<IDatasmithMaterialExpressionColor>();
-	IncandescenceColorExpression->SetName(TEXT("IncandescenceColor"));
-	IncandescenceColorExpression->GetColor() = FLinearColor::FromSRGBColor(IncandescenceColor);
-
-	IDatasmithMaterialExpressionColor* TransparencyColorExpression = MaterialElement->AddMaterialExpression<IDatasmithMaterialExpressionColor>();
-	TransparencyColorExpression->SetName(TEXT("TransparencyColor"));
-	TransparencyColorExpression->GetColor() = FLinearColor::FromSRGBColor(TransparencyColor);
-
-	IDatasmithMaterialExpressionScalar* GlowIntensityExpression = MaterialElement->AddMaterialExpression<IDatasmithMaterialExpressionScalar>();
-	GlowIntensityExpression->GetScalar() = GlowIntensity;
-	GlowIntensityExpression->SetName(TEXT("GlowIntensity"));
-
-	// Create aux expressions
-	IDatasmithMaterialExpressionGeneric* ColorSpecLerp = MaterialElement->AddMaterialExpression<IDatasmithMaterialExpressionGeneric>();
-	ColorSpecLerp->SetExpressionName(TEXT("LinearInterpolate"));
-
-	IDatasmithMaterialExpressionScalar* ColorSpecLerpValue = MaterialElement->AddMaterialExpression<IDatasmithMaterialExpressionScalar>();
-	ColorSpecLerpValue->GetScalar() = 0.96f;
-
-	IDatasmithMaterialExpressionGeneric* ColorMetallicLerp = MaterialElement->AddMaterialExpression<IDatasmithMaterialExpressionGeneric>();
-	ColorMetallicLerp->SetExpressionName(TEXT("LinearInterpolate"));
-
-	IDatasmithMaterialExpressionGeneric* DiffuseLerp = MaterialElement->AddMaterialExpression<IDatasmithMaterialExpressionGeneric>();
-	DiffuseLerp->SetExpressionName(TEXT("LinearInterpolate"));
-
-	IDatasmithMaterialExpressionScalar* DiffuseLerpA = MaterialElement->AddMaterialExpression<IDatasmithMaterialExpressionScalar>();
-	DiffuseLerpA->GetScalar() = 0.04f;
-
-	IDatasmithMaterialExpressionScalar* DiffuseLerpB = MaterialElement->AddMaterialExpression<IDatasmithMaterialExpressionScalar>();
-	DiffuseLerpB->GetScalar() = 1.0f;
-
-	IDatasmithMaterialExpressionGeneric* BaseColorMultiply = MaterialElement->AddMaterialExpression<IDatasmithMaterialExpressionGeneric>();
-	BaseColorMultiply->SetExpressionName(TEXT("Multiply"));
-
-	IDatasmithMaterialExpressionGeneric* BaseColorAdd = MaterialElement->AddMaterialExpression<IDatasmithMaterialExpressionGeneric>();
-	BaseColorAdd->SetExpressionName(TEXT("Add"));
-
-	IDatasmithMaterialExpressionGeneric* BaseColorTransparencyMultiply = MaterialElement->AddMaterialExpression<IDatasmithMaterialExpressionGeneric>();
-	BaseColorTransparencyMultiply->SetExpressionName(TEXT("Multiply"));
-
-	IDatasmithMaterialExpressionGeneric* IncandescenceMultiply = MaterialElement->AddMaterialExpression<IDatasmithMaterialExpressionGeneric>();
-	IncandescenceMultiply->SetExpressionName(TEXT("Multiply"));
-
-	IDatasmithMaterialExpressionGeneric* IncandescenceScaleMultiply = MaterialElement->AddMaterialExpression<IDatasmithMaterialExpressionGeneric>();
-	IncandescenceScaleMultiply->SetExpressionName(TEXT("Multiply"));
-
-	IDatasmithMaterialExpressionScalar* IncandescenceScale = MaterialElement->AddMaterialExpression<IDatasmithMaterialExpressionScalar>();
-	IncandescenceScale->GetScalar() = 100.0f;
-
-	IDatasmithMaterialExpressionGeneric* EccentricityMultiply = MaterialElement->AddMaterialExpression<IDatasmithMaterialExpressionGeneric>();
-	EccentricityMultiply->SetExpressionName(TEXT("Multiply"));
-
-	IDatasmithMaterialExpressionGeneric* EccentricityOneMinus = MaterialElement->AddMaterialExpression<IDatasmithMaterialExpressionGeneric>();
-	EccentricityOneMinus->SetExpressionName(TEXT("OneMinus"));
-
-	IDatasmithMaterialExpressionGeneric* RoughnessOneMinus = MaterialElement->AddMaterialExpression<IDatasmithMaterialExpressionGeneric>();
-	RoughnessOneMinus->SetExpressionName(TEXT("OneMinus"));
-
-	IDatasmithMaterialExpressionScalar* FresnelExponent = MaterialElement->AddMaterialExpression<IDatasmithMaterialExpressionScalar>();
-	FresnelExponent->GetScalar() = 4.0f;
-
-	IDatasmithMaterialExpressionFunctionCall* FresnelFunc = MaterialElement->AddMaterialExpression<IDatasmithMaterialExpressionFunctionCall>();
-	FresnelFunc->SetFunctionPathName(TEXT("/Engine/Functions/Engine_MaterialFunctions02/Fresnel_Function.Fresnel_Function"));
-
-	IDatasmithMaterialExpressionGeneric* FresnelLerp = MaterialElement->AddMaterialExpression<IDatasmithMaterialExpressionGeneric>();
-	FresnelLerp->SetExpressionName(TEXT("LinearInterpolate"));
-
-	IDatasmithMaterialExpressionScalar* FresnelLerpA = MaterialElement->AddMaterialExpression<IDatasmithMaterialExpressionScalar>();
-	FresnelLerpA->GetScalar() = 1.0f;
-
-	IDatasmithMaterialExpressionScalar* SpecularPowerExp = MaterialElement->AddMaterialExpression<IDatasmithMaterialExpressionScalar>();
-	SpecularPowerExp->GetScalar() = 0.5f;
-
-	IDatasmithMaterialExpressionGeneric* Power = MaterialElement->AddMaterialExpression<IDatasmithMaterialExpressionGeneric>();
-	Power->SetExpressionName(TEXT("Power"));
-
-	IDatasmithMaterialExpressionGeneric* FresnelMultiply = MaterialElement->AddMaterialExpression<IDatasmithMaterialExpressionGeneric>();
-	FresnelMultiply->SetExpressionName(TEXT("Multiply"));
-
-	IDatasmithMaterialExpressionGeneric* TransparencyOneMinus = MaterialElement->AddMaterialExpression<IDatasmithMaterialExpressionGeneric>();
-	TransparencyOneMinus->SetExpressionName(TEXT("OneMinus"));
-
-	IDatasmithMaterialExpressionFunctionCall* BreakFloat3 = nullptr;
-	IDatasmithMaterialExpressionGeneric* AddRG = nullptr;
-	IDatasmithMaterialExpressionGeneric* AddRGB = nullptr;
-	IDatasmithMaterialExpressionGeneric* Divide = nullptr;
-	IDatasmithMaterialExpressionScalar* DivideConstant = nullptr;
-	if (bIsTransparent)
-	{
-		BreakFloat3 = MaterialElement->AddMaterialExpression<IDatasmithMaterialExpressionFunctionCall>();
-		BreakFloat3->SetFunctionPathName(TEXT("/Engine/Functions/Engine_MaterialFunctions02/Utility/BreakFloat3Components.BreakFloat3Components"));
-
-		AddRG = MaterialElement->AddMaterialExpression<IDatasmithMaterialExpressionGeneric>();
-		AddRG->SetExpressionName(TEXT("Add"));
-
-		AddRGB = MaterialElement->AddMaterialExpression<IDatasmithMaterialExpressionGeneric>();
-		AddRGB->SetExpressionName(TEXT("Add"));
-
-		Divide = MaterialElement->AddMaterialExpression<IDatasmithMaterialExpressionGeneric>();
-		Divide->SetExpressionName(TEXT("Divide"));
-
-		DivideConstant = MaterialElement->AddMaterialExpression<IDatasmithMaterialExpressionScalar>();
-		DivideConstant->GetScalar() = 3.0f;
-	}
-
-	// Connect expressions
-	SpecularColorExpression->ConnectExpression(*ColorSpecLerp->GetInput(0));
-	ColorExpression->ConnectExpression(*ColorSpecLerp->GetInput(1));
-	ColorSpecLerpValue->ConnectExpression(*ColorSpecLerp->GetInput(2));
-
-	ColorExpression->ConnectExpression(*ColorMetallicLerp->GetInput(0));
-	ColorSpecLerp->ConnectExpression(*ColorMetallicLerp->GetInput(1));
-	GlossExpression->ConnectExpression(*ColorMetallicLerp->GetInput(2));
-
-	DiffuseLerpA->ConnectExpression(*DiffuseLerp->GetInput(0));
-	DiffuseLerpB->ConnectExpression(*DiffuseLerp->GetInput(1));
-	DiffuseExpression->ConnectExpression(*DiffuseLerp->GetInput(2));
-
-	ColorMetallicLerp->ConnectExpression(*BaseColorMultiply->GetInput(0));
-	DiffuseLerp->ConnectExpression(*BaseColorMultiply->GetInput(1));
-
-	BaseColorMultiply->ConnectExpression(*BaseColorAdd->GetInput(0));
-	IncandescenceColorExpression->ConnectExpression(*BaseColorAdd->GetInput(1));
-
-	BaseColorAdd->ConnectExpression(*BaseColorTransparencyMultiply->GetInput(0));
-	TransparencyOneMinus->ConnectExpression(*BaseColorTransparencyMultiply->GetInput(1));
-
-	GlowIntensityExpression->ConnectExpression(*IncandescenceScaleMultiply->GetInput(0));
-	IncandescenceScale->ConnectExpression(*IncandescenceScaleMultiply->GetInput(1));
-
-	BaseColorTransparencyMultiply->ConnectExpression(*IncandescenceMultiply->GetInput(0));
-	IncandescenceScaleMultiply->ConnectExpression(*IncandescenceMultiply->GetInput(1));
-
-	EccentricityExpression->ConnectExpression(*EccentricityOneMinus->GetInput(0));
-
-	EccentricityOneMinus->ConnectExpression(*EccentricityMultiply->GetInput(0));
-	SpecularityExpression->ConnectExpression(*EccentricityMultiply->GetInput(1));
-
-	EccentricityMultiply->ConnectExpression(*RoughnessOneMinus->GetInput(0));
-
-	FresnelExponent->ConnectExpression(*FresnelFunc->GetInput(3));
-
-	SpecularRolloffExpression->ConnectExpression(*Power->GetInput(0));
-	SpecularPowerExp->ConnectExpression(*Power->GetInput(1));
-
-	FresnelLerpA->ConnectExpression(*FresnelLerp->GetInput(0));
-	FresnelFunc->ConnectExpression(*FresnelLerp->GetInput(1));
-	Power->ConnectExpression(*FresnelLerp->GetInput(2));
-
-	FresnelLerp->ConnectExpression(*FresnelMultiply->GetInput(0));
-	ReflectivityExpression->ConnectExpression(*FresnelMultiply->GetInput(1));
-
-	TransparencyColorExpression->ConnectExpression(*TransparencyOneMinus->GetInput(0));
-
-	if (bIsTransparent)
-	{
-		TransparencyOneMinus->ConnectExpression(*BreakFloat3->GetInput(0));
-
-		BreakFloat3->ConnectExpression(*AddRG->GetInput(0), 0);
-		BreakFloat3->ConnectExpression(*AddRG->GetInput(1), 1);
-
-		AddRG->ConnectExpression(*AddRGB->GetInput(0));
-		BreakFloat3->ConnectExpression(*AddRGB->GetInput(1), 2);
-
-		AddRGB->ConnectExpression(*Divide->GetInput(0));
-		DivideConstant->ConnectExpression(*Divide->GetInput(1));
-	}
-
-	// Connect material outputs
-	MaterialElement->GetBaseColor().SetExpression(BaseColorTransparencyMultiply);
-	MaterialElement->GetMetallic().SetExpression(GlossExpression);
-	MaterialElement->GetSpecular().SetExpression(FresnelMultiply);
-	MaterialElement->GetRoughness().SetExpression(RoughnessOneMinus);
-	MaterialElement->GetEmissiveColor().SetExpression(IncandescenceMultiply);
-
-	if (bIsTransparent)
-	{
-		MaterialElement->GetOpacity().SetExpression(Divide);
-		MaterialElement->SetParentLabel(TEXT("M_DatasmithAliasBlinnTransparent"));
-	}
-	else
-	{
-		MaterialElement->SetParentLabel(TEXT("M_DatasmithAliasBlinn"));
-	}
-
-}
-
-void FWireTranslatorImpl::AddAlLambertParameters(const TAlObjectPtr<AlShader>& Shader, TSharedRef<IDatasmithUEPbrMaterialElement> MaterialElement)
-{
-	// Default values for a Lambert material
-	FColor Color(145, 148, 153);
-	FColor TransparencyColor(0, 0, 0);
-	FColor IncandescenceColor(0, 0, 0);
-	double Diffuse = 1.0;
-	double GlowIntensity = 0.0;
-
-	AlList* List = Shader->fields();
-	for (AlShadingFieldItem* Item = static_cast<AlShadingFieldItem*>(List->first()); Item; Item = Item->nextField())
-	{
-		double Value = 0.0f;
-		statusCode ErrorCode = Shader->parameter(Item->field(), Value);
-		if (ErrorCode != 0)
-		{
-			continue;
-		}
-
-		if (GetCommonParameters(Item->field(), Value, Color, TransparencyColor, IncandescenceColor, GlowIntensity))
-		{
-			continue;
-		}
-
-		switch (Item->field())
-		{
-		case AlShadingFields::kFLD_SHADING_LAMBERT_DIFFUSE:
-			Diffuse = Value;
-			break;
-		}
-	}
-
-	bool bIsTransparent = IsTransparent(TransparencyColor);
-
-	// Construct parameter expressions
-	IDatasmithMaterialExpressionScalar* DiffuseExpression = MaterialElement->AddMaterialExpression<IDatasmithMaterialExpressionScalar>();
-	DiffuseExpression->GetScalar() = Diffuse;
-	DiffuseExpression->SetName(TEXT("Diffuse"));
-
-	IDatasmithMaterialExpressionColor* ColorExpression = MaterialElement->AddMaterialExpression<IDatasmithMaterialExpressionColor>();
-	ColorExpression->SetName(TEXT("Color"));
-	ColorExpression->GetColor() = FLinearColor::FromSRGBColor(Color);
-
-	IDatasmithMaterialExpressionColor* IncandescenceColorExpression = MaterialElement->AddMaterialExpression<IDatasmithMaterialExpressionColor>();
-	IncandescenceColorExpression->SetName(TEXT("IncandescenceColor"));
-	IncandescenceColorExpression->GetColor() = FLinearColor::FromSRGBColor(IncandescenceColor);
-
-	IDatasmithMaterialExpressionColor* TransparencyColorExpression = MaterialElement->AddMaterialExpression<IDatasmithMaterialExpressionColor>();
-	TransparencyColorExpression->SetName(TEXT("TransparencyColor"));
-	TransparencyColorExpression->GetColor() = FLinearColor::FromSRGBColor(TransparencyColor);
-
-	IDatasmithMaterialExpressionScalar* GlowIntensityExpression = MaterialElement->AddMaterialExpression<IDatasmithMaterialExpressionScalar>();
-	GlowIntensityExpression->GetScalar() = GlowIntensity;
-	GlowIntensityExpression->SetName(TEXT("GlowIntensity"));
-
-	// Create aux expressions
-	IDatasmithMaterialExpressionGeneric* DiffuseLerp = MaterialElement->AddMaterialExpression<IDatasmithMaterialExpressionGeneric>();
-	DiffuseLerp->SetExpressionName(TEXT("LinearInterpolate"));
-
-	IDatasmithMaterialExpressionScalar* DiffuseLerpA = MaterialElement->AddMaterialExpression<IDatasmithMaterialExpressionScalar>();
-	DiffuseLerpA->GetScalar() = 0.04f;
-
-	IDatasmithMaterialExpressionScalar* DiffuseLerpB = MaterialElement->AddMaterialExpression<IDatasmithMaterialExpressionScalar>();
-	DiffuseLerpB->GetScalar() = 1.0f;
-
-	IDatasmithMaterialExpressionGeneric* BaseColorMultiply = MaterialElement->AddMaterialExpression<IDatasmithMaterialExpressionGeneric>();
-	BaseColorMultiply->SetExpressionName(TEXT("Multiply"));
-
-	IDatasmithMaterialExpressionGeneric* BaseColorAdd = MaterialElement->AddMaterialExpression<IDatasmithMaterialExpressionGeneric>();
-	BaseColorAdd->SetExpressionName(TEXT("Add"));
-
-	IDatasmithMaterialExpressionGeneric* BaseColorTransparencyMultiply = MaterialElement->AddMaterialExpression<IDatasmithMaterialExpressionGeneric>();
-	BaseColorTransparencyMultiply->SetExpressionName(TEXT("Multiply"));
-
-	IDatasmithMaterialExpressionGeneric* IncandescenceMultiply = MaterialElement->AddMaterialExpression<IDatasmithMaterialExpressionGeneric>();
-	IncandescenceMultiply->SetExpressionName(TEXT("Multiply"));
-
-	IDatasmithMaterialExpressionGeneric* IncandescenceScaleMultiply = MaterialElement->AddMaterialExpression<IDatasmithMaterialExpressionGeneric>();
-	IncandescenceScaleMultiply->SetExpressionName(TEXT("Multiply"));
-
-	IDatasmithMaterialExpressionScalar* IncandescenceScale = MaterialElement->AddMaterialExpression<IDatasmithMaterialExpressionScalar>();
-	IncandescenceScale->GetScalar() = 100.0f;
-
-	IDatasmithMaterialExpressionGeneric* TransparencyOneMinus = MaterialElement->AddMaterialExpression<IDatasmithMaterialExpressionGeneric>();
-	TransparencyOneMinus->SetExpressionName(TEXT("OneMinus"));
-
-	IDatasmithMaterialExpressionFunctionCall* BreakFloat3 = nullptr;
-	IDatasmithMaterialExpressionGeneric* AddRG = nullptr;
-	IDatasmithMaterialExpressionGeneric* AddRGB = nullptr;
-	IDatasmithMaterialExpressionGeneric* Divide = nullptr;
-	IDatasmithMaterialExpressionScalar* DivideConstant = nullptr;
-	if (bIsTransparent)
-	{
-		BreakFloat3 = MaterialElement->AddMaterialExpression<IDatasmithMaterialExpressionFunctionCall>();
-		BreakFloat3->SetFunctionPathName(TEXT("/Engine/Functions/Engine_MaterialFunctions02/Utility/BreakFloat3Components.BreakFloat3Components"));
-
-		AddRG = MaterialElement->AddMaterialExpression<IDatasmithMaterialExpressionGeneric>();
-		AddRG->SetExpressionName(TEXT("Add"));
-
-		AddRGB = MaterialElement->AddMaterialExpression<IDatasmithMaterialExpressionGeneric>();
-		AddRGB->SetExpressionName(TEXT("Add"));
-
-		Divide = MaterialElement->AddMaterialExpression<IDatasmithMaterialExpressionGeneric>();
-		Divide->SetExpressionName(TEXT("Divide"));
-
-		DivideConstant = MaterialElement->AddMaterialExpression<IDatasmithMaterialExpressionScalar>();
-		DivideConstant->GetScalar() = 3.0f;
-	}
-
-	// Connect expressions
-	DiffuseLerpA->ConnectExpression(*DiffuseLerp->GetInput(0));
-	DiffuseLerpB->ConnectExpression(*DiffuseLerp->GetInput(1));
-	DiffuseExpression->ConnectExpression(*DiffuseLerp->GetInput(2));
-
-	ColorExpression->ConnectExpression(*BaseColorMultiply->GetInput(0));
-	DiffuseLerp->ConnectExpression(*BaseColorMultiply->GetInput(1));
-
-	BaseColorMultiply->ConnectExpression(*BaseColorAdd->GetInput(0));
-	IncandescenceColorExpression->ConnectExpression(*BaseColorAdd->GetInput(1));
-
-	BaseColorAdd->ConnectExpression(*BaseColorTransparencyMultiply->GetInput(0));
-	TransparencyOneMinus->ConnectExpression(*BaseColorTransparencyMultiply->GetInput(1));
-
-	GlowIntensityExpression->ConnectExpression(*IncandescenceScaleMultiply->GetInput(0));
-	IncandescenceScale->ConnectExpression(*IncandescenceScaleMultiply->GetInput(1));
-
-	BaseColorTransparencyMultiply->ConnectExpression(*IncandescenceMultiply->GetInput(0));
-	IncandescenceScaleMultiply->ConnectExpression(*IncandescenceMultiply->GetInput(1));
-
-	TransparencyColorExpression->ConnectExpression(*TransparencyOneMinus->GetInput(0));
-
-	if (bIsTransparent)
-	{
-		TransparencyOneMinus->ConnectExpression(*BreakFloat3->GetInput(0));
-
-		BreakFloat3->ConnectExpression(*AddRG->GetInput(0), 0);
-		BreakFloat3->ConnectExpression(*AddRG->GetInput(1), 1);
-
-		AddRG->ConnectExpression(*AddRGB->GetInput(0));
-		BreakFloat3->ConnectExpression(*AddRGB->GetInput(1), 2);
-
-		AddRGB->ConnectExpression(*Divide->GetInput(0));
-		DivideConstant->ConnectExpression(*Divide->GetInput(1));
-	}
-
-	// Connect material outputs
-	MaterialElement->GetBaseColor().SetExpression(BaseColorTransparencyMultiply);
-	MaterialElement->GetEmissiveColor().SetExpression(IncandescenceMultiply);
-	if (bIsTransparent)
-	{
-		MaterialElement->GetOpacity().SetExpression(Divide);
-		MaterialElement->SetParentLabel(TEXT("M_DatasmithAliasLambertTransparent"));
-	}
-	else {
-		MaterialElement->SetParentLabel(TEXT("M_DatasmithAliasLambert"));
-	}
-
-}
-
-void FWireTranslatorImpl::AddAlLightSourceParameters(const TAlObjectPtr<AlShader>& Shader, TSharedRef<IDatasmithUEPbrMaterialElement> MaterialElement)
-{
-	// Default values for a LightSource material
-	FColor Color(145, 148, 153);
-	FColor TransparencyColor(0, 0, 0);
-	FColor IncandescenceColor(0, 0, 0);
-	double GlowIntensity = 0.0;
-
-	AlList* List = Shader->fields();
-	for (AlShadingFieldItem* Item = static_cast<AlShadingFieldItem*>(List->first()); Item; Item = Item->nextField())
-	{
-		double Value = 0.0f;
-		statusCode ErrorCode = Shader->parameter(Item->field(), Value);
-		if (ErrorCode != 0)
-		{
-			continue;
-		}
-
-		GetCommonParameters(Item->field(), Value, Color, TransparencyColor, IncandescenceColor, GlowIntensity);
-	}
-
-	bool bIsTransparent = IsTransparent(TransparencyColor);
-
-	// Construct parameter expressions
-	IDatasmithMaterialExpressionColor* ColorExpression = MaterialElement->AddMaterialExpression<IDatasmithMaterialExpressionColor>();
-	ColorExpression->SetName(TEXT("Color"));
-	ColorExpression->GetColor() = FLinearColor::FromSRGBColor(Color);
-
-	IDatasmithMaterialExpressionColor* IncandescenceColorExpression = MaterialElement->AddMaterialExpression<IDatasmithMaterialExpressionColor>();
-	IncandescenceColorExpression->SetName(TEXT("IncandescenceColor"));
-	IncandescenceColorExpression->GetColor() = FLinearColor::FromSRGBColor(IncandescenceColor);
-
-	IDatasmithMaterialExpressionColor* TransparencyColorExpression = MaterialElement->AddMaterialExpression<IDatasmithMaterialExpressionColor>();
-	TransparencyColorExpression->SetName(TEXT("TransparencyColor"));
-	TransparencyColorExpression->GetColor() = FLinearColor::FromSRGBColor(TransparencyColor);
-
-	IDatasmithMaterialExpressionScalar* GlowIntensityExpression = MaterialElement->AddMaterialExpression<IDatasmithMaterialExpressionScalar>();
-	GlowIntensityExpression->GetScalar() = GlowIntensity;
-	GlowIntensityExpression->SetName(TEXT("GlowIntensity"));
-
-	// Create aux expressions
-	IDatasmithMaterialExpressionGeneric* BaseColorAdd = MaterialElement->AddMaterialExpression<IDatasmithMaterialExpressionGeneric>();
-	BaseColorAdd->SetExpressionName(TEXT("Add"));
-
-	IDatasmithMaterialExpressionGeneric* BaseColorTransparencyMultiply = MaterialElement->AddMaterialExpression<IDatasmithMaterialExpressionGeneric>();
-	BaseColorTransparencyMultiply->SetExpressionName(TEXT("Multiply"));
-
-	IDatasmithMaterialExpressionGeneric* IncandescenceMultiply = MaterialElement->AddMaterialExpression<IDatasmithMaterialExpressionGeneric>();
-	IncandescenceMultiply->SetExpressionName(TEXT("Multiply"));
-
-	IDatasmithMaterialExpressionGeneric* IncandescenceScaleMultiply = MaterialElement->AddMaterialExpression<IDatasmithMaterialExpressionGeneric>();
-	IncandescenceScaleMultiply->SetExpressionName(TEXT("Multiply"));
-
-	IDatasmithMaterialExpressionScalar* IncandescenceScale = MaterialElement->AddMaterialExpression<IDatasmithMaterialExpressionScalar>();
-	IncandescenceScale->GetScalar() = 100.0f;
-
-	IDatasmithMaterialExpressionGeneric* TransparencyOneMinus = MaterialElement->AddMaterialExpression<IDatasmithMaterialExpressionGeneric>();
-	TransparencyOneMinus->SetExpressionName(TEXT("OneMinus"));
-
-	IDatasmithMaterialExpressionFunctionCall* BreakFloat3 = nullptr;
-	IDatasmithMaterialExpressionGeneric* AddRG = nullptr;
-	IDatasmithMaterialExpressionGeneric* AddRGB = nullptr;
-	IDatasmithMaterialExpressionGeneric* Divide = nullptr;
-	IDatasmithMaterialExpressionScalar* DivideConstant = nullptr;
-	if (bIsTransparent)
-	{
-		BreakFloat3 = MaterialElement->AddMaterialExpression<IDatasmithMaterialExpressionFunctionCall>();
-		BreakFloat3->SetFunctionPathName(TEXT("/Engine/Functions/Engine_MaterialFunctions02/Utility/BreakFloat3Components.BreakFloat3Components"));
-
-		AddRG = MaterialElement->AddMaterialExpression<IDatasmithMaterialExpressionGeneric>();
-		AddRG->SetExpressionName(TEXT("Add"));
-
-		AddRGB = MaterialElement->AddMaterialExpression<IDatasmithMaterialExpressionGeneric>();
-		AddRGB->SetExpressionName(TEXT("Add"));
-
-		Divide = MaterialElement->AddMaterialExpression<IDatasmithMaterialExpressionGeneric>();
-		Divide->SetExpressionName(TEXT("Divide"));
-
-		DivideConstant = MaterialElement->AddMaterialExpression<IDatasmithMaterialExpressionScalar>();
-		DivideConstant->GetScalar() = 3.0f;
-	}
-
-	// Connect expressions
-	ColorExpression->ConnectExpression(*BaseColorAdd->GetInput(0));
-	IncandescenceColorExpression->ConnectExpression(*BaseColorAdd->GetInput(1));
-
-	BaseColorAdd->ConnectExpression(*BaseColorTransparencyMultiply->GetInput(0));
-	TransparencyOneMinus->ConnectExpression(*BaseColorTransparencyMultiply->GetInput(1));
-
-	GlowIntensityExpression->ConnectExpression(*IncandescenceScaleMultiply->GetInput(0));
-	IncandescenceScale->ConnectExpression(*IncandescenceScaleMultiply->GetInput(1));
-
-	BaseColorTransparencyMultiply->ConnectExpression(*IncandescenceMultiply->GetInput(0));
-	IncandescenceScaleMultiply->ConnectExpression(*IncandescenceMultiply->GetInput(1));
-
-	TransparencyColorExpression->ConnectExpression(*TransparencyOneMinus->GetInput(0));
-
-	if (bIsTransparent)
-	{
-		TransparencyOneMinus->ConnectExpression(*BreakFloat3->GetInput(0));
-
-		BreakFloat3->ConnectExpression(*AddRG->GetInput(0), 0);
-		BreakFloat3->ConnectExpression(*AddRG->GetInput(1), 1);
-
-		AddRG->ConnectExpression(*AddRGB->GetInput(0));
-		BreakFloat3->ConnectExpression(*AddRGB->GetInput(1), 2);
-
-		AddRGB->ConnectExpression(*Divide->GetInput(0));
-		DivideConstant->ConnectExpression(*Divide->GetInput(1));
-	}
-
-	// Connect material outputs
-	MaterialElement->GetBaseColor().SetExpression(BaseColorTransparencyMultiply);
-	MaterialElement->GetEmissiveColor().SetExpression(IncandescenceMultiply);
-
-	if (bIsTransparent)
-	{
-		MaterialElement->GetOpacity().SetExpression(Divide);
-		MaterialElement->SetParentLabel(TEXT("M_DatasmithAliasLightSourceTransparent"));
-	}
-	else {
-		MaterialElement->SetParentLabel(TEXT("M_DatasmithAliasLightSource"));
-	}
-
-}
-
-void FWireTranslatorImpl::AddAlPhongParameters(const TAlObjectPtr<AlShader>& Shader, TSharedRef<IDatasmithUEPbrMaterialElement> MaterialElement)
-{
-	// Default values for a Phong material
-	FColor Color(145, 148, 153);
-	FColor TransparencyColor(0, 0, 0);
-	FColor IncandescenceColor(0, 0, 0);
-	FColor SpecularColor(38, 38, 38);
-	double Diffuse = 1.0;
-	double GlowIntensity = 0.0;
-	double Gloss = 0.8;
-	double Shinyness = 20.0;
-	double Specularity = 1.0;
-	double Reflectivity = 0.5;
-
-	AlList* List = Shader->fields();
-	for (AlShadingFieldItem* Item = static_cast<AlShadingFieldItem*>(List->first()); Item; Item = Item->nextField())
-	{
-		double Value = 0.0f;
-		statusCode ErrorCode = Shader->parameter(Item->field(), Value);
-		if (ErrorCode != 0)
-		{
-			continue;
-		}
-
-		if (GetCommonParameters(Item->field(), Value, Color, TransparencyColor, IncandescenceColor, GlowIntensity))
-		{
-			continue;
-		}
-
-		switch (Item->field())
-		{
-		case AlShadingFields::kFLD_SHADING_PHONG_DIFFUSE:
-			Diffuse = Value;
-			break;
-		case AlShadingFields::kFLD_SHADING_PHONG_GLOSS_:
-			Gloss = Value;
-			break;
-		case AlShadingFields::kFLD_SHADING_PHONG_SPECULAR_R:
-			SpecularColor.R = (uint8)(255.f * Value);;
-			break;
-		case AlShadingFields::kFLD_SHADING_PHONG_SPECULAR_G:
-			SpecularColor.G = (uint8)(255.f * Value);;
-			break;
-		case AlShadingFields::kFLD_SHADING_PHONG_SPECULAR_B:
-			SpecularColor.B = (uint8)(255.f * Value);;
-			break;
-		case AlShadingFields::kFLD_SHADING_PHONG_SPECULARITY_:
-			Specularity = Value;
-			break;
-		case AlShadingFields::kFLD_SHADING_PHONG_SHINYNESS:
-			Shinyness = Value;
-			break;
-		case AlShadingFields::kFLD_SHADING_PHONG_REFLECTIVITY:
-			Reflectivity = Value;
-			break;
-		}
-	}
-
-	bool bIsTransparent = IsTransparent(TransparencyColor);
-
-	// Construct parameter expressions
-	IDatasmithMaterialExpressionScalar* DiffuseExpression = MaterialElement->AddMaterialExpression<IDatasmithMaterialExpressionScalar>();
-	DiffuseExpression->GetScalar() = Diffuse;
-	DiffuseExpression->SetName(TEXT("Diffuse"));
-
-	IDatasmithMaterialExpressionScalar* GlossExpression = MaterialElement->AddMaterialExpression<IDatasmithMaterialExpressionScalar>();
-	GlossExpression->GetScalar() = Gloss;
-	GlossExpression->SetName(TEXT("Gloss"));
-
-	IDatasmithMaterialExpressionColor* SpecularColorExpression = MaterialElement->AddMaterialExpression<IDatasmithMaterialExpressionColor>();
-	SpecularColorExpression->SetName(TEXT("SpecularColor"));
-	SpecularColorExpression->GetColor() = FLinearColor::FromSRGBColor(SpecularColor);
-
-	IDatasmithMaterialExpressionScalar* SpecularityExpression = MaterialElement->AddMaterialExpression<IDatasmithMaterialExpressionScalar>();
-	SpecularityExpression->GetScalar() = Specularity * 0.3;
-	SpecularityExpression->SetName(TEXT("Specularity"));
-
-	IDatasmithMaterialExpressionScalar* ShinynessExpression = MaterialElement->AddMaterialExpression<IDatasmithMaterialExpressionScalar>();
-	ShinynessExpression->GetScalar() = Shinyness;
-	ShinynessExpression->SetName(TEXT("Shinyness"));
-
-	IDatasmithMaterialExpressionScalar* ReflectivityExpression = MaterialElement->AddMaterialExpression<IDatasmithMaterialExpressionScalar>();
-	ReflectivityExpression->GetScalar() = Reflectivity;
-	ReflectivityExpression->SetName(TEXT("Reflectivity"));
-
-	IDatasmithMaterialExpressionColor* ColorExpression = MaterialElement->AddMaterialExpression<IDatasmithMaterialExpressionColor>();
-	ColorExpression->SetName(TEXT("Color"));
-	ColorExpression->GetColor() = FLinearColor::FromSRGBColor(Color);
-
-	IDatasmithMaterialExpressionColor* IncandescenceColorExpression = MaterialElement->AddMaterialExpression<IDatasmithMaterialExpressionColor>();
-	IncandescenceColorExpression->SetName(TEXT("IncandescenceColor"));
-	IncandescenceColorExpression->GetColor() = FLinearColor::FromSRGBColor(IncandescenceColor);
-
-	IDatasmithMaterialExpressionColor* TransparencyColorExpression = MaterialElement->AddMaterialExpression<IDatasmithMaterialExpressionColor>();
-	TransparencyColorExpression->SetName(TEXT("TransparencyColor"));
-	TransparencyColorExpression->GetColor() = FLinearColor::FromSRGBColor(TransparencyColor);
-
-	IDatasmithMaterialExpressionScalar* GlowIntensityExpression = MaterialElement->AddMaterialExpression<IDatasmithMaterialExpressionScalar>();
-	GlowIntensityExpression->GetScalar() = GlowIntensity;
-	GlowIntensityExpression->SetName(TEXT("GlowIntensity"));
-
-	// Create aux expressions
-	IDatasmithMaterialExpressionGeneric* ColorSpecLerp = MaterialElement->AddMaterialExpression<IDatasmithMaterialExpressionGeneric>();
-	ColorSpecLerp->SetExpressionName(TEXT("LinearInterpolate"));
-
-	IDatasmithMaterialExpressionScalar* ColorSpecLerpValue = MaterialElement->AddMaterialExpression<IDatasmithMaterialExpressionScalar>();
-	ColorSpecLerpValue->GetScalar() = 0.96f;
-
-	IDatasmithMaterialExpressionGeneric* ColorMetallicLerp = MaterialElement->AddMaterialExpression<IDatasmithMaterialExpressionGeneric>();
-	ColorMetallicLerp->SetExpressionName(TEXT("LinearInterpolate"));
-
-	IDatasmithMaterialExpressionGeneric* DiffuseLerp = MaterialElement->AddMaterialExpression<IDatasmithMaterialExpressionGeneric>();
-	DiffuseLerp->SetExpressionName(TEXT("LinearInterpolate"));
-
-	IDatasmithMaterialExpressionScalar* DiffuseLerpA = MaterialElement->AddMaterialExpression<IDatasmithMaterialExpressionScalar>();
-	DiffuseLerpA->GetScalar() = 0.04f;
-
-	IDatasmithMaterialExpressionScalar* DiffuseLerpB = MaterialElement->AddMaterialExpression<IDatasmithMaterialExpressionScalar>();
-	DiffuseLerpB->GetScalar() = 1.0f;
-
-	IDatasmithMaterialExpressionGeneric* BaseColorMultiply = MaterialElement->AddMaterialExpression<IDatasmithMaterialExpressionGeneric>();
-	BaseColorMultiply->SetExpressionName(TEXT("Multiply"));
-
-	IDatasmithMaterialExpressionGeneric* BaseColorAdd = MaterialElement->AddMaterialExpression<IDatasmithMaterialExpressionGeneric>();
-	BaseColorAdd->SetExpressionName(TEXT("Add"));
-
-	IDatasmithMaterialExpressionGeneric* BaseColorTransparencyMultiply = MaterialElement->AddMaterialExpression<IDatasmithMaterialExpressionGeneric>();
-	BaseColorTransparencyMultiply->SetExpressionName(TEXT("Multiply"));
-
-	IDatasmithMaterialExpressionGeneric* IncandescenceMultiply = MaterialElement->AddMaterialExpression<IDatasmithMaterialExpressionGeneric>();
-	IncandescenceMultiply->SetExpressionName(TEXT("Multiply"));
-
-	IDatasmithMaterialExpressionGeneric* IncandescenceScaleMultiply = MaterialElement->AddMaterialExpression<IDatasmithMaterialExpressionGeneric>();
-	IncandescenceScaleMultiply->SetExpressionName(TEXT("Multiply"));
-
-	IDatasmithMaterialExpressionScalar* IncandescenceScale = MaterialElement->AddMaterialExpression<IDatasmithMaterialExpressionScalar>();
-	IncandescenceScale->GetScalar() = 100.0f;
-
-	IDatasmithMaterialExpressionGeneric* ShinynessSubtract = MaterialElement->AddMaterialExpression<IDatasmithMaterialExpressionGeneric>();
-	ShinynessSubtract->SetExpressionName(TEXT("Subtract"));
-
-	IDatasmithMaterialExpressionScalar* ShinynessSubtract2 = MaterialElement->AddMaterialExpression<IDatasmithMaterialExpressionScalar>();
-	ShinynessSubtract2->GetScalar() = 2.0f;
-
-	IDatasmithMaterialExpressionGeneric* ShinynessDivide = MaterialElement->AddMaterialExpression<IDatasmithMaterialExpressionGeneric>();
-	ShinynessDivide->SetExpressionName(TEXT("Divide"));
-
-	IDatasmithMaterialExpressionScalar* ShinynessDivide98 = MaterialElement->AddMaterialExpression<IDatasmithMaterialExpressionScalar>();
-	ShinynessDivide98->GetScalar() = 98.0f;
-
-	IDatasmithMaterialExpressionGeneric* SpecularityMultiply = MaterialElement->AddMaterialExpression<IDatasmithMaterialExpressionGeneric>();
-	SpecularityMultiply->SetExpressionName(TEXT("Multiply"));
-
-	IDatasmithMaterialExpressionGeneric* RoughnessOneMinus = MaterialElement->AddMaterialExpression<IDatasmithMaterialExpressionGeneric>();
-	RoughnessOneMinus->SetExpressionName(TEXT("OneMinus"));
-
-	IDatasmithMaterialExpressionGeneric* TransparencyOneMinus = MaterialElement->AddMaterialExpression<IDatasmithMaterialExpressionGeneric>();
-	TransparencyOneMinus->SetExpressionName(TEXT("OneMinus"));
-
-	IDatasmithMaterialExpressionFunctionCall* BreakFloat3 = nullptr;
-	IDatasmithMaterialExpressionGeneric* AddRG = nullptr;
-	IDatasmithMaterialExpressionGeneric* AddRGB = nullptr;
-	IDatasmithMaterialExpressionGeneric* Divide = nullptr;
-	IDatasmithMaterialExpressionScalar* DivideConstant = nullptr;
-	if (bIsTransparent)
-	{
-		BreakFloat3 = MaterialElement->AddMaterialExpression<IDatasmithMaterialExpressionFunctionCall>();
-		BreakFloat3->SetFunctionPathName(TEXT("/Engine/Functions/Engine_MaterialFunctions02/Utility/BreakFloat3Components.BreakFloat3Components"));
-
-		AddRG = MaterialElement->AddMaterialExpression<IDatasmithMaterialExpressionGeneric>();
-		AddRG->SetExpressionName(TEXT("Add"));
-
-		AddRGB = MaterialElement->AddMaterialExpression<IDatasmithMaterialExpressionGeneric>();
-		AddRGB->SetExpressionName(TEXT("Add"));
-
-		Divide = MaterialElement->AddMaterialExpression<IDatasmithMaterialExpressionGeneric>();
-		Divide->SetExpressionName(TEXT("Divide"));
-
-		DivideConstant = MaterialElement->AddMaterialExpression<IDatasmithMaterialExpressionScalar>();
-		DivideConstant->GetScalar() = 3.0f;
-	}
-
-	// Connect expressions
-	SpecularColorExpression->ConnectExpression(*ColorSpecLerp->GetInput(0));
-	ColorExpression->ConnectExpression(*ColorSpecLerp->GetInput(1));
-	ColorSpecLerpValue->ConnectExpression(*ColorSpecLerp->GetInput(2));
-
-	ColorExpression->ConnectExpression(*ColorMetallicLerp->GetInput(0));
-	ColorSpecLerp->ConnectExpression(*ColorMetallicLerp->GetInput(1));
-	GlossExpression->ConnectExpression(*ColorMetallicLerp->GetInput(2));
-
-	DiffuseLerpA->ConnectExpression(*DiffuseLerp->GetInput(0));
-	DiffuseLerpB->ConnectExpression(*DiffuseLerp->GetInput(1));
-	DiffuseExpression->ConnectExpression(*DiffuseLerp->GetInput(2));
-
-	ColorMetallicLerp->ConnectExpression(*BaseColorMultiply->GetInput(0));
-	DiffuseLerp->ConnectExpression(*BaseColorMultiply->GetInput(1));
-
-	BaseColorMultiply->ConnectExpression(*BaseColorAdd->GetInput(0));
-	IncandescenceColorExpression->ConnectExpression(*BaseColorAdd->GetInput(1));
-
-	BaseColorAdd->ConnectExpression(*BaseColorTransparencyMultiply->GetInput(0));
-	TransparencyOneMinus->ConnectExpression(*BaseColorTransparencyMultiply->GetInput(1));
-
-	GlowIntensityExpression->ConnectExpression(*IncandescenceScaleMultiply->GetInput(0));
-	IncandescenceScale->ConnectExpression(*IncandescenceScaleMultiply->GetInput(1));
-
-	BaseColorTransparencyMultiply->ConnectExpression(*IncandescenceMultiply->GetInput(0));
-	IncandescenceScaleMultiply->ConnectExpression(*IncandescenceMultiply->GetInput(1));
-
-	ShinynessExpression->ConnectExpression(*ShinynessSubtract->GetInput(0));
-	ShinynessSubtract2->ConnectExpression(*ShinynessSubtract->GetInput(1));
-
-	ShinynessSubtract->ConnectExpression(*ShinynessDivide->GetInput(0));
-	ShinynessDivide98->ConnectExpression(*ShinynessDivide->GetInput(1));
-
-	ShinynessDivide->ConnectExpression(*SpecularityMultiply->GetInput(0));
-	SpecularityExpression->ConnectExpression(*SpecularityMultiply->GetInput(1));
-
-	SpecularityMultiply->ConnectExpression(*RoughnessOneMinus->GetInput(0));
-
-	TransparencyColorExpression->ConnectExpression(*TransparencyOneMinus->GetInput(0));
-
-	if (bIsTransparent)
-	{
-		TransparencyOneMinus->ConnectExpression(*BreakFloat3->GetInput(0));
-
-		BreakFloat3->ConnectExpression(*AddRG->GetInput(0), 0);
-		BreakFloat3->ConnectExpression(*AddRG->GetInput(1), 1);
-
-		AddRG->ConnectExpression(*AddRGB->GetInput(0));
-		BreakFloat3->ConnectExpression(*AddRGB->GetInput(1), 2);
-
-		AddRGB->ConnectExpression(*Divide->GetInput(0));
-		DivideConstant->ConnectExpression(*Divide->GetInput(1));
-	}
-
-	// Connect material outputs
-	MaterialElement->GetBaseColor().SetExpression(BaseColorTransparencyMultiply);
-	MaterialElement->GetMetallic().SetExpression(GlossExpression);
-	MaterialElement->GetSpecular().SetExpression(ReflectivityExpression);
-	MaterialElement->GetRoughness().SetExpression(RoughnessOneMinus);
-	MaterialElement->GetEmissiveColor().SetExpression(IncandescenceMultiply);
-	if (bIsTransparent)
-	{
-		MaterialElement->GetOpacity().SetExpression(Divide);
-		MaterialElement->SetParentLabel(TEXT("M_DatasmithAliasPhongTransparent"));
-	}
-	else {
-		MaterialElement->SetParentLabel(TEXT("M_DatasmithAliasPhong"));
-	}
-}
-
-// Make material
-bool FWireTranslatorImpl::GetShader()
-{
-	for (TAlObjectPtr<AlShader> Shader(AlUniverse::firstShader()); Shader.IsValid(); Shader = TAlObjectPtr<AlShader>(AlUniverse::nextShader(Shader.Get())))
-	{
-		const FString ShaderName = UTF8_TO_TCHAR(Shader->name());
-		const FString ShaderModelName = Shader->shadingModel();
-
-		const FColor Color = CreateShaderColorFromShaderName(ShaderName);
-		ShaderNameToColor.Add(ShaderName, Color);
-
-		const int32 ShaderId = CreateShaderId(Color);
-		const FString MaterialElementName = FString::FromInt(ShaderId);
-
-		TSharedRef<IDatasmithUEPbrMaterialElement> MaterialElement = FDatasmithSceneFactory::CreateUEPbrMaterial(*MaterialElementName);
-		MaterialElement->SetLabel(*ShaderName);
-
-		if (ShaderModelName.Equals(TEXT("BLINN")))
-		{
-			AddAlBlinnParameters(Shader, MaterialElement);
-		}
-		else if (ShaderModelName.Equals(TEXT("LAMBERT")))
-		{
-			AddAlLambertParameters(Shader, MaterialElement);
-		}
-		else if (ShaderModelName.Equals(TEXT("LIGHTSOURCE")))
-		{
-			AddAlLightSourceParameters(Shader, MaterialElement);
-		}
-		else if (ShaderModelName.Equals(TEXT("PHONG")))
-		{
-			AddAlPhongParameters(Shader, MaterialElement);
-		}
-
-		DatasmithScene->AddMaterial(MaterialElement);
-
-		TSharedPtr<IDatasmithMaterialIDElement> MaterialIDElement = FDatasmithSceneFactory::CreateMaterialId(MaterialElement->GetName());
-		ShaderNameToUEMaterialId.Add(*ShaderName, MaterialIDElement);
-	}
-	return true;
-}
-
-bool FWireTranslatorImpl::GetDagLeaves()
-{
-	FDagNodeInfo RootContainer;
-	AlRootNode = TSharedPtr<AlDagNode>(AlUniverse::firstDagNode());
-	if (AlRootNode.IsValid())
-	{
-		TArray<TSharedPtr<IDatasmithActorElement>> ActorsToRemove;
-		if (RecurseDagForLeaves(AlRootNode, RootContainer))
-		{
-			// Delete empty branches
-			for (int32 Index = DatasmithScene->GetActorsCount() - 1; Index >= 0; --Index)
-			{
-				TSharedPtr<IDatasmithActorElement> Actor = DatasmithScene->GetActor(Index);
-				if (Actor.IsValid() && !Actor->IsA(EDatasmithElementType::StaticMeshActor))
-				{
-					RecurseDeleteEmptyActor(Actor);
-					if (Actor->GetChildrenCount() == 0)
-					{
-						ActorsToRemove.Add(Actor);
-					}
-				}
-			}
-			for (const TSharedPtr<IDatasmithActorElement>& Actor : ActorsToRemove)
-			{
-				DatasmithScene->RemoveActor(Actor, EDatasmithActorRemovalRule::RemoveChildren);
-			}
-			return true;
-		}
-	}
-
-	return false;
-}
-
-void FWireTranslatorImpl::RecurseDeleteEmptyActor(TSharedPtr<IDatasmithActorElement> Actor)
-{
-	for (int32 Index = Actor->GetChildrenCount() - 1; Index >= 0; --Index)
-	{
-		TSharedPtr<IDatasmithActorElement> ActorChild = Actor->GetChild(Index);
-		if (ActorChild.IsValid() && !ActorChild->IsA(EDatasmithElementType::StaticMeshActor))
-		{
-			RecurseDeleteEmptyActor(ActorChild);
-			if (ActorChild->GetChildrenCount() == 0)
-			{
-				Actor->RemoveChild(ActorChild);
-			}
-		}
-	}
-}
-
-
-void FWireTranslatorImpl::FDagNodeInfo::ExtractDagNodeLayer(const AlDagNode& InDagNode)
-{
-	AlLayer* LayerPtr = InDagNode.layer();
-	if (LayerPtr)
-	{
-		TAlObjectPtr<AlLayer> Layer(LayerPtr);
-		LayerName = UTF8_TO_TCHAR(Layer->name());
-		bLayerIsVisible = !Layer->invisible();
-		if (ActorElement.IsValid())
-		{
-			ActorElement->SetLayer(*LayerName);
-			LayerToActorElement.Add(LayerName, ActorElement);
-		}
-	}
-}
-
-void FWireTranslatorImpl::FDagNodeInfo::ExtractDagNodeInfo(AlDagNode& CurrentNode)
-{
-	Label = UTF8_TO_TCHAR(CurrentNode.name());
-	uint32 ThisGroupNodeUuid = OpenModelUtils::GetAlDagNodeUuid(CurrentNode);
-	Uuid = HashCombine(ParentInfo ? ParentInfo->Uuid : 1, ThisGroupNodeUuid);
-}
-
-void FWireTranslatorImpl::FDagNodeInfo::ExtractDagNodeInfo(BodyData& CurrentNode)
-{
-	if (!ParentInfo)
-	{
-		return;
-	}
-	Label = ParentInfo->Label;
-	CurrentNode.SetLabel(ParentInfo->Label);
-	Uuid = CurrentNode.GetUuid(ParentInfo->Uuid);
-}
-
-TSharedPtr<IDatasmithActorElement> FWireTranslatorImpl::FindOrAddLayerActor(const FString& LayerName)
-{
-	if (!bGLayersAsActors || LayerName.IsEmpty())
-	{
-		return TSharedPtr<IDatasmithActorElement>();
-	}
-
-	TSharedPtr<IDatasmithActorElement>* LayerActor = LayerNameToActor.Find(LayerName);
-	if (LayerActor)
-	{
-		return *LayerActor;
-	}
-
-	uint32 LayerUuid = GetTypeHash(LayerName);
-	TSharedRef<IDatasmithActorElement> NewLayerActor = FDatasmithSceneFactory::CreateActor(*OpenModelUtils::UuidToString(LayerUuid));
-	NewLayerActor->SetLabel(*LayerName);
-	NewLayerActor->SetLayer(*LayerName);
-	DatasmithScene->AddActor(NewLayerActor);
-	LayerNameToActor.Add(LayerName, NewLayerActor);
-	return NewLayerActor;
-}
-
-TSharedPtr<IDatasmithActorElement> FWireTranslatorImpl::FindOrAddParentActor(FDagNodeInfo& ParentInfo, const FString& LayerName)
-{
-	if (!ParentInfo.ActorElement.IsValid())
-	{
-		if (!LayerName.IsEmpty())
-		{
-			return FindOrAddLayerActor(LayerName);
-		}
-		return TSharedPtr<IDatasmithActorElement>();
-	}
-
-	if (!bGLayersAsActors)
-	{
-		return ParentInfo.ActorElement;
-	}
-
-	TSharedPtr<IDatasmithActorElement>* ActorParent = ParentInfo.LayerToActorElement.Find(LayerName);
-	if (ActorParent)
-	{
-		return *ActorParent;
-	}
-
-	uint32 LayerActorUuid = HashCombine(ParentInfo.Uuid, GetTypeHash(LayerName));
-
-	TSharedPtr<IDatasmithActorElement> LayerActorParent = FDatasmithSceneFactory::CreateActor(*OpenModelUtils::UuidToString(LayerActorUuid));
-	LayerActorParent->SetLabel(*ParentInfo.Label);
-	LayerActorParent->SetLayer(*LayerName);
-
-	ParentInfo.LayerToActorElement.Add(LayerName, LayerActorParent);
-
-	TSharedPtr<IDatasmithActorElement> GrandParentActor = FindOrAddParentActor(*ParentInfo.ParentInfo, LayerName);
-	if (GrandParentActor.IsValid())
-	{
-		GrandParentActor->AddChild(LayerActorParent);
-	}
-	else
-	{
-		TSharedPtr<IDatasmithActorElement> LayerActor = FindOrAddLayerActor(LayerName);
-		LayerActor->AddChild(LayerActorParent);
-	}
-
-	return LayerActorParent;
-}
-
-bool FWireTranslatorImpl::ProcessAlGroupNode(AlDagNode& GroupNode, FDagNodeInfo& ParentInfo)
-{
-	AlGroupNode* AlGroup = GroupNode.asGroupNodePtr();
-	if (!AlIsValid(AlGroup))
-	{
-		return false;
-	}
-
-	AlDagNode* AlChildPtr = AlGroup->childNode();
-	if (!AlIsValid(AlChildPtr))
-	{
-		return false;
-	}
-
-	TSharedPtr<AlDagNode> ChildNode(AlChildPtr);
-
-	FDagNodeInfo ThisGroupNodeInfo(ParentInfo);
-	ThisGroupNodeInfo.ExtractDagNodeInfo(GroupNode);
-
-	ThisGroupNodeInfo.ActorElement = FDatasmithSceneFactory::CreateActor(*OpenModelUtils::UuidToString(ThisGroupNodeInfo.Uuid));
-	ThisGroupNodeInfo.ActorElement->SetLabel(*ThisGroupNodeInfo.Label);
-	ThisGroupNodeInfo.ExtractDagNodeLayer(GroupNode);
-
-	TSharedPtr<IDatasmithActorElement> ParentActorElement = FindOrAddParentActor(ParentInfo, ThisGroupNodeInfo.LayerName);
-	if (ParentActorElement.IsValid())
-	{
-		ParentActorElement->AddChild(ThisGroupNodeInfo.ActorElement);
-	}
-	else
-	{
-		DatasmithScene->AddActor(ThisGroupNodeInfo.ActorElement);
-	}
-
-	RecurseDagForLeaves(ChildNode, ThisGroupNodeInfo);
-
-	// Apply local transform to actor element
-	OpenModelUtils::SetActorTransform(ThisGroupNodeInfo.ActorElement, GroupNode);
-
-	return true;
-}
-
-TSharedPtr<IDatasmithMeshElement> FWireTranslatorImpl::FindOrAddMeshElement(const TSharedPtr<BodyData>& Body, const FDagNodeInfo& NodeInfo)
-{
-	TSharedPtr<IDatasmithMeshElement>* MeshElementPtr = BodyUuidToMeshElementMap.Find(NodeInfo.Uuid);
-	if (MeshElementPtr != nullptr)
-	{
-		return *MeshElementPtr;
-	}
-
-	TSharedPtr<IDatasmithMeshElement> MeshElement = FDatasmithSceneFactory::CreateMesh(*OpenModelUtils::UuidToString(NodeInfo.Uuid));
-	MeshElement->SetLabel(*NodeInfo.Label);
-	MeshElement->SetLightmapSourceUV(-1);
-
-	const TSet<FString>& ShaderNames = Body->GetShaderNames();
-	for (const FString& ShaderName : ShaderNames)
-	{
-		TSharedPtr<IDatasmithMaterialIDElement> MaterialElement = ShaderNameToUEMaterialId[ShaderName];
-		if (MaterialElement.IsValid())
-		{
-			const FColor* ShaderColor = ShaderNameToColor.Find(ShaderName);
-			if (ShaderColor)
-			{
-				int32 ShaderId = CreateShaderId(*ShaderColor);
-				MeshElement->SetMaterial(MaterialElement->GetName(), ShaderId);
-			}
-		}
-	}
-
-	DatasmithScene->AddMesh(MeshElement);
-
-	ShellUuidToMeshElementMap.Add(NodeInfo.Uuid, MeshElement);
-	MeshElementToBodyMap.Add(MeshElement.Get(), Body);
-
-	BodyUuidToMeshElementMap.Add(NodeInfo.Uuid, MeshElement);
-
-	return MeshElement;
-}
-
-TSharedPtr<IDatasmithMeshElement> FWireTranslatorImpl::FindOrAddMeshElement(const TSharedPtr<AlDagNode>& ShellNode, const FDagNodeInfo& ShellNodeInfo, const FString& ShaderName)
-{
-	uint32 ShellUuid = OpenModelUtils::GetAlDagNodeUuid(*ShellNode);
-
-	// Look if geometry has not been already processed, return it if found
-	TSharedPtr<IDatasmithMeshElement>* MeshElementPtr = ShellUuidToMeshElementMap.Find(ShellUuid);
-	if (MeshElementPtr != nullptr)
-	{
-		return *MeshElementPtr;
-	}
-
-	TSharedPtr<IDatasmithMeshElement> MeshElement = FDatasmithSceneFactory::CreateMesh(*OpenModelUtils::UuidToString(ShellNodeInfo.Uuid));
-	MeshElement->SetLabel(*ShellNodeInfo.Label);
-	MeshElement->SetLightmapSourceUV(-1);
-
-	// Set MeshElement FileHash used for re-import task
-	FMD5 MD5; // unique Value that define the mesh
-	MD5.Update(reinterpret_cast<const uint8*>(&SceneFileHash), sizeof SceneFileHash);
-	// MeshActor Name
-	MD5.Update(reinterpret_cast<const uint8*>(&ShellUuid), sizeof ShellUuid);
-	FMD5Hash Hash;
-	Hash.Set(MD5);
-	MeshElement->SetFileHash(Hash);
-
-	if (ShaderName.Len())
-	{
-		TSharedPtr<IDatasmithMaterialIDElement> MaterialElement = ShaderNameToUEMaterialId[ShaderName];
-		if (MaterialElement.IsValid())
-		{
-			const FColor* ShaderColor = ShaderNameToColor.Find(ShaderName);
-			if (ShaderColor)
-			{
-				int32 ShaderId = CreateShaderId(*ShaderColor);
-				MeshElement->SetMaterial(MaterialElement->GetName(), ShaderId);
-			}
-		}
-
-	}
-
-	DatasmithScene->AddMesh(MeshElement);
-
-	ShellUuidToMeshElementMap.Add(ShellUuid, MeshElement);
-	MeshElementToAlDagNodeMap.Add(MeshElement.Get(), ShellNode);
-	MeshElementToShaderName.Add(MeshElement.Get(), ShaderName);
-
-	return MeshElement;
-}
-
-void FWireTranslatorImpl::ProcessAlShellNode(const TSharedPtr<AlDagNode>& ShellNode, FDagNodeInfo& ParentInfo, const FString& ShaderName)
-{
-	FDagNodeInfo ShellInfo(ParentInfo);
-	ShellInfo.ExtractDagNodeInfo(*ShellNode);
-
-	TSharedPtr<IDatasmithMeshElement> MeshElement = FindOrAddMeshElement(ShellNode, ShellInfo, ShaderName);
-	if (!MeshElement.IsValid())
-	{
-		return;
-	}
-
-	TSharedPtr<IDatasmithMeshActorElement> ActorElement = FDatasmithSceneFactory::CreateMeshActor(*OpenModelUtils::UuidToString(ShellInfo.Uuid));
-	if (!ActorElement.IsValid())
-	{
-		return;
-	}
-
-	ActorElement->SetLabel(*ShellInfo.Label);
-	ActorElement->SetStaticMeshPathName(MeshElement->GetName());
-	ShellInfo.ActorElement = ActorElement;
-
-	ShellInfo.ExtractDagNodeLayer(*ShellNode);
-	if (!ShellInfo.bLayerIsVisible)
-	{
-		return;
-	}
-
-	// Apply materials on the current part
-	if (ShaderName.Len())
-	{
-		const TSharedPtr<IDatasmithMaterialIDElement>& MaterialIDElement = ShaderNameToUEMaterialId[ShaderName];
-		if (MaterialIDElement.IsValid())
-		{
-			for (int32 Index = 0; Index < MeshElement->GetMaterialSlotCount(); ++Index)
-			{
-				MaterialIDElement->SetId(MeshElement->GetMaterialSlotAt(Index)->GetId());
-				ActorElement->AddMaterialOverride(MaterialIDElement);
-			}
-		}
-	}
-
-	TSharedPtr<IDatasmithActorElement> ParentElement = FindOrAddParentActor(ParentInfo, ShellInfo.LayerName);
-	if (ParentElement.IsValid())
-	{
-		ParentElement->AddChild(ActorElement);
-	}
-	else
-	{
-		DatasmithScene->AddActor(ActorElement);
-	}
-
-	OpenModelUtils::SetActorTransform(ShellInfo.ActorElement, *ShellNode);
-}
-
-void FWireTranslatorImpl::ProcessBodyNode(const TSharedPtr<BodyData>& Body, FDagNodeInfo& ParentInfo)
-{
-
-	if (!Body.IsValid())
-	{
-		return;
-	}
-
-	if (Body->GetShells().Num() == 1)
-	{
-		DagForLeavesNoMerge(Body->GetShells()[0].Key, ParentInfo);
-		return;
-	}
-
-	FDagNodeInfo ShellInfo(ParentInfo);
-	ShellInfo.ExtractDagNodeInfo(*Body);
-
-	TSharedPtr<IDatasmithMeshElement> MeshElement = FindOrAddMeshElement(Body, ShellInfo);
-	if (!MeshElement.IsValid())
-	{
-		return;
-	}
-
-	TSharedPtr<IDatasmithMeshActorElement> ActorElement = FDatasmithSceneFactory::CreateMeshActor(*OpenModelUtils::UuidToString(ShellInfo.Uuid));
-	if (!ActorElement.IsValid())
-	{
-		return;
-	}
-
-	ActorElement->SetLabel(ShellInfo.Label.IsEmpty() ? *Body->GetLayerName() : *ShellInfo.Label);
-	ActorElement->SetStaticMeshPathName(MeshElement->GetName());
-	ShellInfo.ActorElement = ActorElement;
-
-	ActorElement->SetLayer(*Body->GetLayerName());
-
-	// Apply materials on the current part
-	int32 Index = 0;
-	const TSet<FString>& ShaderNames = Body->GetShaderNames();
-	for (const FString& ShaderName : ShaderNames)
-	{
-		TSharedPtr<IDatasmithMaterialIDElement> MaterialIDElement = ShaderNameToUEMaterialId[ShaderName];
-		if (MaterialIDElement.IsValid())
-		{
-			MaterialIDElement->SetId(MeshElement->GetMaterialSlotAt(Index)->GetId());
-			ActorElement->AddMaterialOverride(MaterialIDElement);
-			Index++;
-		}
-	}
-
-	TSharedPtr<IDatasmithActorElement> ParentElement = FindOrAddParentActor(ParentInfo, Body->GetLayerName());
-	if (ParentElement.IsValid())
-	{
-		ParentElement->AddChild(ActorElement);
-	}
-	else
-	{
-		DatasmithScene->AddActor(ActorElement);
-	}
-
-	return;
-}
-
-TSharedPtr<AlDagNode> GetNextNode(const TSharedPtr<AlDagNode>& DagNode)
-{
-	// Grab the next sibling before deleting the node.
-	AlDagNode* SiblingNode = DagNode->nextNode();
-	if (AlIsValid(SiblingNode))
-	{
-		return TSharedPtr<AlDagNode>(SiblingNode);
-	}
-	else
-	{
-		return TSharedPtr<AlDagNode>();
-	}
-}
-
-uint32 GetBodyGroupUuid(const FString& ShaderName, const FString& LayerName, bool bCadData)
-{
-	uint32 Uuid = HashCombine(GetTypeHash(LayerName), GetTypeHash(bCadData));
-	if (bGSewByMaterial)
-	{
-		Uuid = HashCombine(GetTypeHash(ShaderName), Uuid);
-	}
-	return Uuid;
-}
-
-uint32 GetPatchCount(const TAlObjectPtr<AlShell>& Shell)
-{
-	uint32 PatchCount = 0;
-	TAlObjectPtr<AlTrimRegion> TrimRegion(Shell->firstTrimRegion());
-	while (TrimRegion.IsValid())
-	{
-		PatchCount++;
-		TrimRegion = TAlObjectPtr<AlTrimRegion>(TrimRegion->nextRegion());
-	}
-	return PatchCount;
-}
-
-void FWireTranslatorImpl::AddNodeInBodyGroup(TSharedPtr<AlDagNode>& DagNode, const FString& ShaderName, TMap<uint32, TSharedPtr<BodyData>>& ShellToProcess, bool bIsAPatch, uint32 MaxSize)
-{
-	FString LayerName;
-
-	AlLayer* LayerPtr = DagNode->layer();
-	if (AlIsValid(LayerPtr))
-	{
-		const TAlObjectPtr<AlLayer> Layer(LayerPtr);
-		LayerName = UTF8_TO_TCHAR(Layer->name());
-		if (Layer->invisible())
-		{
-			return;
-		}
-	}
-
-	uint32 SetId = GetBodyGroupUuid(ShaderName, LayerName, bIsAPatch);
-
-	TSharedPtr<BodyData> Body;
-	if (TSharedPtr<BodyData>* PBody = ShellToProcess.Find(SetId))
-	{
-		Body = *PBody;
-	}
-	else
-	{
-		TSharedRef<BodyData> BodyRef = MakeShared<BodyData>(LayerName, bIsAPatch);
-		ShellToProcess.Add(SetId, BodyRef);
-		BodyRef->Reserve(MaxSize);
-		Body = BodyRef;
-	}
-
-	const FColor& Color = ShaderNameToColor[ShaderName];
-	Body->Add(DagNode, Color, ShaderName);
-}
-
-bool FWireTranslatorImpl::RecurseDagForLeaves(const TSharedPtr<AlDagNode>& FirstDagNode, FDagNodeInfo& ParentInfo)
-{
-	if (TessellationOptions.StitchingTechnique != EDatasmithCADStitchingTechnique::StitchingSew)
-	{
-		RecurseDagForLeavesNoMerge(FirstDagNode, ParentInfo);
-		return true;
-	}
-
-	TSharedPtr<AlDagNode> DagNode = FirstDagNode;
-	uint32 MaxSize = 0;
-	while (DagNode)
-	{
-		MaxSize++;
-		DagNode = GetNextNode(DagNode);
-	}
-
-	DagNode = FirstDagNode;
-
-	TMap<uint32, TSharedPtr<BodyData>> ShellToProcess;
-
-	const char* ShaderName = nullptr;
-
-	while (DagNode)
-	{
-		AlObjectType objectType = DagNode->type();
-
-		// Process the current node.
-		switch (objectType)
-		{
-			// Push all leaf nodes into 'leaves'
-		case kShellNodeType:
-		{
-			AlShellNode* ShellNode = DagNode->asShellNodePtr();
-			if (!AlIsValid(ShellNode))
-			{
-				break;
-			}
-
-			AlShell* ShellPtr = ShellNode->shell();
-			if (!AlIsValid(ShellPtr))
-			{
-				break;
-			}
-
-			TAlObjectPtr<AlShell> Shell(ShellPtr);
-			uint32 NbPatch = GetPatchCount(Shell);
-
-			TAlObjectPtr<AlShader> Shader(Shell->firstShader());
-			if (Shader.IsValid())
-			{
-				ShaderName = Shader->name();
-			}
-
-			if (NbPatch == 1)
-			{
-				AddNodeInBodyGroup(DagNode, ShaderName, ShellToProcess, true, MaxSize);
-			}
-			else
-			{
-				ProcessAlShellNode(DagNode, ParentInfo, ShaderName);
-			}
-			break;
-		}
-
-		case kSurfaceNodeType:
-		{
-			AlSurfaceNode* SurfaceNode = DagNode->asSurfaceNodePtr();
-			TAlObjectPtr<AlSurface> Surface(SurfaceNode->surface());
-			if (Surface.IsValid())
-			{
-				TAlObjectPtr<AlShader> Shader(Surface->firstShader());
-				if (Shader.IsValid())
-				{
-					ShaderName = Shader->name();
-				}
-			}
-			AddNodeInBodyGroup(DagNode, ShaderName, ShellToProcess, true, MaxSize);
-			break;
-		}
-
-		case kMeshNodeType:
-		{
-			AlMeshNode* MeshNode = DagNode->asMeshNodePtr();
-			TAlObjectPtr<AlMesh> Mesh(MeshNode->mesh());
-			if (Mesh.IsValid())
-			{
-				TAlObjectPtr<AlShader> Shader(Mesh->firstShader());
-				if (Shader.IsValid())
-				{
-					ShaderName = Shader->name();
-				}
-			}
-			AddNodeInBodyGroup(DagNode, ShaderName, ShellToProcess, false, MaxSize);
-			break;
-		}
+#include "IWireInterface.h"
 
-		// Traverse down through groups
-		case kGroupNodeType:
-		{
-			ProcessAlGroupNode(*DagNode, ParentInfo);
-			break;
-		}
+#include "GenericPlatform/GenericPlatformTLS.h"
+#include "HAL/ConsoleManager.h"
+#include "Logging/LogMacros.h"
+#include "Misc/FileHelper.h"
+#include "Misc/Paths.h"
+#include "Modules/ModuleInterface.h"
+#include "Modules/ModuleManager.h"
 
-		default:
-		{
-			break;
-		}
-		}
+DEFINE_LOG_CATEGORY_STATIC(LogDatasmithWireTranslator, Log, All);
 
-		DagNode = GetNextNode(DagNode);
-	}
+static TAutoConsoleVariable<bool> CVarAliasThreadSafe(
+	TEXT("ds.WireTranslator.ThreadSafe"),
+	false,
+	TEXT("If true, the translator will be called in more than one thread. Default false.\n"),
+	ECVF_Default);
 
-	for (const TPair<uint32, TSharedPtr<BodyData>>& Body : ShellToProcess)
-	{
-		ProcessBodyNode(Body.Value, ParentInfo);
-	}
-	return true;
-}
+static TAutoConsoleVariable<bool> CVarAliasUseNative(
+	TEXT("ds.WireTranslator.UseNative"),
+	false,
+	TEXT("If true, the AlaisStudio tessellator will be used. Default false.\n"),
+	ECVF_Default);
 
-void FWireTranslatorImpl::RecurseDagForLeavesNoMerge(const TSharedPtr<AlDagNode>& FirstDagNode, FDagNodeInfo& ParentInfo)
-{
-	TSharedPtr<AlDagNode> DagNode = FirstDagNode;
-	while (DagNode)
-	{
-		DagForLeavesNoMerge(DagNode, ParentInfo);
-		DagNode = GetNextNode(DagNode);
-	}
-}
+static TAutoConsoleVariable<bool> CVarAliasSkipThinZoneMeshing(
+	TEXT("ds.WireTranslator.SkipThinZoneMeshing"),
+	true,
+	TEXT("If true, the global thin zone cvar will be honored. Default true.\n"),
+	ECVF_Default);
 
-void FWireTranslatorImpl::DagForLeavesNoMerge(const TSharedPtr<AlDagNode>& DagNode, FDagNodeInfo& ParentInfo)
-{
-	const char* ShaderName = nullptr;
+TAutoConsoleVariable<bool> GAliasLayerAsActor(
+	TEXT("ds.CADTranslator.Alias.LayersAsActors"),
+	false,
+	TEXT("If true, the first level of actors in the outliner is the layers. Default is false\n"),
+	ECVF_Default);
 
-	// Process the current node.
-	AlObjectType objectType = DagNode->type();
-	switch (objectType)
-	{
-		// Push all leaf nodes into 'leaves'
-	case kShellNodeType:
-	{
-		AlShellNode* ShellNode = DagNode->asShellNodePtr();
-		TAlObjectPtr<AlShell> Shell(ShellNode->shell());
-		if (Shell.IsValid())
-		{
-			TAlObjectPtr<AlShader> Shader(Shell->firstShader());
-			if (Shader.IsValid())
-			{
-				ShaderName = Shader->name();
-			}
+static TMap<uint32, FInterfaceMaker> RegisteredInterfaces;
+static uint64 AliasVersion = 0;
 
-			ProcessAlShellNode(DagNode, ParentInfo, ShaderName);
-		}
-		break;
-	}
-	case kSurfaceNodeType:
+namespace WireTranslator
+{
+	TSharedPtr<IWireInterface> GetInterfaceFromFile(const TCHAR* Filename)
 	{
-		AlSurfaceNode* SurfaceNode = DagNode->asSurfaceNodePtr();
-		TAlObjectPtr<AlSurface> Surface(SurfaceNode->surface());
-		if (Surface.IsValid())
-		{
-			TAlObjectPtr<AlShader> Shader(Surface->firstShader());
-			if (Shader.IsValid())
-			{
-				ShaderName = Shader->name();
-			}
-		}
-		ProcessAlShellNode(DagNode, ParentInfo, ShaderName);
-		break;
-	}
+		TSharedPtr<IWireInterface> WireInterface;
 
-	case kMeshNodeType:
-	{
-		AlMeshNode* MeshNode = DagNode->asMeshNodePtr();
-		TAlObjectPtr<AlMesh> Mesh(MeshNode->mesh());
-		if (Mesh.IsValid())
+		for (const TPair<uint32, FInterfaceMaker>& Entry : RegisteredInterfaces)
 		{
-			TAlObjectPtr<AlShader> Shader(Mesh->firstShader());
-			if (Shader.IsValid())
+			WireInterface = Entry.Value();
+			if (WireInterface->Initialize(Filename))
 			{
-				ShaderName = Shader->name();
+				return WireInterface;
 			}
 		}
-		ProcessAlShellNode(DagNode, ParentInfo, ShaderName);
-		break;
-	}
-
-	// Traverse down through groups
-	case kGroupNodeType:
-	{
-		ProcessAlGroupNode(*DagNode, ParentInfo);
-		break;
-	}
-
-	default:
-	{
-		break;
-	}
-
-	}
-}
-
-TSharedPtr<CADLibrary::ICADModelConverter> FWireTranslatorImpl::GetModelConverter(IAliasBRepConverter*& BRepConverter) const
-{
-	TSharedPtr<CADLibrary::ICADModelConverter> ModelConverter;
-
-	CADLibrary::FImportParameters ImportParameters;
-	if (CADLibrary::FImportParameters::bGDisableCADKernelTessellation)
-	{
-		TSharedRef<FAliasModelToTechSoftConverter> AliasToTechSoftConverter = MakeShared<FAliasModelToTechSoftConverter>(ImportParameters);
-		ModelConverter = AliasToTechSoftConverter;
-		BRepConverter = &AliasToTechSoftConverter.Get();
-	}
-	else
-	{
-		TSharedRef<FAliasModelToCADKernelConverter> AliasToCADKernelConverter = MakeShared<FAliasModelToCADKernelConverter>(ImportParameters);
-		ModelConverter = AliasToCADKernelConverter;
-		BRepConverter = &AliasToCADKernelConverter.Get();
-	}
-
-	ModelConverter->SetImportParameters(TessellationOptions.ChordTolerance, TessellationOptions.MaxEdgeLength, TessellationOptions.NormalTolerance, (CADLibrary::EStitchingTechnique)TessellationOptions.StitchingTechnique);
-
-	return ModelConverter;
-}
-
-TOptional<FMeshDescription> FWireTranslatorImpl::MeshDagNodeWithExternalMesher(AlDagNode& DagNode, TSharedRef<IDatasmithMeshElement> MeshElement, CADLibrary::FMeshParameters& MeshParameters)
-{
-	// Wire unit is cm
-	IAliasBRepConverter* BRepConverter = nullptr;
-	TSharedPtr<CADLibrary::ICADModelConverter> ModelConverter = GetModelConverter(BRepConverter);
-
-	ModelConverter->InitializeProcess();
-
-	EAliasObjectReference ObjectReference = EAliasObjectReference::LocalReference;
-
-	if (MeshParameters.bIsSymmetric)
-	{
-		// All actors of a Alias symmetric layer are defined in the world Reference i.e. they have identity transform. So Mesh actor has to be defined in the world reference.
-		ObjectReference = EAliasObjectReference::WorldReference;
-	}
-
-	const FColor* ColorPtr = nullptr;
-	const FString* ShaderName = MeshElementToShaderName.Find(&MeshElement.Get());
-	if (ShaderName)
-	{
-		ColorPtr = ShaderNameToColor.Find(*ShaderName);
-	}
-	if (!ColorPtr) // Should never happen
-	{
-		ColorPtr = &DefaultColor;
-	}
-
-	if (!BRepConverter->AddBRep(DagNode, *ColorPtr, ObjectReference))
-	{
-		return TOptional<FMeshDescription>();
-	}
-
-	ModelConverter->RepairTopology();
-
-	CADModelConverter->SaveModel(*OutputPath, MeshElement);
-
-	FMeshDescription MeshDescription;
-	DatasmithMeshHelper::PrepareAttributeForStaticMesh(MeshDescription);
-
-	bool bRet = ModelConverter->Tessellate(MeshParameters, MeshDescription);
-	if (!bRet)
-	{
-		const TCHAR* StaticMeshLable = MeshElement->GetLabel();
-		const TCHAR* StaticMeshName = MeshElement->GetName();
-		UE_LOG(LogDatasmithWireTranslator, Warning, TEXT("Failed to generate the mesh of \"%s\" (%s) StaticMesh."), StaticMeshLable, StaticMeshName);
-	}
-
-	return MoveTemp(MeshDescription);
-}
-
-TOptional<FMeshDescription> FWireTranslatorImpl::MeshDagNodeWithExternalMesher(TSharedRef<BodyData> Body, TSharedRef<IDatasmithMeshElement> MeshElement, CADLibrary::FMeshParameters& MeshParameters)
-{
-	// Wire unit is cm
-	IAliasBRepConverter* BRepConverter = nullptr;
-	TSharedPtr<CADLibrary::ICADModelConverter> ModelConverter = GetModelConverter(BRepConverter);
-
-	ModelConverter->InitializeProcess();
-
-	EAliasObjectReference ObjectReference = EAliasObjectReference::LocalReference;
-	if (MeshParameters.bIsSymmetric)
-	{
-		// All actors of a Alias symmetric layer are defined in the world Reference i.e. they have identity transform. So Mesh actor has to be defined in the world reference.
-		ObjectReference = EAliasObjectReference::WorldReference;
-	}
-	else if (TessellationOptions.StitchingTechnique == EDatasmithCADStitchingTechnique::StitchingSew)
-	{
-		// In the case of StitchingSew, AlDagNode children of a GroupNode are merged together. To be merged, they have to be defined in the reference of parent GroupNode.
-		ObjectReference = EAliasObjectReference::ParentReference;
-	}
-
-	for (const TPair<TSharedPtr<AlDagNode>, FColor>& DagNode : Body->GetShells())
-	{
-		BRepConverter->AddBRep(*DagNode.Key, DagNode.Value, ObjectReference);
-	}
-
-	ModelConverter->RepairTopology();
 
-	ModelConverter->SaveModel(*OutputPath, MeshElement);
-
-	FMeshDescription MeshDescription;
-	DatasmithMeshHelper::PrepareAttributeForStaticMesh(MeshDescription);
-
-	ModelConverter->Tessellate(MeshParameters, MeshDescription);
-
-	return MoveTemp(MeshDescription);
-}
-
-TOptional<FMeshDescription> FWireTranslatorImpl::GetMeshOfShellNode(AlDagNode& DagNode, TSharedRef<IDatasmithMeshElement> MeshElement, CADLibrary::FMeshParameters& MeshParameters)
-{
-	if (CADModelConverter.IsValid() && CADModelConverter->IsSessionValid())
-	{
-		TOptional<FMeshDescription> UEMesh = MeshDagNodeWithExternalMesher(DagNode, MeshElement, MeshParameters);
-		return UEMesh;
-	}
-	else
-	{
-		AlMatrix4x4 AlMatrix;
-		DagNode.inverseGlobalTransformationMatrix(AlMatrix);
-		// TODO: the best way, should be to don't have to apply inverse global transform to the generated mesh
-		TSharedPtr<AlDagNode> TesselatedNode = OpenModelUtils::TesselateDagLeaf(DagNode, ETesselatorType::Fast, TessellationOptions.ChordTolerance);
-		if (TesselatedNode.IsValid())
-		{
-			// Get the meshes from the dag nodes. Note that removing the mesh's DAG.
-			// will also removes the meshes, so we have to do it later.
-			TOptional<FMeshDescription> UEMesh = GetMeshOfNodeMesh(*TesselatedNode, MeshElement, MeshParameters, &AlMatrix);
-			return UEMesh;
-		}
+		return TSharedPtr<IWireInterface>();
 	}
-	return TOptional<FMeshDescription>();
-}
-
-TOptional<FMeshDescription> FWireTranslatorImpl::GetMeshOfShellBody(TSharedRef<BodyData> Body, TSharedRef<IDatasmithMeshElement> MeshElement, CADLibrary::FMeshParameters& MeshParameters)
-{
-	TOptional<FMeshDescription> UEMesh = MeshDagNodeWithExternalMesher(Body, MeshElement, MeshParameters);
-	return UEMesh;
 }
 
-TOptional<FMeshDescription> FWireTranslatorImpl::GetMeshOfMeshBody(TSharedRef<BodyData> Body, TSharedRef<IDatasmithMeshElement> MeshElement, CADLibrary::FMeshParameters& MeshParameters)
+uint64 IWireInterface::GetRequiredAliasVersion()
 {
-	FMeshDescription MeshDescription;
-	DatasmithMeshHelper::PrepareAttributeForStaticMesh(MeshDescription);
-	MeshDescription.Empty();
-
-	for (const TPair<TSharedPtr<AlDagNode>, FColor>& DagNode : Body->GetShells())
-	{
-		AlMeshNode* MeshNode = DagNode.Key->asMeshNodePtr();
-		if (!AlIsValid(MeshNode))
+	static bool bInitialized = []() -> bool
 		{
-			continue;
-		}
-
-		AlMesh* MeshPtr = MeshNode->mesh();
-		if (!AlIsValid(MeshPtr))
-		{
-			continue;
-		}
-		TAlObjectPtr<AlMesh> Mesh(MeshPtr);
+			const TCHAR* AliasDllName = TEXT("libalias_api.dll");
 
-		AlMatrix4x4 AlMatrix;
-		DagNode.Key->localTransformationMatrix(AlMatrix);
-
-		Mesh->transform(AlMatrix);
-
-		const FColor& ShaderColor = DagNode.Value;
-		const int32 ShaderId = CreateShaderId(ShaderColor);
-		const FString SlotMaterialName = FString::FromInt(ShaderId);
-
-		const bool bMerge = true;
-		OpenModelUtils::TransferAlMeshToMeshDescription(*MeshPtr, *SlotMaterialName, MeshDescription, MeshParameters, bMerge);
-	}
+			void* DllHandle = FPlatformProcess::GetDllHandle(AliasDllName);
+			if (DllHandle)
+			{
+				AliasVersion = FPlatformMisc::GetFileVersion(AliasDllName);
+			}
 
-	// Build edge meta data
-	FStaticMeshOperations::DetermineEdgeHardnessesFromVertexInstanceNormals(MeshDescription);
+			return true;
+		}();
 
-	return MoveTemp(MeshDescription);
+	return AliasVersion;
 }
 
-TOptional<FMeshDescription> FWireTranslatorImpl::GetMeshOfNodeMesh(AlDagNode& TesselatedNode, TSharedRef<IDatasmithMeshElement> MeshElement, CADLibrary::FMeshParameters& MeshParameters, AlMatrix4x4* AlMeshInvGlobalMatrix)
+void IWireInterface::RegisterInterface(int16 MajorVersion, int16 MinorVersion, FInterfaceMaker&& MakeInterface)
 {
-	AlMeshNode* MeshNode = TesselatedNode.asMeshNodePtr();
-	if (!AlIsValid(MeshNode))
-	{
-		return TOptional<FMeshDescription>();
-	}
-
-	AlMesh* Mesh = MeshNode->mesh();
-	if (!AlIsValid(Mesh))
-	{
-		return TOptional<FMeshDescription>();
-	}
-
-	TAlObjectPtr<AlMesh> SharedMesh(Mesh);
-	if (AlMeshInvGlobalMatrix != nullptr)
-	{
-		SharedMesh->transform(*AlMeshInvGlobalMatrix);
-	}
-
-	return ImportMesh(*Mesh, MeshElement, MeshParameters);
+	uint32 Version = (MajorVersion << 16) + MinorVersion;
+	RegisteredInterfaces.Emplace(Version, MoveTemp(MakeInterface));
 }
 
-TOptional<FMeshDescription> FWireTranslatorImpl::GetMeshDescription(TSharedRef<IDatasmithMeshElement> MeshElement, CADLibrary::FMeshParameters& MeshParameters, TSharedRef<BodyData> Body)
+FDatasmithWireTranslator::FDatasmithWireTranslator()
 {
-	if (Body->GetShells().Num() == 0)
-	{
-		return TOptional<FMeshDescription>();
-	}
-
-	const TArray<TPair<TSharedPtr<AlDagNode>, FColor>>& Shells = Body->GetShells();
-	TSharedPtr<AlDagNode> DagNode = Shells[0].Key;
-	AlLayer* LayerPtr = DagNode->layer();
-	if (AlIsValid(LayerPtr))
-	{
-		const TAlObjectPtr<AlLayer> Layer(LayerPtr);
-		if (LayerPtr->isSymmetric())
-		{
-			MeshParameters.bIsSymmetric = true;
-			double Normal[3], Origin[3];
-			LayerPtr->symmetricNormal(Normal[0], Normal[1], Normal[2]);
-			LayerPtr->symmetricOrigin(Origin[0], Origin[1], Origin[2]);
-
-			MeshParameters.SymmetricOrigin.X = (float)Origin[0];
-			MeshParameters.SymmetricOrigin.Y = (float)Origin[1];
-			MeshParameters.SymmetricOrigin.Z = (float)Origin[2];
-			MeshParameters.SymmetricNormal.X = (float)Normal[0];
-			MeshParameters.SymmetricNormal.Y = (float)Normal[1];
-			MeshParameters.SymmetricNormal.Z = (float)Normal[2];
-		}
-	}
-
-	if (Body->IsCadData())
-	{
-		return GetMeshOfShellBody(Body, MeshElement, MeshParameters);
-	}
-	else
-	{
-		return GetMeshOfMeshBody(Body, MeshElement, MeshParameters);
-	}
 }
 
-TOptional<FMeshDescription> FWireTranslatorImpl::GetMeshDescription(TSharedRef<IDatasmithMeshElement> MeshElement, CADLibrary::FMeshParameters& MeshParameters)
+bool FDatasmithWireTranslator::CanTranslate()
 {
-	TSharedPtr<AlDagNode>* DagNodeTemp = MeshElementToAlDagNodeMap.Find(&MeshElement.Get());
-	if (DagNodeTemp == nullptr || !DagNodeTemp->IsValid())
-	{
-		TSharedPtr<BodyData>* BodyTemp = MeshElementToBodyMap.Find(&MeshElement.Get());
-		if (BodyTemp != nullptr && BodyTemp->IsValid())
-		{
-			return GetMeshDescription(MeshElement, MeshParameters, (*BodyTemp).ToSharedRef());
-		}
-		return TOptional<FMeshDescription>();
-	}
-
-	AlDagNode& DagNode = **DagNodeTemp;
-	AlObjectType objectType = DagNode.type();
-
-	if (objectType == kShellNodeType || objectType == kSurfaceNodeType || objectType == kMeshNodeType)
-	{
-		boolean bAlOrientation;
-		DagNode.getSurfaceOrientation(bAlOrientation);
-		MeshParameters.bNeedSwapOrientation = (objectType == kMeshNodeType) ? (bool)bAlOrientation : false; 
-
-		AlLayer* LayerPtr = DagNode.layer();
-		if (AlIsValid(LayerPtr))
+	static const bool bCanTranslate = []() -> bool
 		{
-			const TAlObjectPtr<AlLayer> Layer(LayerPtr);
-			if (LayerPtr->isSymmetric())
+			if (IWireInterface::GetRequiredAliasVersion() == 0)
 			{
-				MeshParameters.bIsSymmetric = true;
-				double Normal[3], Origin[3];
-				LayerPtr->symmetricNormal(Normal[0], Normal[1], Normal[2]);
-				LayerPtr->symmetricOrigin(Origin[0], Origin[1], Origin[2]);
-
-				MeshParameters.SymmetricOrigin.X = (float)Origin[0];
-				MeshParameters.SymmetricOrigin.Y = (float)Origin[1];
-				MeshParameters.SymmetricOrigin.Z = (float)Origin[2];
-				MeshParameters.SymmetricNormal.X = (float)Normal[0];
-				MeshParameters.SymmetricNormal.Y = (float)Normal[1];
-				MeshParameters.SymmetricNormal.Z = (float)Normal[2];
+				return false;
 			}
-		}
-	}
-
-	switch (objectType)
-	{
-	case kShellNodeType:
-	case kSurfaceNodeType:
-	{
-		return GetMeshOfShellNode(DagNode, MeshElement, MeshParameters);
-		break;
-	}
-
-	case kMeshNodeType:
-	{
-		return GetMeshOfNodeMesh(DagNode, MeshElement, MeshParameters);
-		break;
-	}
-
-	default:
-		break;
-	}
-
-	return TOptional<FMeshDescription>();
-}
-
-
-// Note that Alias file unit is cm like UE
-TOptional<FMeshDescription> FWireTranslatorImpl::ImportMesh(AlMesh& InMesh, TSharedRef<IDatasmithMeshElement> MeshElement, CADLibrary::FMeshParameters& InMeshParameters)
-{
-	FMeshDescription MeshDescription;
-	DatasmithMeshHelper::PrepareAttributeForStaticMesh(MeshDescription);
-
-	const FColor* ShaderColorPtr = nullptr;
-	const FString* ShaderName = MeshElementToShaderName.Find(&MeshElement.Get());
-	if (ShaderName)
-	{
-		ShaderColorPtr = ShaderNameToColor.Find(*ShaderName);
-	}
-	if (!ShaderColorPtr) // Should never happen
-	{
-		ShaderColorPtr = &DefaultColor;
-	}
-	int32 ShaderId = CreateShaderId(*ShaderColorPtr);
-	FString SlotMaterialName = FString::FromInt(ShaderId);
-
-	const bool bMerge = false;
-	OpenModelUtils::TransferAlMeshToMeshDescription(InMesh, *SlotMaterialName, MeshDescription, InMeshParameters, bMerge);
-
-	// Build edge meta data
-	FStaticMeshOperations::DetermineEdgeHardnessesFromVertexInstanceNormals(MeshDescription);
-
-	return MoveTemp(MeshDescription);
-}
 
-#define TRACK_MESH 0
-
-bool FWireTranslatorImpl::LoadStaticMesh(const TSharedRef<IDatasmithMeshElement> MeshElement, FDatasmithMeshElementPayload& OutMeshPayload, const FDatasmithTessellationOptions& InTessellationOptions)
-{
-#if TRACK_MESH
-	static const FString MeshName(TEXT("shell_30181"));
-	bool bSkip = true;
-	if (MeshName.Equals(MeshElement->GetLabel()))
-	{
-		bSkip = false;
-	}
-	if (bSkip)
-	{
-		return false;
-	}
-#endif
-	CADLibrary::FMeshParameters MeshParameters;
-
-	uint64 StartTime = FPlatformTime::Cycles64();
-
-	if (TOptional<FMeshDescription> Mesh = GetMeshDescription(MeshElement, MeshParameters))
-	{
-		OutMeshPayload.LodMeshes.Add(MoveTemp(Mesh.GetValue()));
-		const TCHAR* MeshFilename = MeshElement->GetFile();
-		if (!bAliasUseNative && FPaths::FileExists(MeshFilename))
-		{
-			CADModelConverter->AddSurfaceDataForMesh(MeshFilename, MeshParameters, InTessellationOptions, OutMeshPayload);
-
-			// Remove the file because it is temporary since caching is disabled.
-			if (!CADLibrary::FImportParameters::bGEnableCADCache)
+			if (!WireTranslator::GetInterfaceFromFile(nullptr).IsValid())
 			{
-				IFileManager::Get().Delete(MeshFilename);
+				return false;
 			}
-		}
-	}
-
-	double ElapsedSeconds = FPlatformTime::ToSeconds64(FPlatformTime::Cycles64() - StartTime);
-
-	UE_LOG(LogDatasmithWireTranslator, Display, TEXT("Mesh description of %s retrieved in %.3f s"), MeshElement->GetLabel(), ElapsedSeconds);
 
-	return OutMeshPayload.LodMeshes.Num() > 0;
-}
-
-
-#endif
-
-//////////////////////////////////////////////////////////////////////////
-// UDatasmithWireTranslator
-//////////////////////////////////////////////////////////////////////////
+			return true;
+		}();
 
-FDatasmithWireTranslator::FDatasmithWireTranslator()
-	: Translator(nullptr)
-{
+		return bCanTranslate;
 }
 
 void FDatasmithWireTranslator::Initialize(FDatasmithTranslatorCapabilities& OutCapabilities)
 {
-#if WITH_EDITOR
-	if (GIsEditor && !GEditor->PlayWorld && !GIsPlayInEditorWorld)
+	if (!CanTranslate())
 	{
-#ifdef USE_OPENMODEL
-		static void* DllHandle = FPlatformProcess::GetDllHandle(TEXT("libalias_api.dll"));
-		if (DllHandle)
-		{
-			// Check installed version of Alias Tools because binaries before 2021.3 are not compatible with Alias 2022
-			uint64 FileVersion = FPlatformMisc::GetFileVersion(TEXT("libalias_api.dll"));
-
-#ifdef OPEN_MODEL_2020
-			if (LibAlias2020_Version < FileVersion && FileVersion < LibAlias2021_Version)
-			{
-				TFunction<bool()> DisplayMessage = []() -> bool
-				{
-					UE_LOG(LogDatasmithWireTranslator, Warning, TEXT(WRONG_VERSION_TEXT)); return true;
-				};
-				static const bool bIsDispaly = DisplayMessage();
+		OutCapabilities.bIsEnabled = false;
+		return;
+	}
 
-				const IConsoleVariable* ConsoleVariable = IConsoleManager::Get().FindConsoleVariable(TEXT("ds.Wiretranslator.ThreadSafe"));
-				OutCapabilities.bParallelLoadStaticMeshSupported = ConsoleVariable ? ConsoleVariable->GetBool() : false;
+	const int32 MajorVersion = (AliasVersion & 0xffff0000) >> 16;
+	const int32 MinorVersion = AliasVersion & 0x0000ffff;
 
-				OutCapabilities.bIsEnabled = false;
-				return;
-			}
-#endif
+	const FString FormatDescription = FString::Printf(TEXT("AliasStudio %d.%d model files"), MajorVersion, MinorVersion);
 
-			if (LibAliasVersionMin <= FileVersion && FileVersion < LibAliasVersionMax)
-			{
-				const FString AliasVersion = FString::Printf(TEXT("AliasStudio %s model files"), *AliasSdkVersion);
-				OutCapabilities.SupportedFileFormats.Add(FFileFormatInfo{ TEXT("wire"), *AliasVersion });
-				
-				const IConsoleVariable* ConsoleVariable = IConsoleManager::Get().FindConsoleVariable(TEXT("ds.Wiretranslator.ThreadSafe"));
-				OutCapabilities.bParallelLoadStaticMeshSupported = ConsoleVariable ? ConsoleVariable->GetBool() : false;
-				
-				OutCapabilities.bIsEnabled = true;
-				return;
-			}
+	OutCapabilities.SupportedFileFormats.Add(FFileFormatInfo{ TEXT("wire"), *FormatDescription });
 
-			OutCapabilities.bIsEnabled = false;
-			return;
-		}
-#endif
-	}
-#endif
+	const IConsoleVariable* ConsoleVariable = IConsoleManager::Get().FindConsoleVariable(TEXT("ds.Wiretranslator.ThreadSafe"));
+	OutCapabilities.bParallelLoadStaticMeshSupported = ConsoleVariable ? ConsoleVariable->GetBool() : false;
 
-	OutCapabilities.bIsEnabled = false;
+	OutCapabilities.bIsEnabled = true;
 }
 
 bool FDatasmithWireTranslator::IsSourceSupported(const FDatasmithSceneSource& Source)
 {
-#ifdef USE_OPENMODEL
-	return true;
-#else
-	return false;
-#endif
+	WireInterface = WireTranslator::GetInterfaceFromFile(*Source.GetSourceFile());
+	return WireInterface.IsValid();
 }
 
 bool FDatasmithWireTranslator::LoadScene(TSharedRef<IDatasmithScene> OutScene)
 {
-#ifdef USE_OPENMODEL
-	const FString& Filename = GetSource().GetSourceFile();
-
-	Translator = MakeShared<FWireTranslatorImpl>(Filename, OutScene);
-	if (!Translator)
+	//WireInterface = WireTranslator::GetInterfaceFromFile(*Filename);
+	if (!WireInterface.IsValid())
 	{
 		return false;
 	}
 
-	UE_LOG(LogDatasmithWireTranslator, Display, TEXT("CAD translation [%s]."), *Filename);
+	//const FString& Filename = GetSource().GetSourceFile();
+
+	UE_LOG(LogDatasmithWireTranslator, Display, TEXT("CAD translation [%s]."), *GetSource().GetSourceFile());
 	UE_LOG(LogDatasmithWireTranslator, Display, TEXT(" - Parsing Library:      %s"), TEXT("Alias"));
-	UE_LOG(LogDatasmithWireTranslator, Display, TEXT(" - Tessellation Library: %s"), CADLibrary::FImportParameters::bGDisableCADKernelTessellation ? TEXT("TechSoft") : TEXT("CADKernel"));
+	UE_LOG(LogDatasmithWireTranslator, Display, TEXT(" - Tessellation Library: %s"), CommonTessellationOptions.bUseCADKernel ? TEXT("CADKernel") : TEXT("TechSoft"));
 	UE_LOG(LogDatasmithWireTranslator, Display, TEXT(" - Import parameters:"));
 
-	const FDatasmithTessellationOptions& TessellationOptions = GetCommonTessellationOptions();
-	UE_LOG(LogDatasmithWireTranslator, Display, TEXT("     - ChordTolerance:     %lf"), TessellationOptions.ChordTolerance);
-	UE_LOG(LogDatasmithWireTranslator, Display, TEXT("     - MaxEdgeLength:      %lf"), TessellationOptions.MaxEdgeLength);
-	UE_LOG(LogDatasmithWireTranslator, Display, TEXT("     - MaxNormalAngle:     %lf"), TessellationOptions.NormalTolerance);
+	UE_LOG(LogDatasmithWireTranslator, Display, TEXT("     - ChordTolerance:     %lf"), CommonTessellationOptions.ChordTolerance);
+	UE_LOG(LogDatasmithWireTranslator, Display, TEXT("     - MaxEdgeLength:      %lf"), CommonTessellationOptions.MaxEdgeLength);
+	UE_LOG(LogDatasmithWireTranslator, Display, TEXT("     - MaxNormalAngle:     %lf"), CommonTessellationOptions.NormalTolerance);
+	
 	FString StitchingTechnique;
-	switch (TessellationOptions.StitchingTechnique)
+	switch (CommonTessellationOptions.StitchingTechnique)
 	{
 	case EDatasmithCADStitchingTechnique::StitchingHeal:
 		StitchingTechnique = TEXT("Heal");
@@ -2556,56 +169,61 @@ bool FDatasmithWireTranslator::LoadScene(TSharedRef<IDatasmithScene> OutScene)
 	}
 	UE_LOG(LogDatasmithWireTranslator, Display, TEXT("     - StitchingTechnique: %s"), *StitchingTechnique);
 
-	if (!CADLibrary::FImportParameters::bGDisableCADKernelTessellation)
+	if (CommonTessellationOptions.bUseCADKernel)
 	{
 		UE_LOG(LogDatasmithWireTranslator, Display, TEXT("     - Stitching Options:"));
-		UE_LOG(LogDatasmithWireTranslator, Display, TEXT("         - ForceSew:              %s"), CADLibrary::FImportParameters::bGStitchingForceSew ? TEXT("True") : TEXT("False"));
-		UE_LOG(LogDatasmithWireTranslator, Display, TEXT("         - RemoveThinFaces:       %s"), CADLibrary::FImportParameters::bGStitchingRemoveThinFaces ? TEXT("True") : TEXT("False"));
-		UE_LOG(LogDatasmithWireTranslator, Display, TEXT("         - RemoveDuplicatedFaces: %s"), CADLibrary::FImportParameters::bGStitchingRemoveDuplicatedFaces ? TEXT("True") : TEXT("False"));
-		UE_LOG(LogDatasmithWireTranslator, Display, TEXT("         - ForceFactor:           %f"), CADLibrary::FImportParameters::GStitchingForceFactor);
+		//UE_LOG(LogDatasmithWireTranslator, Display, TEXT("         - ForceSew:              %s"), CADLibrary::FImportParameters::bGStitchingForceSew ? TEXT("True") : TEXT("False"));
+		//UE_LOG(LogDatasmithWireTranslator, Display, TEXT("         - RemoveThinFaces:       %s"), CADLibrary::FImportParameters::bGStitchingRemoveThinFaces ? TEXT("True") : TEXT("False"));
+		//UE_LOG(LogDatasmithWireTranslator, Display, TEXT("         - RemoveDuplicatedFaces: %s"), CADLibrary::FImportParameters::bGStitchingRemoveDuplicatedFaces ? TEXT("True") : TEXT("False"));
+		//UE_LOG(LogDatasmithWireTranslator, Display, TEXT("         - ForceFactor:           %f"), CADLibrary::FImportParameters::GStitchingForceFactor);
 	}
 
-	FString OutputPath = FPaths::ConvertRelativePathToFull(FPaths::Combine(FDatasmithWireTranslatorModule::Get().GetTempDir(), TEXT("Cache"), GetSource().GetSceneName()));
+	static const FString CacheRootDir = FPaths::ConvertRelativePathToFull(FPaths::Combine(FPaths::ProjectSavedDir(), TEXT("WireTranslator"), TEXT("Cache")));
+	
+	const FString OutputPath = FPaths::Combine(CacheRootDir, GetSource().GetSceneName());
 	IFileManager::Get().MakeDirectory(*OutputPath, true);
-	Translator->SetOutputPath(OutputPath);
-
-	Translator->SetTessellationOptions(GetCommonTessellationOptions());
 
-	return Translator->Read();
+	WireInterface->SetOutputPath(OutputPath);
+	WireInterface->SetTessellationOptions(CommonTessellationOptions);
 
-#else
-	return false;
-#endif
+	return WireInterface->Load(OutScene);
 }
 
 void FDatasmithWireTranslator::UnloadScene()
 {
-	Translator.Reset();
+	WireInterface.Reset();
 }
 
 bool FDatasmithWireTranslator::LoadStaticMesh(const TSharedRef<IDatasmithMeshElement> MeshElement, FDatasmithMeshElementPayload& OutMeshPayload)
 {
-#ifdef USE_OPENMODEL
-	if (Translator)
-	{
-		return Translator->LoadStaticMesh(MeshElement, OutMeshPayload, GetCommonTessellationOptions());
-	}
-#endif
-	return false;
+	return ensure(WireInterface.IsValid()) ? WireInterface->LoadStaticMesh(MeshElement, OutMeshPayload, CommonTessellationOptions) : false;
 }
 
 void FDatasmithWireTranslator::SetSceneImportOptions(const TArray<TObjectPtr<UDatasmithOptionsBase>>& Options)
 {
-#ifdef USE_OPENMODEL
 	FParametricSurfaceTranslator::SetSceneImportOptions(Options);
 
-	if (Translator)
+	if (ensure(WireInterface.IsValid()))
 	{
-		Translator->SetTessellationOptions(GetCommonTessellationOptions());
+		WireInterface->SetTessellationOptions(CommonTessellationOptions);
 	}
-#endif
 }
 
-} // namespace
+class FDatasmithWireTranslatorModule : public IModuleInterface
+{
+public:
+	virtual void StartupModule() override
+	{
+		Datasmith::RegisterTranslator<FDatasmithWireTranslator>();
+	}
+
+	virtual void ShutdownModule() override
+	{
+		Datasmith::UnregisterTranslator<FDatasmithWireTranslator>();
+	}
+};
+
+IMPLEMENT_MODULE(FDatasmithWireTranslatorModule, DatasmithWireTranslator);
+
 
 #undef LOCTEXT_NAMESPACE // "DatasmithWireTranslator"
diff --git a/Engine/Plugins/Enterprise/DatasmithCADImporter/Source/DatasmithWireTranslator/Private/DatasmithWireTranslator.h b/Engine/Plugins/Enterprise/DatasmithCADImporter/Source/DatasmithWireTranslator/Private/DatasmithWireTranslator.h
--- a/Engine/Plugins/Enterprise/DatasmithCADImporter/Source/DatasmithWireTranslator/Private/DatasmithWireTranslator.h
+++ b/Engine/Plugins/Enterprise/DatasmithCADImporter/Source/DatasmithWireTranslator/Private/DatasmithWireTranslator.h
@@ -4,15 +4,40 @@
 
 #include "CoreMinimal.h"
 
-#include "OpenModelUtils.h"
 #include "ParametricSurfaceTranslator.h"
+
+#include "DatasmithImportOptions.h"
+
 #include "UObject/ObjectMacros.h"
 
+#include "DatasmithWireTranslator.generated.h"
+
+class IWireInterface;
+
+USTRUCT(BlueprintType)
+struct FDatasmithWireOptions : public FDatasmithTessellationOptions
+{
+	GENERATED_BODY()
+
+	UPROPERTY(config, EditAnywhere, BlueprintReadWrite, Category = "Wire Translation Options")
+	bool bUseLayerAsActor = false;
+
+	uint32 GetHash() const
+	{
+		return HashCombine(FDatasmithTessellationOptions::GetHash(), GetTypeHash(false));
+	}
+};
+
 
-namespace UE_DATASMITHWIRETRANSLATOR_NAMESPACE
+UCLASS(BlueprintType, config = EditorPerProjectUserSettings)
+class UDatasmithWireTranslationOptions : public UDatasmithOptionsBase
 {
+	GENERATED_BODY()
 
-class FWireTranslatorImpl;
+public:
+	UPROPERTY(config, EditAnywhere, BlueprintReadWrite, Category = "Wire Translation Options", meta = (ShowOnlyInnerProperties))
+	FDatasmithWireOptions Options;
+};
 
 class FDatasmithWireTranslator : public FParametricSurfaceTranslator
 {
@@ -43,7 +68,8 @@ protected:
 	// End ADatasmithCoreTechTranslator overrides
 
 private:
-	TSharedPtr<FWireTranslatorImpl> Translator;
-};
+	bool CanTranslate();
 
-}
\ No newline at end of file
+private:
+	TSharedPtr<IWireInterface> WireInterface;
+};
diff --git a/Engine/Plugins/Enterprise/DatasmithCADImporter/Source/DatasmithWireTranslator/Private/DatasmithWireTranslatorModule.cpp b/Engine/Plugins/Enterprise/DatasmithCADImporter/Source/DatasmithWireTranslator/Private/DatasmithWireTranslatorModule.cpp
--- a/Engine/Plugins/Enterprise/DatasmithCADImporter/Source/DatasmithWireTranslator/Private/DatasmithWireTranslatorModule.cpp
+++ b/Engine/Plugins/Enterprise/DatasmithCADImporter/Source/DatasmithWireTranslator/Private/DatasmithWireTranslatorModule.cpp
@@ -1,38 +0,0 @@
-// Copyright Epic Games, Inc. All Rights Reserved.
-
-#include "DatasmithWireTranslatorModule.h"
-
-#include "HAL/FileManager.h"
-#include "Misc/Paths.h"
-#include "Modules/ModuleInterface.h"
-#include "Modules/ModuleManager.h"
-
-#include "DatasmithWireTranslator.h"
-
-// need this macro wrapper to expand the UE_DATASMITHWIRETRANSLATOR_MODULE_NAME macro and not create symbols with "UE_DATASMITHWIRETRANSLATOR_MODULE_NAME" in the token
-#define IMPLEMENT_MODULE_WRAPPER(ModuleName) IMPLEMENT_MODULE(UE_DATASMITHWIRETRANSLATOR_NAMESPACE::FDatasmithWireTranslatorModule, ModuleName);
-IMPLEMENT_MODULE_WRAPPER(UE_DATASMITHWIRETRANSLATOR_MODULE_NAME)
-
-namespace UE_DATASMITHWIRETRANSLATOR_NAMESPACE
-{
-
-void FDatasmithWireTranslatorModule::StartupModule()
-{
-	// Create temporary directory which will be used by CoreTech to store tessellation data
-	TempDir = FPaths::Combine(FPaths::ProjectIntermediateDir(), TEXT("WireImportTemp"));
-	IFileManager::Get().MakeDirectory(*TempDir);
-
-	Datasmith::RegisterTranslator<FDatasmithWireTranslator>();
-}
-
-void FDatasmithWireTranslatorModule::ShutdownModule()
-{
-	Datasmith::UnregisterTranslator<FDatasmithWireTranslator>();
-}
-
-FString FDatasmithWireTranslatorModule::GetTempDir() const
-{
-	return TempDir;
-}
-
-}
\ No newline at end of file
diff --git a/Engine/Plugins/Enterprise/DatasmithCADImporter/Source/DatasmithWireTranslator/Private/OpenModelUtils.cpp b/Engine/Plugins/Enterprise/DatasmithCADImporter/Source/DatasmithWireTranslator/Private/OpenModelUtils.cpp
deleted file mode 100644
diff --git a/Engine/Plugins/Enterprise/DatasmithCADImporter/Source/DatasmithWireTranslator/Private/OpenModelUtils.h b/Engine/Plugins/Enterprise/DatasmithCADImporter/Source/DatasmithWireTranslator/Private/OpenModelUtils.h
deleted file mode 100644
diff --git a/Engine/Plugins/Enterprise/DatasmithCADImporter/Source/DatasmithWireTranslator/Public/DatasmithWireTranslatorModule.h b/Engine/Plugins/Enterprise/DatasmithCADImporter/Source/DatasmithWireTranslator/Public/DatasmithWireTranslatorModule.h
deleted file mode 100644
diff --git a/Engine/Plugins/Enterprise/DatasmithCADImporter/Source/DatasmithWireTranslator/Public/IWireInterface.h b/Engine/Plugins/Enterprise/DatasmithCADImporter/Source/DatasmithWireTranslator/Public/IWireInterface.h
new file mode 100644
--- /dev/null
+++ b/Engine/Plugins/Enterprise/DatasmithCADImporter/Source/DatasmithWireTranslator/Public/IWireInterface.h
@@ -0,0 +1,30 @@
+// Copyright Epic Games, Inc. All Rights Reserved.
+
+#pragma once
+
+#include "CoreMinimal.h"
+
+class IDatasmithMeshElement;
+class IDatasmithScene;
+
+struct FDatasmithTessellationOptions;
+struct FDatasmithMeshElementPayload;
+
+typedef TFunction<TSharedPtr<class IWireInterface>()> FInterfaceMaker;
+
+class IWireInterface
+{
+public:
+	virtual ~IWireInterface() {}
+
+	virtual bool Initialize(const TCHAR* Filename = nullptr) = 0;
+
+	virtual bool Load(TSharedPtr<IDatasmithScene> Scene) = 0;
+
+	virtual void SetTessellationOptions(const FDatasmithTessellationOptions& Options) = 0;
+	virtual void SetOutputPath(const FString& Path) = 0;
+	virtual bool LoadStaticMesh(const TSharedPtr<IDatasmithMeshElement> MeshElement, FDatasmithMeshElementPayload& OutMeshPayload, const FDatasmithTessellationOptions& InTessellationOptions) = 0;
+
+	static DATASMITHWIRETRANSLATOR_API uint64 GetRequiredAliasVersion();
+	static DATASMITHWIRETRANSLATOR_API void RegisterInterface(int16 MajorVersion, int16 MinorVersion, FInterfaceMaker&& MakeInterface);
+};
diff --git a/Engine/Plugins/Enterprise/DatasmithCADImporter/Source/ParametricSurface/Private/ParametricSurfaceTranslator.cpp b/Engine/Plugins/Enterprise/DatasmithCADImporter/Source/ParametricSurface/Private/ParametricSurfaceTranslator.cpp
--- a/Engine/Plugins/Enterprise/DatasmithCADImporter/Source/ParametricSurface/Private/ParametricSurfaceTranslator.cpp
+++ b/Engine/Plugins/Enterprise/DatasmithCADImporter/Source/ParametricSurface/Private/ParametricSurfaceTranslator.cpp
@@ -5,11 +5,19 @@
 #include "ParametricSurfaceData.h"
 #include "ParametricSurfaceModule.h"
 
+#include "CADOptions.h"
+
 #include "DatasmithImportOptions.h"
 #include "IDatasmithSceneElements.h"
 
 #include "Misc/FileHelper.h"
 
+FParametricSurfaceTranslator::FParametricSurfaceTranslator()
+{
+	// Initialize bUseCADKernel with current value of CVar ds.CADTranslator.DisableCADKernelTessellation
+	CommonTessellationOptions.bUseCADKernel = !CADLibrary::FImportParameters::bGDisableCADKernelTessellation;
+}
+
 void FParametricSurfaceTranslator::GetSceneImportOptions(TArray<TObjectPtr<UDatasmithOptionsBase>>& Options)
 {
 	FString Extension = GetSource().GetSourceFileExtension();
diff --git a/Engine/Plugins/Enterprise/DatasmithCADImporter/Source/ParametricSurface/Private/TechSoft/CADModelToTechSoftConverterBase.cpp b/Engine/Plugins/Enterprise/DatasmithCADImporter/Source/ParametricSurface/Private/TechSoft/CADModelToTechSoftConverterBase.cpp
--- a/Engine/Plugins/Enterprise/DatasmithCADImporter/Source/ParametricSurface/Private/TechSoft/CADModelToTechSoftConverterBase.cpp
+++ b/Engine/Plugins/Enterprise/DatasmithCADImporter/Source/ParametricSurface/Private/TechSoft/CADModelToTechSoftConverterBase.cpp
@@ -44,7 +44,7 @@ void FCADModelToTechSoftConverterBase::InitializeProcess()
 }
 
 
-bool FCADModelToTechSoftConverterBase::SaveModel(const TCHAR* InFolderPath, TSharedRef<IDatasmithMeshElement>& MeshElement)
+bool FCADModelToTechSoftConverterBase::SaveModel(const TCHAR* InFolderPath, TSharedPtr<IDatasmithMeshElement> MeshElement)
 {
 #ifdef USE_TECHSOFT_SDK
 	FString FilePath = FPaths::Combine(InFolderPath, MeshElement->GetName()) + TEXT(".prc");
diff --git a/Engine/Plugins/Enterprise/DatasmithCADImporter/Source/ParametricSurface/Public/CADModelConverter.h b/Engine/Plugins/Enterprise/DatasmithCADImporter/Source/ParametricSurface/Public/CADModelConverter.h
--- a/Engine/Plugins/Enterprise/DatasmithCADImporter/Source/ParametricSurface/Public/CADModelConverter.h
+++ b/Engine/Plugins/Enterprise/DatasmithCADImporter/Source/ParametricSurface/Public/CADModelConverter.h
@@ -40,7 +40,7 @@ public:
 	virtual void InitializeProcess() = 0;
 
 	virtual bool RepairTopology() = 0;
-	virtual bool SaveModel(const TCHAR* OutputPath, TSharedRef<IDatasmithMeshElement>& MeshElement) = 0;
+	virtual bool SaveModel(const TCHAR* OutputPath, TSharedPtr<IDatasmithMeshElement> MeshElement) = 0;
 	virtual bool Tessellate(const CADLibrary::FMeshParameters& InMeshParameters, FMeshDescription& OutMeshDescription) = 0;
 
 	/**
diff --git a/Engine/Plugins/Enterprise/DatasmithCADImporter/Source/ParametricSurface/Public/CADModelToTechSoftConverterBase.h b/Engine/Plugins/Enterprise/DatasmithCADImporter/Source/ParametricSurface/Public/CADModelToTechSoftConverterBase.h
--- a/Engine/Plugins/Enterprise/DatasmithCADImporter/Source/ParametricSurface/Public/CADModelToTechSoftConverterBase.h
+++ b/Engine/Plugins/Enterprise/DatasmithCADImporter/Source/ParametricSurface/Public/CADModelToTechSoftConverterBase.h
@@ -33,7 +33,7 @@ public:
 
 	virtual bool RepairTopology() override;
 
-	virtual bool SaveModel(const TCHAR* InFolderPath, TSharedRef<IDatasmithMeshElement>& MeshElement) override;
+	virtual bool SaveModel(const TCHAR* InFolderPath, TSharedPtr<IDatasmithMeshElement> MeshElement) override;
 
 	virtual bool Tessellate(const CADLibrary::FMeshParameters& InMeshParameters, FMeshDescription& OutMeshDescription) override;
 
@@ -54,3 +54,4 @@ protected:
 	TArray<A3DRiRepresentationItem*> RiRepresentationItems;
 	CADLibrary::FUniqueTechSoftModelFile ModelFile;
 };
+
diff --git a/Engine/Plugins/Enterprise/DatasmithCADImporter/Source/ParametricSurface/Public/ParametricSurfaceTranslator.h b/Engine/Plugins/Enterprise/DatasmithCADImporter/Source/ParametricSurface/Public/ParametricSurfaceTranslator.h
--- a/Engine/Plugins/Enterprise/DatasmithCADImporter/Source/ParametricSurface/Public/ParametricSurfaceTranslator.h
+++ b/Engine/Plugins/Enterprise/DatasmithCADImporter/Source/ParametricSurface/Public/ParametricSurfaceTranslator.h
@@ -4,9 +4,6 @@
 
 #include "CoreMinimal.h"
 
-#include "CADOptions.h"
-
-#include "DatasmithImportOptions.h"
 #include "DatasmithTranslator.h"
 #include "UObject/ObjectMacros.h"
 
@@ -22,11 +19,7 @@ class PARAMETRICSURFACE_API FParametricSurfaceTranslator : public IDatasmithTran
 {
 public:
 
-	FParametricSurfaceTranslator()
-	{
-		// Initialize bUseCADKernel with current value of CVar ds.CADTranslator.DisableCADKernelTessellation
-		CommonTessellationOptions.bUseCADKernel = !CADLibrary::FImportParameters::bGDisableCADKernelTessellation;
-	}
+	FParametricSurfaceTranslator();
 
 	// Begin IDatasmithTranslator overrides
 	virtual void GetSceneImportOptions(TArray<TObjectPtr<UDatasmithOptionsBase>>& Options) override;
@@ -46,7 +39,7 @@ protected:
 	 */
 	virtual void InitCommonTessellationOptions(FDatasmithTessellationOptions& TessellationOptions) {}
 
-private:
+protected:
 	FDatasmithTessellationOptions CommonTessellationOptions;
 };
 
diff --git a/Engine/Plugins/Enterprise/DatasmithCADImporter/Source/WireInterface/Private/AliasBrepConverter.h b/Engine/Plugins/Enterprise/DatasmithCADImporter/Source/WireInterface/Private/AliasBrepConverter.h
new file mode 100644
--- /dev/null
+++ b/Engine/Plugins/Enterprise/DatasmithCADImporter/Source/WireInterface/Private/AliasBrepConverter.h
@@ -0,0 +1,25 @@
+// Copyright Epic Games, Inc. All Rights Reserved.
+
+#pragma once
+
+class AlDagNode;
+struct FColor;
+
+namespace UE_DATASMITHWIRETRANSLATOR_NAMESPACE
+{
+
+// Defined the reference in which the object has to be defined
+enum class EAliasObjectReference
+{
+	LocalReference,  
+	ParentReference, 
+	WorldReference, 
+};
+
+class IAliasBRepConverter
+{
+public:
+	virtual bool AddBRep(AlDagNode& DagNode, const FColor& Color, EAliasObjectReference ObjectReference) = 0;
+};
+
+}
\ No newline at end of file
diff --git a/Engine/Plugins/Enterprise/DatasmithCADImporter/Source/WireInterface/Private/AliasModelToCADKernelConverter.cpp b/Engine/Plugins/Enterprise/DatasmithCADImporter/Source/WireInterface/Private/AliasModelToCADKernelConverter.cpp
new file mode 100644
--- /dev/null
+++ b/Engine/Plugins/Enterprise/DatasmithCADImporter/Source/WireInterface/Private/AliasModelToCADKernelConverter.cpp
@@ -0,0 +1,402 @@
+// Copyright Epic Games, Inc. All Rights Reserved.
+
+#include "AliasModelToCADKernelConverter.h"
+
+#include "OpenModelUtils.h"
+
+#include "CADData.h"
+#include "CADKernelTools.h"
+#include "HAL/ConsoleManager.h"
+#include "Hal/PlatformMemory.h"
+#include "MeshDescriptionHelper.h"
+
+#include "CADKernel/Core/Session.h"
+#include "CADKernel/Geo/Curves/NURBSCurveData.h"
+#include "CADKernel/Geo/GeoEnum.h"
+#include "CADKernel/Geo/Surfaces/NurbsSurfaceData.h"
+#include "CADKernel/Geo/Surfaces/Surface.h"
+
+#include "CADKernel/Math/Point.h"
+#include "CADKernel/Mesh/Structure/ModelMesh.h"
+
+#include "CADKernel/Topo/Body.h"
+#include "CADKernel/Topo/Model.h"
+#include "CADKernel/Topo/Shell.h"
+#include "CADKernel/Topo/TopologicalEdge.h"
+#include "CADKernel/Topo/TopologicalFace.h"
+#include "CADKernel/Topo/TopologicalLoop.h"
+
+#ifdef USE_OPENMODEL
+
+#if PLATFORM_WINDOWS
+#include "Windows/AllowWindowsPlatformTypes.h"
+#endif
+
+// Alias API wrappes object in AlObjects. This is an abstract base class which holds a reference to an anonymous data structure.
+// The only way to compare two AlObjects is to compare their data structure. That is the reason why private fields are made public. 
+#define private public
+#include "AlCurve.h"
+#include "AlDagNode.h"
+#include "AlShell.h"
+#include "AlShellNode.h"
+#include "AlSurface.h"
+#include "AlSurfaceNode.h"
+#include "AlTrimBoundary.h"
+#include "AlTrimCurve.h"
+#include "AlTrimRegion.h"
+#include "AlTM.h"
+
+#if PLATFORM_WINDOWS
+#include "Windows/HideWindowsPlatformTypes.h"
+#endif
+
+namespace UE_DATASMITHWIRETRANSLATOR_NAMESPACE
+{
+
+namespace AliasToCADKernelUtils
+{
+
+	template<typename Surface_T>
+	TSharedPtr<UE::CADKernel::FSurface> AddNURBSSurface(double GeometricTolerance, Surface_T& AliasSurface, EAliasObjectReference InObjectReference, const AlMatrix4x4& InAlMatrix)
+	{
+		UE::CADKernel::FNurbsSurfaceHomogeneousData NURBSData;
+		NURBSData.bSwapUV= true;
+		NURBSData.bIsRational = true;
+
+		NURBSData.PoleUCount = AliasSurface.uNumberOfCVsInclMultiples();
+		NURBSData.PoleVCount = AliasSurface.vNumberOfCVsInclMultiples();
+
+		NURBSData.UDegree = AliasSurface.uDegree();  // U order of the surface
+		NURBSData.VDegree = AliasSurface.vDegree();  // V order of the surface
+
+		int32 KnotSizeU = AliasSurface.realuNumberOfKnots() + 2;
+		int32 KnotSizeV = AliasSurface.realvNumberOfKnots() + 2;
+
+		NURBSData.UNodalVector.SetNumUninitialized(KnotSizeU);
+		NURBSData.VNodalVector.SetNumUninitialized(KnotSizeV);
+
+		AliasSurface.realuKnotVector(NURBSData.UNodalVector.GetData() + 1);
+		AliasSurface.realvKnotVector(NURBSData.VNodalVector.GetData() + 1);
+
+		NURBSData.UNodalVector[0] = NURBSData.UNodalVector[1];
+		NURBSData.UNodalVector[KnotSizeU - 1] = NURBSData.UNodalVector[KnotSizeU - 2];
+		NURBSData.VNodalVector[0] = NURBSData.VNodalVector[1];
+		NURBSData.VNodalVector[KnotSizeV - 1] = NURBSData.VNodalVector[KnotSizeV - 2];
+
+		const int32 CoordinateCount = NURBSData.PoleUCount * NURBSData.PoleVCount * 4;
+		NURBSData.HomogeneousPoles.SetNumUninitialized(CoordinateCount);
+
+		if (InObjectReference == EAliasObjectReference::WorldReference)
+		{
+			AliasSurface.CVsWorldPositionInclMultiples(NURBSData.HomogeneousPoles.GetData());
+		}
+		else if (InObjectReference == EAliasObjectReference::ParentReference)
+		{
+			AlTM TranformMatrix(InAlMatrix);
+			AliasSurface.CVsAffectedPositionInclMultiples(TranformMatrix, NURBSData.HomogeneousPoles.GetData());
+		}
+		else  // EAliasObjectReference::LocalReference
+		{
+			AliasSurface.CVsUnaffectedPositionInclMultiples(NURBSData.HomogeneousPoles.GetData());
+		}
+
+		// convert cm to mm
+		double* Poles = NURBSData.HomogeneousPoles.GetData();
+		for (int32 Index = 0; Index < NURBSData.HomogeneousPoles.Num(); Index +=4)
+		{
+			Poles[Index + 0] *= 10.;
+			Poles[Index + 1] *= 10.;
+			Poles[Index + 2] *= 10.;
+		}
+
+		return UE::CADKernel::FSurface::MakeNurbsSurface(GeometricTolerance, NURBSData);
+	}
+}
+
+TSharedPtr<UE::CADKernel::FTopologicalEdge> FAliasModelToCADKernelConverter::AddEdge(const AlTrimCurve& AliasTrimCurve, TSharedPtr<UE::CADKernel::FSurface>& CarrierSurface)
+{
+	UE::CADKernel::FNurbsCurveData NurbsCurveData;
+
+	NurbsCurveData.Degree = AliasTrimCurve.degree();
+	int32 ControlPointCount = AliasTrimCurve.numberOfCVs();
+
+	NurbsCurveData.Dimension = 2;
+	NurbsCurveData.bIsRational = true;
+
+	int32 KnotCount = AliasTrimCurve.realNumberOfKnots() + 2;
+
+	NurbsCurveData.Weights.SetNumUninitialized(ControlPointCount);
+	NurbsCurveData.Poles.SetNumUninitialized(ControlPointCount);
+	NurbsCurveData.NodalVector.SetNumUninitialized(KnotCount);
+
+	using AlPoint = double[3];
+	// Notice that each CV has three coordinates - the three coordinates describe 2D parameter space, with a homogeneous coordinate.
+	// Each control point is u, v and w, where u and v are parameter space and w is the homogeneous coordinate.
+	AliasTrimCurve.CVsUVPosition(NurbsCurveData.NodalVector.GetData() + 1, (AlPoint*) NurbsCurveData.Poles.GetData());
+
+	AliasTrimCurve.realKnotVector(NurbsCurveData.NodalVector.GetData() + 1);
+	NurbsCurveData.NodalVector[0] = NurbsCurveData.NodalVector[1];
+	NurbsCurveData.NodalVector[KnotCount - 1] = NurbsCurveData.NodalVector[KnotCount - 2];
+
+	for (int32 Index = 0; Index < ControlPointCount; ++Index)
+	{
+		NurbsCurveData.Weights[Index] = NurbsCurveData.Poles[Index].Z;
+		NurbsCurveData.Poles[Index].Z = 0;
+	}
+
+	TSharedPtr<UE::CADKernel::FCurve> Nurbs = UE::CADKernel::FCurve::MakeNurbsCurve(NurbsCurveData);
+
+	TSharedRef<UE::CADKernel::FRestrictionCurve> RestrictionCurve = UE::CADKernel::FCurve::MakeShared<UE::CADKernel::FRestrictionCurve>(CarrierSurface.ToSharedRef(), Nurbs.ToSharedRef());
+	TSharedPtr<UE::CADKernel::FTopologicalEdge> Edge = UE::CADKernel::FTopologicalEdge::Make(RestrictionCurve);
+	if (!Edge.IsValid())
+	{
+		return TSharedPtr<UE::CADKernel::FTopologicalEdge>();
+	}
+
+	// Only TrimCurve with twin need to be in the map used in LinkEdgesLoop 
+	TAlObjectPtr<AlTrimCurve> TwinCurve(AliasTrimCurve.getTwinCurve());
+	if (TwinCurve.IsValid())
+	{
+		AlEdge2CADKernelEdge.Add(AliasTrimCurve.fSpline, Edge);
+	}
+
+	return Edge;
+}
+
+TSharedPtr<UE::CADKernel::FTopologicalLoop> FAliasModelToCADKernelConverter::AddLoop(const AlTrimBoundary& TrimBoundary, TSharedPtr<UE::CADKernel::FSurface>& CarrierSurface, const bool bIsExternal)
+{
+	TArray<TSharedPtr<UE::CADKernel::FTopologicalEdge>> Edges;
+	TArray<UE::CADKernel::EOrientation> Directions;
+
+	for (TAlObjectPtr<AlTrimCurve> TrimCurve(TrimBoundary.firstCurve()); TrimCurve.IsValid(); TrimCurve = TAlObjectPtr<AlTrimCurve>(TrimCurve->nextCurve()))
+	{
+		TSharedPtr<UE::CADKernel::FTopologicalEdge> Edge = AddEdge(*TrimCurve, CarrierSurface);
+		if (Edge.IsValid())
+		{
+			Edges.Add(Edge);
+			Directions.Emplace(UE::CADKernel::EOrientation::Front);
+		}
+	}
+
+	if (Edges.Num() == 0)
+	{
+		return TSharedPtr<UE::CADKernel::FTopologicalLoop>();
+	}
+
+	TSharedPtr<UE::CADKernel::FTopologicalLoop> Loop = UE::CADKernel::FTopologicalLoop::Make(Edges, Directions, bIsExternal, GeometricTolerance);
+	return Loop;
+}
+
+void FAliasModelToCADKernelConverter::LinkEdgesLoop(const AlTrimBoundary& TrimBoundary, UE::CADKernel::FTopologicalLoop& Loop)
+{
+	for (TAlObjectPtr<AlTrimCurve> TrimCurve(TrimBoundary.firstCurve()); TrimCurve.IsValid(); TrimCurve = TAlObjectPtr<AlTrimCurve>(TrimCurve->nextCurve()))
+	{
+		TSharedPtr<UE::CADKernel::FTopologicalEdge>* Edge = AlEdge2CADKernelEdge.Find(TrimCurve->fSpline);
+		if (!Edge || !Edge->IsValid() || (*Edge)->IsDeleted() || (*Edge)->IsDegenerated())
+		{
+			continue;
+		}
+
+		ensure(&Loop == (*Edge)->GetLoop());
+
+		// Link edges
+		TAlObjectPtr<AlTrimCurve> TwinCurve(TrimCurve->getTwinCurve());
+		if (TwinCurve.IsValid())
+		{
+			if (TSharedPtr<UE::CADKernel::FTopologicalEdge>* TwinEdge = AlEdge2CADKernelEdge.Find(TwinCurve->fSpline))
+			{
+				if (TwinEdge->IsValid() && !(*TwinEdge)->IsDeleted() && !(*TwinEdge)->IsDegenerated())
+				{
+					(*Edge)->LinkIfCoincident(**TwinEdge, EdgeLengthTolerance, SquareTolerance);
+				}
+			}
+		}
+	}
+}
+
+TSharedPtr<UE::CADKernel::FTopologicalFace> FAliasModelToCADKernelConverter::AddTrimRegion(const AlTrimRegion& TrimRegion, EAliasObjectReference InObjectReference, const AlMatrix4x4& InAlMatrix, bool bInOrientation)
+{
+	TSharedPtr<UE::CADKernel::FSurface> Surface = AliasToCADKernelUtils::AddNURBSSurface(GeometricTolerance, TrimRegion, InObjectReference, InAlMatrix);
+	if (!Surface.IsValid())
+	{
+		return TSharedPtr<UE::CADKernel::FTopologicalFace>();
+	}
+
+	bool bIsExternal = true;
+	TArray<TSharedPtr<UE::CADKernel::FTopologicalLoop>> Loops;
+	for (TAlObjectPtr<AlTrimBoundary> TrimBoundary(TrimRegion.firstBoundary()); TrimBoundary.IsValid(); TrimBoundary = TAlObjectPtr<AlTrimBoundary>(TrimBoundary->nextBoundary()))
+	{
+		TSharedPtr<UE::CADKernel::FTopologicalLoop> Loop = AddLoop(*TrimBoundary, Surface, bIsExternal);
+		if (Loop.IsValid())
+		{
+			LinkEdgesLoop(*TrimBoundary, *Loop);
+			Loops.Add(Loop);
+			bIsExternal = false;
+		}
+	}
+
+	if (Loops.Num() == 0)
+	{
+		UE::CADKernel::FMessage::Printf(UE::CADKernel::EVerboseLevel::Log, TEXT("The Face %s is degenerate, this face is ignored\n"), TrimRegion.name());
+		return TSharedPtr<UE::CADKernel::FTopologicalFace>();
+	}
+
+	TSharedRef<UE::CADKernel::FTopologicalFace> Face = UE::CADKernel::FEntity::MakeShared<UE::CADKernel::FTopologicalFace>(Surface);
+	Face->SetPatchId(LastFaceId++);
+
+	int32 DoubtfulLoopOrientationCount = 0;
+	Face->AddLoops(Loops, DoubtfulLoopOrientationCount);
+
+	if (Face->GetLoops().Num() == 0)
+	{
+		Face->SetAsDegenerated();
+		Face->Delete();
+
+		UE::CADKernel::FMessage::Printf(UE::CADKernel::EVerboseLevel::Log, TEXT("The Face %s is degenerate, this face is ignored\n"), TrimRegion.name());
+		return TSharedPtr<UE::CADKernel::FTopologicalFace>();
+	}
+
+	return Face;
+}
+
+void FAliasModelToCADKernelConverter::AddFace(const AlSurface& Surface, EAliasObjectReference InObjectReference, const AlMatrix4x4& InAlMatrix, bool bInOrientation, TSharedRef<UE::CADKernel::FShell>& Shell)
+{
+	TAlObjectPtr<AlTrimRegion> TrimRegion(Surface.firstTrimRegion());
+	if (TrimRegion.IsValid())
+	{
+		for (; TrimRegion.IsValid(); TrimRegion = TAlObjectPtr<AlTrimRegion>(TrimRegion->nextRegion()))
+		{
+			TSharedPtr<UE::CADKernel::FTopologicalFace> Face = AddTrimRegion(*TrimRegion, InObjectReference, InAlMatrix, bInOrientation);
+			if (Face.IsValid())
+			{
+				Shell->Add(Face.ToSharedRef(), bInOrientation ? UE::CADKernel::EOrientation::Front : UE::CADKernel::EOrientation::Back);
+			}
+		}
+		return;
+	}
+
+	TSharedPtr<UE::CADKernel::FSurface> CADKernelSurface = AliasToCADKernelUtils::AddNURBSSurface(GeometricTolerance, Surface, InObjectReference, InAlMatrix);
+	if (CADKernelSurface.IsValid())
+	{
+		TSharedRef<UE::CADKernel::FTopologicalFace> Face = UE::CADKernel::FEntity::MakeShared<UE::CADKernel::FTopologicalFace>(CADKernelSurface);
+		Face->ApplyNaturalLoops();
+		Shell->Add(Face, bInOrientation ? UE::CADKernel::EOrientation::Front : UE::CADKernel::EOrientation::Back);
+	}
+}
+
+void FAliasModelToCADKernelConverter::AddShell(const AlShell& InShell, EAliasObjectReference InObjectReference, const AlMatrix4x4& InAlMatrix, bool bInOrientation, TSharedRef<UE::CADKernel::FShell>& CADKernelShell)
+{
+	for(TAlObjectPtr<AlTrimRegion> TrimRegion(InShell.firstTrimRegion()); TrimRegion.IsValid(); TrimRegion = TAlObjectPtr<AlTrimRegion>(TrimRegion->nextRegion()))
+	{
+		TSharedPtr<UE::CADKernel::FTopologicalFace> Face = AddTrimRegion(*TrimRegion, InObjectReference, InAlMatrix, bInOrientation);
+		if (Face.IsValid())
+		{
+			CADKernelShell->Add(Face.ToSharedRef(), bInOrientation ? UE::CADKernel::EOrientation::Front : UE::CADKernel::EOrientation::Back);
+		}
+	}
+}
+
+bool FAliasModelToCADKernelConverter::AddBRep(AlDagNode& DagNode, const FColor& Color, EAliasObjectReference InObjectReference)
+{
+	AlEdge2CADKernelEdge.Empty();
+
+	TSharedRef<UE::CADKernel::FBody> CADKernelBody = UE::CADKernel::FEntity::MakeShared<UE::CADKernel::FBody>();
+	uint32 ColorId = (uint32) CADLibrary::BuildColorUId(Color);
+	CADKernelBody->SetColorId(ColorId);
+
+	TSharedRef<UE::CADKernel::FShell> CADKernelShell = UE::CADKernel::FEntity::MakeShared<UE::CADKernel::FShell>();
+	CADKernelBody->AddShell(CADKernelShell);
+
+	boolean bAlOrientation;
+	DagNode.getSurfaceOrientation(bAlOrientation);
+	bool bOrientation = !(bool)bAlOrientation;
+
+	AlMatrix4x4 AlMatrix;
+	if (InObjectReference == EAliasObjectReference::ParentReference)
+	{
+		DagNode.localTransformationMatrix(AlMatrix);
+	}
+
+	AlObjectType objectType = DagNode.type();
+	switch (objectType)
+	{
+		// Push all leaf nodes into 'leaves'
+	case kShellNodeType:
+	{
+		if (AlShellNode* ShellNode = DagNode.asShellNodePtr())
+		{
+			AlShell* ShellPtr = ShellNode->shell();
+			if(AlIsValid(ShellPtr))
+			{
+				TAlObjectPtr<AlShell> AliasShell(ShellPtr);
+				AddShell(*ShellPtr, InObjectReference, AlMatrix, bOrientation, CADKernelShell);
+			}
+		}
+		break;
+	}
+
+	case kSurfaceNodeType:
+	{
+		if (AlSurfaceNode* SurfaceNode = DagNode.asSurfaceNodePtr())
+		{
+			AlSurface* SurfacePtr = SurfaceNode->surface();
+			if (AlIsValid(SurfacePtr))
+			{
+				TAlObjectPtr<AlSurface> AliasSurface(SurfacePtr);
+				AddFace(*SurfacePtr, InObjectReference, AlMatrix, bOrientation, CADKernelShell);
+			}
+		}
+		break;
+	}
+	default:
+		break;
+	}
+
+	CADKernelBody->CompleteMetaData();
+	CADKernelSession.GetModel().Add(CADKernelBody);
+
+	return CADKernelShell->FaceCount() > 0;
+}
+
+bool FAliasModelToCADKernelConverter::Tessellate(const CADLibrary::FMeshParameters& InMeshParameters, FMeshDescription& OutMeshDescription)
+{
+	UE::CADKernel::FModel& Model = CADKernelSession.GetModel();
+
+	CADLibrary::FMeshConversionContext Context(ImportParameters, InMeshParameters);
+
+	const IConsoleVariable* ConsoleVariable = IConsoleManager::Get().FindConsoleVariable(TEXT("ds.WireTranslator.SkipThinZoneMeshing"));
+	Context.bSkipThinZoneMeshing = ConsoleVariable ? ConsoleVariable->GetBool() : false;
+
+	return CADLibrary::FCADKernelTools::Tessellate(Model, Context, OutMeshDescription);
+}
+
+bool FAliasModelToCADKernelConverter::RepairTopology()
+{
+	using namespace CADLibrary;
+	// Apply stitching if applicable
+	if (ImportParameters.GetStitchingTechnique() != StitchingNone)
+	{
+		const double StitchingTolerance = FImportParameters::GStitchingTolerance * 10.; //CM to MM
+		UE::CADKernel::ESewOption SewOptionValue = (UE::CADKernel::ESewOption)SewOption::GetFromImportParameters();
+
+		const IConsoleVariable* ConsoleVariable = IConsoleManager::Get().FindConsoleVariable(TEXT("ds.WireTranslator.SkipThinZoneMeshing"));
+		bool bSkipThinZoneMeshing = ConsoleVariable ? ConsoleVariable->GetBool() : false;
+		if (bSkipThinZoneMeshing)
+		{
+			SewOptionValue = UE::CADKernel::ESewOption((uint8)SewOptionValue & (uint8)~ESewOption::RemoveThinFaces);
+		}
+
+		UE::CADKernel::FTopomakerOptions TopomakerOptions(SewOptionValue, StitchingTolerance, FImportParameters::GStitchingForceFactor);
+
+		UE::CADKernel::FTopomaker Topomaker(CADKernelSession, TopomakerOptions);
+		Topomaker.Sew();
+		Topomaker.SplitIntoConnectedShells();
+		Topomaker.OrientShells();
+	}
+
+	return true;
+}
+
+}
+
+#endif
\ No newline at end of file
diff --git a/Engine/Plugins/Enterprise/DatasmithCADImporter/Source/WireInterface/Private/AliasModelToCADKernelConverter.h b/Engine/Plugins/Enterprise/DatasmithCADImporter/Source/WireInterface/Private/AliasModelToCADKernelConverter.h
new file mode 100644
--- /dev/null
+++ b/Engine/Plugins/Enterprise/DatasmithCADImporter/Source/WireInterface/Private/AliasModelToCADKernelConverter.h
@@ -0,0 +1,70 @@
+// Copyright Epic Games, Inc. All Rights Reserved.
+
+#pragma once
+
+#include "CoreMinimal.h"
+
+#include "OpenModelUtils.h"
+#include "AliasBRepConverter.h"
+
+#include "CADModelToCADKernelConverterBase.h"
+
+#include "CADKernel/Core/Session.h"
+#include "CADKernel/Core/Types.h"
+
+class AlDagNode;
+class AlShell;
+class AlSurface;
+class AlTrimBoundary;
+class AlTrimCurve;
+class AlTrimRegion;
+
+namespace UE::CADKernel
+{
+	class FShell;
+	class FSurface;
+	class FTopologicalEdge;
+	class FTopologicalFace;
+	class FTopologicalLoop;
+}
+
+namespace UE_DATASMITHWIRETRANSLATOR_NAMESPACE
+{
+
+class FAliasModelToCADKernelConverter : public FCADModelToCADKernelConverterBase, public IAliasBRepConverter
+{
+
+public:
+	FAliasModelToCADKernelConverter(CADLibrary::FImportParameters InImportParameters)
+		: FCADModelToCADKernelConverterBase(InImportParameters)
+	{
+	}
+
+	// Begin FCADModelToCADKernelConverterBase overrides
+	virtual bool Tessellate(const CADLibrary::FMeshParameters& InMeshParameters, FMeshDescription& OutMeshDescription) override;
+	virtual bool RepairTopology() override;
+	// End FCADModelToCADKernelConverterBase overrides
+
+	virtual bool AddBRep(AlDagNode& DagNode, const FColor& Color, EAliasObjectReference ObjectReference) override;
+
+protected:
+	TSharedPtr<UE::CADKernel::FTopologicalEdge> AddEdge(const AlTrimCurve& TrimCurve, TSharedPtr<UE::CADKernel::FSurface>& CarrierSurface);
+
+	TSharedPtr<UE::CADKernel::FTopologicalLoop> AddLoop(const AlTrimBoundary& TrimBoundary, TSharedPtr<UE::CADKernel::FSurface>& CarrierSurface, const bool bIsExternal);
+	
+	/**
+	 * Build face's links with its neighbor have to be done after the loop is finalize.
+	 * This is to avoid to link an edge with another and then to delete it...
+	 */
+	void LinkEdgesLoop(const AlTrimBoundary& TrimBoundary, UE::CADKernel::FTopologicalLoop& Loop);
+
+	TSharedPtr<UE::CADKernel::FTopologicalFace> AddTrimRegion(const AlTrimRegion& InTrimRegion, EAliasObjectReference InObjectReference, const AlMatrix4x4& InAlMatrix, bool bInOrientation);
+	void AddFace(const AlSurface& InSurface, EAliasObjectReference InObjectReference, const AlMatrix4x4& InAlMatrix, bool bInOrientation, TSharedRef<UE::CADKernel::FShell>& Shell);
+	void AddShell(const AlShell& InShell, EAliasObjectReference InObjectReference, const AlMatrix4x4& InAlMatrix, bool bInOrientation, TSharedRef<UE::CADKernel::FShell>& Shell);
+
+protected:
+	int32 LastFaceId = 1;
+	TMap<void*, TSharedPtr<UE::CADKernel::FTopologicalEdge>>  AlEdge2CADKernelEdge;
+};
+
+}
\ No newline at end of file
diff --git a/Engine/Plugins/Enterprise/DatasmithCADImporter/Source/WireInterface/Private/AliasModelToTechSoftConverter.cpp b/Engine/Plugins/Enterprise/DatasmithCADImporter/Source/WireInterface/Private/AliasModelToTechSoftConverter.cpp
new file mode 100644
--- /dev/null
+++ b/Engine/Plugins/Enterprise/DatasmithCADImporter/Source/WireInterface/Private/AliasModelToTechSoftConverter.cpp
@@ -0,0 +1,475 @@
+// Copyright Epic Games, Inc. All Rights Reserved.
+
+#include "AliasModelToTechSoftConverter.h"
+
+#include "OpenModelUtils.h"
+
+#include "Hal/PlatformMemory.h"
+#include "Math/Color.h"
+#include "MeshDescriptionHelper.h"
+#include "TechSoftInterface.h"
+#include "TechSoftUtils.h"
+#include "TUniqueTechSoftObj.h"
+
+
+#ifdef USE_OPENMODEL
+
+#if PLATFORM_WINDOWS
+#include "Windows/AllowWindowsPlatformTypes.h"
+#endif
+
+// Alias API wrappes object in AlObjects. This is an abstract base class which holds a reference to an anonymous data structure.
+// The only way to compare two AlObjects is to compare their data structure. That is the reason why private fields are made public. 
+#define private public
+#include "AlCurve.h"
+#include "AlDagNode.h"
+#include "AlShell.h"
+#include "AlShellNode.h"
+#include "AlSurface.h"
+#include "AlSurfaceNode.h"
+#include "AlTrimBoundary.h"
+#include "AlTrimCurve.h"
+#include "AlTrimRegion.h"
+#include "AlTM.h"
+
+#if PLATFORM_WINDOWS
+#include "Windows/HideWindowsPlatformTypes.h"
+#endif
+
+namespace UE_DATASMITHWIRETRANSLATOR_NAMESPACE
+{
+
+#ifdef USE_TECHSOFT_SDK
+
+namespace AliasToTechSoftUtils
+{
+
+enum EAxis { U, V };
+
+template<typename Surface_T>
+A3DSurfBase* AddNURBSSurface(const Surface_T& AliasSurface, EAliasObjectReference InObjectReference, const AlMatrix4x4& InAlMatrix)
+{
+	CADLibrary::TUniqueTSObj<A3DSurfNurbsData> NurbsSurfaceData;
+
+	NurbsSurfaceData->m_eKnotType = A3DEKnotType::kA3DKnotTypeUnspecified;
+	NurbsSurfaceData->m_eSurfaceForm = A3DEBSplineSurfaceForm::kA3DBSplineSurfaceFormUnspecified;
+
+	NurbsSurfaceData->m_uiUDegree = AliasSurface.uDegree();
+	NurbsSurfaceData->m_uiVDegree = AliasSurface.vDegree();
+
+	NurbsSurfaceData->m_uiUCtrlSize = AliasSurface.uNumberOfCVsInclMultiples();
+	NurbsSurfaceData->m_uiVCtrlSize = AliasSurface.vNumberOfCVsInclMultiples();
+	NurbsSurfaceData->m_uiUKnotSize = AliasSurface.realuNumberOfKnots() + 2;
+	NurbsSurfaceData->m_uiVKnotSize = AliasSurface.realvNumberOfKnots() + 2;
+
+	TArray<A3DDouble> UNodalVector;
+	TArray<A3DDouble> VNodalVector;
+
+	TFunction<void(EAxis, TArray<A3DDouble>&)> FillNodalVector = [&](EAxis Axis, TArray<A3DDouble>& OutNodalVector)
+	{
+		TArray<double> NodalVector;
+		if(Axis == EAxis::U)
+		{
+			OutNodalVector.Reserve(NurbsSurfaceData->m_uiUKnotSize);
+			NodalVector.SetNumUninitialized(NurbsSurfaceData->m_uiUKnotSize - 2);
+			AliasSurface.realuKnotVector(NodalVector.GetData());
+		}
+		else
+		{
+			OutNodalVector.Reserve(NurbsSurfaceData->m_uiVKnotSize);
+			NodalVector.SetNumUninitialized(NurbsSurfaceData->m_uiVKnotSize - 2);
+			AliasSurface.realvKnotVector(NodalVector.GetData());
+		}
+
+		OutNodalVector.Add(NodalVector[0]);
+		for(double Value : NodalVector)
+		{
+			OutNodalVector.Add(Value);
+		}
+		OutNodalVector.Add(NodalVector.Last());
+
+		if (Axis == EAxis::U)
+		{
+			NurbsSurfaceData->m_pdUKnots = OutNodalVector.GetData();
+		}
+		else
+		{
+			NurbsSurfaceData->m_pdVKnots = OutNodalVector.GetData();
+		}
+	};
+
+	FillNodalVector(EAxis::U, UNodalVector);
+	FillNodalVector(EAxis::V, VNodalVector);
+
+	const int32 CoordinateCount = NurbsSurfaceData->m_uiUCtrlSize * NurbsSurfaceData->m_uiVCtrlSize * 4;
+	TArray<double> HomogeneousPoles;
+	HomogeneousPoles.SetNumUninitialized(CoordinateCount);
+
+	if (InObjectReference == EAliasObjectReference::WorldReference)
+	{
+		AliasSurface.CVsWorldPositionInclMultiples(HomogeneousPoles.GetData());
+	}
+	else if (InObjectReference == EAliasObjectReference::ParentReference)
+	{
+		AlTM TranformMatrix(InAlMatrix);
+		AliasSurface.CVsAffectedPositionInclMultiples(TranformMatrix, HomogeneousPoles.GetData());
+	}
+	else  // EAliasObjectReference::LocalReference
+	{
+		AliasSurface.CVsUnaffectedPositionInclMultiples(HomogeneousPoles.GetData());
+	}
+
+	int32 PoleCount = NurbsSurfaceData->m_uiUCtrlSize * NurbsSurfaceData->m_uiVCtrlSize;
+	TArray<A3DDouble> Weights;
+	TArray<A3DVector3dData> ControlPoints;
+	ControlPoints.SetNum(PoleCount);
+	Weights.SetNum(PoleCount);
+	A3DVector3dData* ControlPointPtr = ControlPoints.GetData();
+	A3DDouble* WeightPtr = Weights.GetData();
+	
+	NurbsSurfaceData->m_pdWeights = WeightPtr;
+	NurbsSurfaceData->m_pCtrlPts = ControlPointPtr;
+	
+	TFunction<void(const double*, A3DVector3dData&, A3DDouble&)> SetA3DPole = [](const double* HomogeneousPole, A3DVector3dData & OutTSPoint, A3DDouble & OutWeight)
+	{
+		OutTSPoint.m_dX = HomogeneousPole[0] * 10.;  // cm (Alias MetricUnit) to mm
+		OutTSPoint.m_dY = HomogeneousPole[1] * 10.;
+		OutTSPoint.m_dZ = HomogeneousPole[2] * 10.;
+		OutWeight = HomogeneousPole[3];
+	};
+
+	double* AliasPoles = HomogeneousPoles.GetData();
+	for (int32 Index = 0; Index < PoleCount; ++Index, AliasPoles += 4, ++ControlPointPtr, ++WeightPtr)
+	{
+		SetA3DPole(AliasPoles, *ControlPointPtr, *WeightPtr);
+	}
+
+	return CADLibrary::TechSoftInterface::CreateSurfaceNurbs(*NurbsSurfaceData);
+}
+
+} // ns AliasToTechsoftUtils
+
+A3DCrvBase* FAliasModelToTechSoftConverter::CreateCurve(const AlTrimCurve& AliasTrimCurve)
+{
+	CADLibrary::TUniqueTSObj<A3DCrvNurbsData> NurbsCurveData;
+
+	NurbsCurveData->m_bIs2D = true;
+	NurbsCurveData->m_bRational = true;
+	NurbsCurveData->m_uiDegree = AliasTrimCurve.degree();
+
+	int ControlPointCount = AliasTrimCurve.numberOfCVs();
+	int32 KnotCount = AliasTrimCurve.realNumberOfKnots();
+
+	using AlPoint = TStaticArray<double, 3>;
+
+	TArray<AlPoint> AliasPoles;
+	TArray<double> AliasNodalVector;
+	AliasNodalVector.SetNum(ControlPointCount);
+	AliasPoles.SetNumUninitialized(ControlPointCount*4);
+
+	// Notice that each CV has three coordinates - the three coordinates describe 2D parameter space, with a homogeneous coordinate.
+	// Each control point is u, v and w, where u and v are parameter space and w is the homogeneous coordinate.
+	AliasTrimCurve.CVsUVPosition(AliasNodalVector.GetData(), (double(*)[3])AliasPoles.GetData());
+
+	AliasNodalVector.SetNum(KnotCount);
+	AliasTrimCurve.realKnotVector(AliasNodalVector.GetData());
+
+	TArray<A3DDouble> NodalVector;
+	NodalVector.Reserve(KnotCount+2);
+	NodalVector.Add(AliasNodalVector[0]);
+	for (double Value : AliasNodalVector)
+	{
+		NodalVector.Add(Value);
+	}
+	NodalVector.Add(AliasNodalVector.Last());
+
+	TArray<double> Weights;
+	Weights.SetNumUninitialized(ControlPointCount);
+
+	AliasNodalVector.SetNumUninitialized(KnotCount);
+
+	TArray<A3DVector3dData> ControlPointArray;
+	TArray<A3DDouble> WeightArray;
+
+	ControlPointArray.SetNumUninitialized(ControlPointCount);
+	WeightArray.SetNumUninitialized(ControlPointCount);
+
+	A3DVector3dData* ControlPointPtr = ControlPointArray.GetData();
+	A3DDouble* WeightPtr = WeightArray.GetData();
+	AlPoint* AliasPolePtr = AliasPoles.GetData();
+
+	TFunction<void(const AlPoint&, A3DVector3dData&, A3DDouble&)> SetA3DPole = [](const AlPoint& AliasPole, A3DVector3dData& OutTSPoint, A3DDouble& OutWeight)
+	{
+		OutTSPoint.m_dX = AliasPole[0];
+		OutTSPoint.m_dY = AliasPole[1];
+		OutTSPoint.m_dZ = 0;
+		OutWeight = AliasPole[2];
+	};
+
+	for (int32 Index = 0; Index < ControlPointCount; ++Index, ++AliasPolePtr, ++ControlPointPtr, ++WeightPtr)
+	{
+		SetA3DPole(*AliasPolePtr, *ControlPointPtr, *WeightPtr);
+	}
+
+	NurbsCurveData->m_eKnotType = kA3DKnotTypeUnspecified;
+	NurbsCurveData->m_eCurveForm = kA3DBSplineCurveFormUnspecified;
+
+	NurbsCurveData->m_pCtrlPts = ControlPointArray.GetData();
+	NurbsCurveData->m_uiCtrlSize = ControlPointArray.Num();
+
+	NurbsCurveData->m_pdWeights = WeightArray.GetData();
+	NurbsCurveData->m_uiWeightSize = WeightArray.Num();
+
+	NurbsCurveData->m_pdKnots = NodalVector.GetData();
+	NurbsCurveData->m_uiKnotSize = NodalVector.Num();
+
+	return CADLibrary::TechSoftInterface::CreateCurveNurbs(*NurbsCurveData);
+}
+
+A3DTopoCoEdge* FAliasModelToTechSoftConverter::CreateEdge(const AlTrimCurve& TrimCurve)
+{
+	A3DCrvBase* NurbsCurvePtr = CreateCurve(TrimCurve);
+	if (NurbsCurvePtr == nullptr)
+	{
+		return nullptr;
+	}
+
+	A3DTopoEdge* EdgePtr = CADLibrary::TechSoftUtils::CreateTopoEdge();
+	if (EdgePtr == nullptr)
+	{
+		return nullptr;
+	}
+
+	CADLibrary::TUniqueTSObj<A3DTopoCoEdgeData> CoEdgeData;
+
+	CoEdgeData->m_pUVCurve = NurbsCurvePtr;
+	CoEdgeData->m_pEdge = EdgePtr;
+	CoEdgeData->m_ucOrientationWithLoop = TrimCurve.isReversed();
+	CoEdgeData->m_ucOrientationUVWithLoop = 1;
+
+	A3DTopoCoEdge* CoEdgePtr = CADLibrary::TechSoftInterface::CreateTopoCoEdge(*CoEdgeData);
+
+	// Only TrimCurve with twin need to be in the map used in LinkEdgesLoop 
+	TAlObjectPtr<AlTrimCurve> TwinCurve(TrimCurve.getTwinCurve());
+	if (TwinCurve.IsValid())
+	{
+		AlEdgeToTSCoEdge.Add(TrimCurve.fSpline, CoEdgePtr);
+	}
+	return CoEdgePtr;
+}
+
+A3DTopoLoop* FAliasModelToTechSoftConverter::CreateTopoLoop(const AlTrimBoundary& TrimBoundary)
+{
+	TArray<A3DTopoCoEdge*> Edges;
+	Edges.Reserve(20);
+
+	for (TAlObjectPtr<AlTrimCurve> TrimCurve(TrimBoundary.firstCurve()); TrimCurve.IsValid(); TrimCurve = TAlObjectPtr<AlTrimCurve>(TrimCurve->nextCurve()))
+	{
+		A3DTopoCoEdge* Edge = CreateEdge(*TrimCurve);
+		if (Edge != nullptr)
+		{
+			Edges.Add(Edge);
+		}
+	}
+
+	if (Edges.Num() == 0)
+	{
+		return nullptr;
+	}
+
+	CADLibrary::TUniqueTSObj<A3DTopoLoopData> LoopData;
+
+	LoopData->m_ppCoEdges = Edges.GetData();
+	LoopData->m_uiCoEdgeSize = Edges.Num();
+	LoopData->m_ucOrientationWithSurface = 1;
+
+	return CADLibrary::TechSoftInterface::CreateTopoLoop(*LoopData);
+}
+
+void FAliasModelToTechSoftConverter::LinkEdgesLoop(const AlTrimBoundary& TrimBoundary)
+{
+	for (TAlObjectPtr<AlTrimCurve> TrimCurve(TrimBoundary.firstCurve()); TrimCurve.IsValid(); TrimCurve = TAlObjectPtr<AlTrimCurve>(TrimCurve->nextCurve()))
+	{
+		A3DTopoCoEdge** Edge = AlEdgeToTSCoEdge.Find(TrimCurve->fSpline);
+		if (Edge == nullptr)
+		{
+			continue;
+		}
+
+		// Link edges
+		TAlObjectPtr<AlTrimCurve> TwinCurve(TrimCurve->getTwinCurve());
+		if (TwinCurve.IsValid())
+		{
+			if (A3DTopoCoEdge** TwinEdge = AlEdgeToTSCoEdge.Find(TwinCurve->fSpline))
+			{
+				if (TwinEdge != nullptr)
+				{
+					CADLibrary::TechSoftInterface::LinkCoEdges(*Edge, *TwinEdge);
+					break;
+				}
+			}
+		}
+	}
+}
+
+A3DTopoFace* FAliasModelToTechSoftConverter::AddTrimRegion(const AlTrimRegion& InTrimRegion, const FColor& Color, EAliasObjectReference InObjectReference, const AlMatrix4x4& InAlMatrix)
+{
+	A3DSurfBase* CarrierSurface = AliasToTechSoftUtils::AddNURBSSurface(InTrimRegion, InObjectReference, InAlMatrix);
+	if (CarrierSurface == nullptr)
+	{
+		return nullptr;
+	}
+
+	TArray<A3DTopoLoop*> Loops;
+	Loops.Reserve(5);
+
+	for (TAlObjectPtr<AlTrimBoundary> TrimBoundary(InTrimRegion.firstBoundary()); TrimBoundary.IsValid(); TrimBoundary = TAlObjectPtr<AlTrimBoundary>(TrimBoundary->nextBoundary()))
+	{
+		A3DTopoLoop* Loop = CreateTopoLoop(*TrimBoundary);
+		if (Loop != nullptr)
+		{
+			Loops.Add(Loop);
+			LinkEdgesLoop(*TrimBoundary);
+		}
+	}
+
+	if (Loops.Num() == 0)
+	{
+		return nullptr;
+	}
+
+	CADLibrary::TUniqueTSObj<A3DTopoFaceData> Face;
+	Face->m_pSurface = CarrierSurface;
+	Face->m_bHasTrimDomain = false;
+	Face->m_ppLoops = Loops.GetData();
+	Face->m_uiLoopSize = Loops.Num();
+	Face->m_uiOuterLoopIndex = 0;
+	Face->m_dTolerance = 0.01; //mm
+
+	A3DTopoFace* FacePtr = CADLibrary::TechSoftInterface::CreateTopoFace(*Face);
+
+	CADLibrary::TechSoftUtils::SetEntityGraphicsColor(FacePtr, Color);
+
+	return FacePtr;
+}
+#endif
+
+bool FAliasModelToTechSoftConverter::AddBRep(AlDagNode& DagNode, const FColor& Color, EAliasObjectReference InObjectReference)
+{
+#ifdef USE_TECHSOFT_SDK
+	AlEdgeToTSCoEdge.Empty();
+
+	AlMatrix4x4 AlMatrix;
+	if (InObjectReference == EAliasObjectReference::ParentReference)
+	{
+		DagNode.localTransformationMatrix(AlMatrix);
+	}
+
+	boolean bAlOrientation;
+	DagNode.getSurfaceOrientation(bAlOrientation);
+	bool bOrientation = !(bool)bAlOrientation;
+
+	TArray<A3DTopoFace*> TSFaces;
+	TSFaces.Reserve(100);
+
+	AlObjectType objectType = DagNode.type();
+	switch (objectType)
+	{
+
+	case kShellNodeType:
+	{
+		if (AlShellNode* ShellNode = DagNode.asShellNodePtr())
+		{
+			TAlObjectPtr<AlShell> AliasShell(ShellNode->shell());
+			if(AlIsValid(AliasShell.Get()))
+			{
+				for (TAlObjectPtr<AlTrimRegion> TrimRegion(AliasShell->firstTrimRegion()); TrimRegion.IsValid(); TrimRegion = TAlObjectPtr<AlTrimRegion>(TrimRegion->nextRegion()))
+				{
+					A3DTopoFace* TSFace = AddTrimRegion(*TrimRegion, Color, InObjectReference, AlMatrix);
+					if (TSFace != nullptr)
+					{
+						TSFaces.Add(TSFace);
+					}
+				}
+			}
+		}
+		break;
+	}
+
+	case kSurfaceNodeType:
+	{
+		if (AlSurfaceNode* SurfaceNode = DagNode.asSurfaceNodePtr())
+		{
+			TAlObjectPtr<AlSurface> AliasSurface(SurfaceNode->surface());
+			if (AlIsValid(AliasSurface.Get()))
+			{
+				TAlObjectPtr<AlTrimRegion> TrimRegion(AliasSurface->firstTrimRegion());
+				if (TrimRegion.IsValid())
+				{
+					for (; TrimRegion.IsValid(); TrimRegion = TAlObjectPtr<AlTrimRegion>(TrimRegion->nextRegion()))
+					{
+						A3DTopoFace* TSFace = AddTrimRegion(*TrimRegion, Color, InObjectReference, AlMatrix);
+						if (TSFace != nullptr)
+						{
+							TSFaces.Add(TSFace);
+						}
+					}
+					break;
+				}
+
+				A3DSurfBase* TSSurface = AliasToTechSoftUtils::AddNURBSSurface(*AliasSurface, InObjectReference, AlMatrix);
+				if (TSSurface != nullptr)
+				{
+					A3DTopoFace* TSFace = CADLibrary::TechSoftUtils::CreateTopoFaceWithNaturalLoop(TSSurface);
+					CADLibrary::TechSoftUtils::SetEntityGraphicsColor(TSFace, Color);
+
+					if (TSFace != nullptr)
+					{
+						TSFaces.Add(TSFace);
+					}
+				}
+			}
+		}
+		break;
+	}
+	default:
+		break;
+	}
+
+	if (TSFaces.IsEmpty())
+	{
+		return false;
+	}
+
+	A3DTopoShell* TopoShellPtr = nullptr;
+	{
+		TArray<A3DUns8> FaceOrientations;
+		FaceOrientations.Init(bOrientation, TSFaces.Num());
+
+		CADLibrary::TUniqueTSObj<A3DTopoShellData> TopoShellData;
+		TopoShellData->m_bClosed = false;
+		TopoShellData->m_ppFaces = TSFaces.GetData();
+		TopoShellData->m_uiFaceSize = TSFaces.Num();
+		TopoShellData->m_pucOrientationWithShell = FaceOrientations.GetData();
+
+		TopoShellPtr = CADLibrary::TechSoftInterface::CreateTopoShell(*TopoShellData);
+
+		if (TopoShellPtr == nullptr)
+		{
+			return false;
+		}
+	}
+
+	A3DRiRepresentationItem* RiRepresentationItem = CADLibrary::TechSoftUtils::CreateRIBRep(TopoShellPtr);
+
+	if (RiRepresentationItem != nullptr)
+	{
+		RiRepresentationItems.Add(RiRepresentationItem);
+		return true;
+	}
+#endif
+	return false;
+}
+
+}
+#endif
diff --git a/Engine/Plugins/Enterprise/DatasmithCADImporter/Source/WireInterface/Private/AliasModelToTechSoftConverter.h b/Engine/Plugins/Enterprise/DatasmithCADImporter/Source/WireInterface/Private/AliasModelToTechSoftConverter.h
new file mode 100644
--- /dev/null
+++ b/Engine/Plugins/Enterprise/DatasmithCADImporter/Source/WireInterface/Private/AliasModelToTechSoftConverter.h
@@ -0,0 +1,54 @@
+// Copyright Epic Games, Inc. All Rights Reserved.
+
+#pragma once
+
+#include "CoreMinimal.h"
+
+#include "OpenModelUtils.h"
+#include "AliasBRepConverter.h"
+
+#include "CADModelToTechSoftConverterBase.h"
+
+#include "CADKernel/Core/Session.h"
+#include "CADKernel/Core/Types.h"
+
+class AlDagNode;
+class AlShell;
+class AlSurface;
+class AlTrimBoundary;
+class AlTrimCurve;
+class AlTrimRegion;
+
+typedef void A3DTopoCoEdge;
+typedef void A3DTopoFace;
+typedef void A3DTopoLoop;
+typedef void A3DTopoCoEdge;
+typedef void A3DCrvBase;
+
+namespace UE_DATASMITHWIRETRANSLATOR_NAMESPACE
+{
+
+class FAliasModelToTechSoftConverter : public FCADModelToTechSoftConverterBase, public IAliasBRepConverter
+{
+
+public:
+	FAliasModelToTechSoftConverter(CADLibrary::FImportParameters InImportParameters)
+		: FCADModelToTechSoftConverterBase(InImportParameters)
+	{
+	}
+
+	virtual bool AddBRep(AlDagNode& DagNode, const FColor& Color, EAliasObjectReference ObjectReference) override;
+
+protected:
+	TMap<void*, A3DTopoCoEdge*> AlEdgeToTSCoEdge;
+
+	A3DTopoFace* AddTrimRegion(const AlTrimRegion& InTrimRegion, const FColor& Color, EAliasObjectReference InObjectReference, const AlMatrix4x4& InAlMatrix);
+	A3DTopoLoop* CreateTopoLoop(const AlTrimBoundary& TrimBoundary);
+	A3DTopoCoEdge* CreateEdge(const AlTrimCurve& TrimCurve);
+	A3DCrvBase* CreateCurve(const AlTrimCurve& TrimCurve);
+
+	void LinkEdgesLoop(const AlTrimBoundary& TrimBoundary);
+};
+
+}
+
diff --git a/Engine/Plugins/Enterprise/DatasmithCADImporter/Source/WireInterface/Private/OpenModelUtils.cpp b/Engine/Plugins/Enterprise/DatasmithCADImporter/Source/WireInterface/Private/OpenModelUtils.cpp
new file mode 100644
--- /dev/null
+++ b/Engine/Plugins/Enterprise/DatasmithCADImporter/Source/WireInterface/Private/OpenModelUtils.cpp
@@ -0,0 +1,317 @@
+// Copyright Epic Games, Inc. All Rights Reserved.
+#include "OpenModelUtils.h"
+
+#ifdef USE_OPENMODEL
+
+#include "CADOptions.h"
+#include "DatasmithUtils.h"
+#include "DatasmithTranslator.h"
+#include "IDatasmithSceneElements.h"
+
+#include "MeshAttributes.h"
+#include "MeshDescription.h"
+#include "StaticMeshAttributes.h"
+#include "StaticMeshOperations.h"
+
+#if PLATFORM_WINDOWS
+#include "Windows/AllowWindowsPlatformTypes.h"
+#endif
+
+#include "AlDagNode.h"
+#include "AlMesh.h"
+#include "AlLayer.h"
+#include "AlTesselate.h"
+
+namespace UE_DATASMITHWIRETRANSLATOR_NAMESPACE
+{
+
+void OpenModelUtils::SetActorTransform(TSharedPtr<IDatasmithActorElement>& OutActorElement, const AlDagNode& InDagNode)
+{
+	if (OutActorElement)
+	{
+		AlLayer* LayerPtr = InDagNode.layer();
+		if (AlIsValid(LayerPtr))
+		{
+			const TUniquePtr<AlLayer> Layer(LayerPtr);
+
+			if (LayerPtr->isSymmetric())
+			{
+				return;
+			}
+		}
+
+		AlMatrix4x4 AlMatrix;
+		InDagNode.localTransformationMatrix(AlMatrix);
+
+		FMatrix Matrix;
+		double* MatrixFloats = (double*)Matrix.M;
+		for (int32 IndexI = 0; IndexI < 4; ++IndexI)
+		{
+			for (int32 IndexJ = 0; IndexJ < 4; ++IndexJ)
+			{
+				MatrixFloats[IndexI * 4 + IndexJ] = AlMatrix[IndexI][IndexJ];
+			}
+		}
+		FTransform LocalTransform(Matrix);
+		FTransform LocalUETransform = FDatasmithUtils::ConvertTransform(FDatasmithUtils::EModelCoordSystem::ZUp_RightHanded, LocalTransform);
+
+		FVector Translation = LocalUETransform.GetTranslation();
+		if (!FMath::IsNearlyEqual(Translation.X, 0.0f))
+		{
+			printf(">>> %f", Translation.X);
+		}
+		//OutActorElement->SetTranslation(LocalUETransform.GetTranslation());
+		//OutActorElement->SetScale(LocalUETransform.GetScale3D());
+		//OutActorElement->SetRotation(LocalUETransform.GetRotation());
+	}
+}
+
+bool OpenModelUtils::IsValidActor(const TSharedPtr<IDatasmithActorElement>& ActorElement)
+{
+	if (ActorElement != nullptr)
+	{
+		if (ActorElement->GetChildrenCount() > 0)
+		{
+			return true;
+		}
+		else if (ActorElement->IsA(EDatasmithElementType::StaticMeshActor))
+		{
+			const TSharedPtr<IDatasmithMeshActorElement>& MeshActorElement = StaticCastSharedPtr<IDatasmithMeshActorElement>(ActorElement);
+			return FCString::Strlen(MeshActorElement->GetStaticMeshPathName()) > 0;
+		}
+	}
+	return false;
+}
+
+bool OpenModelUtils::TransferAlMeshToMeshDescription(const AlMesh& AliasMesh, const TCHAR* SlotMaterialId, FMeshDescription& MeshDescription, CADLibrary::FMeshParameters& MeshParameters, const bool bMerge)
+{
+	if (AliasMesh.numberOfVertices() == 0 || AliasMesh.numberOfTriangles() == 0)
+	{
+		return false;
+	}
+
+	if (!bMerge)
+	{
+		MeshDescription.Empty();
+	}
+
+	int32 NbStep = 1;
+	FMatrix44f SymmetricMatrix;
+	bool bIsSymmetricMesh = MeshParameters.bIsSymmetric;
+	if (bIsSymmetricMesh)
+	{
+		NbStep = 2;
+		SymmetricMatrix = FDatasmithUtils::GetSymmetricMatrix(MeshParameters.SymmetricOrigin, MeshParameters.SymmetricNormal);
+	}
+
+	// Gather all array data
+	FStaticMeshAttributes Attributes(MeshDescription);
+	TVertexInstanceAttributesRef<FVector3f> VertexInstanceNormals = Attributes.GetVertexInstanceNormals();
+	TVertexInstanceAttributesRef<FVector2f> VertexInstanceUVs = Attributes.GetVertexInstanceUVs();
+	TPolygonGroupAttributesRef<FName> PolygonGroupImportedMaterialSlotNames = Attributes.GetPolygonGroupMaterialSlotNames();
+	TVertexAttributesRef<FVector3f> VertexPositions = MeshDescription.GetVertexPositions();
+
+	// Prepared for static mesh usage ?
+	if (!VertexPositions.IsValid() || !VertexInstanceNormals.IsValid() || !VertexInstanceUVs.IsValid() || !PolygonGroupImportedMaterialSlotNames.IsValid())
+	{
+		return false;
+	}
+
+	bool bHasUVData = (AliasMesh.uvs() != nullptr);
+
+	int VertexCount = AliasMesh.numberOfVertices();
+	int TriangleCount = AliasMesh.numberOfTriangles();
+	const int32 VertexInstanceCount = 3 * TriangleCount;
+
+	TArray<FVertexID> VertexPositionIDs;
+	VertexPositionIDs.SetNum(VertexCount * NbStep);
+
+	// Reserve space for attributes
+	// At this point, all the faces are triangles
+	MeshDescription.ReserveNewVertices(VertexCount * NbStep);
+	MeshDescription.ReserveNewVertexInstances(VertexInstanceCount * NbStep);
+	MeshDescription.ReserveNewEdges(VertexInstanceCount * NbStep);
+	MeshDescription.ReserveNewPolygons(TriangleCount * NbStep);
+
+	MeshDescription.ReserveNewPolygonGroups(1);
+	FPolygonGroupID PolyGroupId = MeshDescription.CreatePolygonGroup();
+	FName ImportedSlotName = SlotMaterialId;
+	PolygonGroupImportedMaterialSlotNames[PolyGroupId] = ImportedSlotName;
+
+	// At least one UV set must exist.
+	if (VertexInstanceUVs.GetNumChannels() == 0)
+	{
+		VertexInstanceUVs.SetNumChannels(1);
+	}
+
+	// Get Alias mesh info
+	const float* AlVertices = AliasMesh.vertices();
+
+	for (int32 Step = 0; Step < NbStep; Step++)
+	{
+		// Fill the vertex array
+		if (Step == 0)
+		{
+			FVertexID* VertexPositionIDPtr = VertexPositionIDs.GetData();
+			for (int Index = 0; Index < VertexCount; ++Index, ++VertexPositionIDPtr)
+			{
+				const float* CurVertex = AlVertices + 3 * Index;
+				*VertexPositionIDPtr = MeshDescription.CreateVertex();
+				// ConvertVector_ZUp_RightHanded
+				VertexPositions[*VertexPositionIDPtr] = FVector3f(-CurVertex[0], CurVertex[1], CurVertex[2]);
+			}
+		}
+		else
+		{
+			FVertexID* VertexPositionIDPtr = VertexPositionIDs.GetData() + VertexCount;
+			for (int Index = 0, PositionIndex = VertexCount; Index < VertexCount; ++Index, ++VertexPositionIDPtr)
+			{
+				const float* CurVertex = AlVertices + 3 * Index;
+				*VertexPositionIDPtr = MeshDescription.CreateVertex();
+				// ConvertVector_ZUp_RightHanded
+				VertexPositions[*VertexPositionIDPtr] = SymmetricMatrix.TransformPosition(FVector3f(-CurVertex[0], CurVertex[1], CurVertex[2]));
+			}
+		}
+
+		FBox UVBBox(FVector(MAX_FLT), FVector(-MAX_FLT));
+
+		const int32 CornerCount = 3; // only triangles
+		FVertexID CornerVertexIDs[3];
+		TArray<FVertexInstanceID> CornerVertexInstanceIDs;
+		CornerVertexInstanceIDs.SetNum(3);
+
+		// Get Alias mesh info
+		const int* Triangles = AliasMesh.triangles();
+		const float* AlNormals = AliasMesh.normals();
+		const float* AlUVs = AliasMesh.uvs();
+
+		// Get per-triangle data: indices, normals and uvs
+		if (!MeshParameters.bNeedSwapOrientation == ((bool)Step))
+		{
+			for (int32 FaceIndex = 0; FaceIndex < TriangleCount; ++FaceIndex, Triangles += 3)
+			{
+				// Create Vertex instances and set their attributes
+				for (int32 VertexIndex = 0, TIndex = 2; VertexIndex < CornerCount; ++VertexIndex, --TIndex)
+				{
+					CornerVertexIDs[VertexIndex] = VertexPositionIDs[Triangles[TIndex] + VertexCount * Step];
+					CornerVertexInstanceIDs[VertexIndex] = MeshDescription.CreateVertexInstance(CornerVertexIDs[VertexIndex]);
+
+					// Set the normal
+					const float* CurNormal = &AlNormals[3 * Triangles[TIndex]];
+					// ConvertVector_ZUp_RightHanded
+					FVector3f UENormal(-CurNormal[0], CurNormal[1], CurNormal[2]);
+					UENormal = UENormal.GetSafeNormal();
+					if (Step > 0)
+					{
+						UENormal = SymmetricMatrix.TransformVector(UENormal);
+					}
+					else
+					{
+						UENormal *= -1.;
+					}
+					VertexInstanceNormals[CornerVertexInstanceIDs[VertexIndex]] = UENormal;
+				}
+				if (CornerVertexIDs[0] == CornerVertexIDs[1] || CornerVertexIDs[0] == CornerVertexIDs[2] || CornerVertexIDs[1] == CornerVertexIDs[2])
+				{
+					continue;
+				}
+
+				// Set the UV
+				if (bHasUVData)
+				{
+					//for (int32 VertexIndex = 2; VertexIndex >= 0; --VertexIndex)
+					for (int32 VertexIndex = 0, TIndex = 2; VertexIndex < CornerCount; ++VertexIndex, --TIndex)
+					{
+						FVector2D UVValues(AlUVs[2 * Triangles[TIndex] + 0], AlUVs[2 * Triangles[TIndex] + 1]);
+						UVBBox += FVector(UVValues, 0.0f);
+						VertexInstanceUVs.Set(CornerVertexInstanceIDs[VertexIndex], 0, FVector2f(UVValues));
+					}
+				}
+
+				// Triangulate
+				const FPolygonID NewPolygonID = MeshDescription.CreatePolygon(PolyGroupId, CornerVertexInstanceIDs);
+			}
+		}
+		else
+		{
+			for (int32 FaceIndex = 0; FaceIndex < TriangleCount; ++FaceIndex, Triangles += 3)
+			{
+				// Create Vertex instances and set their attributes
+				for (int32 VertexIndex = 0; VertexIndex < CornerCount; ++VertexIndex)
+				{
+					CornerVertexIDs[VertexIndex] = VertexPositionIDs[Triangles[VertexIndex] + VertexCount * Step];
+					CornerVertexInstanceIDs[VertexIndex] = MeshDescription.CreateVertexInstance(CornerVertexIDs[VertexIndex]);
+
+					// Set the normal
+					const float* CurNormal = &AlNormals[3 * Triangles[VertexIndex]];
+
+					// ConvertVector_ZUp_RightHanded
+					FVector3f UENormal(-CurNormal[0], CurNormal[1], CurNormal[2]);
+					UENormal = UENormal.GetSafeNormal();
+					if (Step > 0)
+					{
+						UENormal = SymmetricMatrix.TransformVector(UENormal) * -1;
+					}
+					VertexInstanceNormals[CornerVertexInstanceIDs[VertexIndex]] = (FVector3f)UENormal;
+				}
+				if (CornerVertexIDs[0] == CornerVertexIDs[1] || CornerVertexIDs[0] == CornerVertexIDs[2] || CornerVertexIDs[1] == CornerVertexIDs[2])
+				{
+					continue;
+				}
+
+				// Set the UV
+				if (bHasUVData)
+				{
+					for (int32 VertexIndex = 0; VertexIndex < CornerCount; ++VertexIndex)
+					{
+						FVector2D UVValues(AlUVs[2 * Triangles[VertexIndex] + 0], AlUVs[2 * Triangles[VertexIndex] + 1]);
+						UVBBox += FVector(UVValues, 0.0f);
+						VertexInstanceUVs.Set(CornerVertexInstanceIDs[VertexIndex], 0, FVector2f(UVValues));
+					}
+				}
+
+				// Triangulate
+				const FPolygonID NewPolygonID = MeshDescription.CreatePolygon(PolyGroupId, CornerVertexInstanceIDs);
+			}
+		}
+	}
+
+	return true;
+}
+
+
+TAlObjectPtr<AlDagNode> OpenModelUtils::TesselateDagLeaf(const AlDagNode& DagLeaf, ETesselatorType TessType, double Tolerance)
+{
+	AlDagNode* TesselatedNode = nullptr;
+	statusCode TessStatus;
+
+	switch (TessType)
+	{
+	case(ETesselatorType::Accurate):
+		TessStatus = AlTesselate::chordHeightDeviationAccurate(TesselatedNode, &DagLeaf, Tolerance);
+		break;
+	case(ETesselatorType::Fast):
+	default:
+		TessStatus = AlTesselate::chordHeightDeviationFast(TesselatedNode, &DagLeaf, Tolerance);
+		break;
+	}
+
+	if ((TessStatus == sSuccess) && (AlIsValid(TesselatedNode) == TRUE))
+	{
+		return TAlObjectPtr<AlDagNode>(TesselatedNode);
+	}
+	else
+	{
+		return TAlObjectPtr<AlDagNode>();
+	}
+}
+
+}
+
+#if PLATFORM_WINDOWS
+#include "Windows/HideWindowsPlatformTypes.h"
+#endif
+
+#endif
+
+
diff --git a/Engine/Plugins/Enterprise/DatasmithCADImporter/Source/WireInterface/Private/OpenModelUtils.h b/Engine/Plugins/Enterprise/DatasmithCADImporter/Source/WireInterface/Private/OpenModelUtils.h
new file mode 100644
--- /dev/null
+++ b/Engine/Plugins/Enterprise/DatasmithCADImporter/Source/WireInterface/Private/OpenModelUtils.h
@@ -0,0 +1,132 @@
+// Copyright Epic Games, Inc. All Rights Reserved.
+
+#pragma once
+
+#include "CoreMinimal.h"
+
+//#include "CADModelToTechSoftConverterBase.h"
+#include "CADModelConverter.h"
+#include "CADOptions.h"
+
+#ifdef USE_OPENMODEL
+
+#if PLATFORM_WINDOWS
+#include "Windows/AllowWindowsPlatformTypes.h"
+#endif
+
+#include "AlShadingFields.h"
+#include "AlDagNode.h"
+#include "AlPersistentID.h"
+
+#if PLATFORM_WINDOWS
+#include "Windows/HideWindowsPlatformTypes.h"
+#endif
+
+class IDatasmithActorElement;
+class AlMesh;
+
+struct FMeshDescription;
+
+typedef double AlMatrix4x4[4][4];
+
+namespace UE_DATASMITHWIRETRANSLATOR_NAMESPACE
+{
+	//enum class EAliasObjectReference : uint8
+	//{
+	//	LocalReference,
+	//	ParentReference,
+	//	WorldReference,
+	//};
+
+	enum class ETesselatorType : uint8
+	{
+		Fast,
+		Accurate,
+	};
+
+	enum class EAlShaderModelType : uint8
+	{
+		BLINN,
+		LAMBERT,
+		LIGHTSOURCE,
+		PHONG,
+	};
+
+	template<typename T>
+	class TAlObjectPtr
+	{
+	public:
+		TAlObjectPtr(T* AlObject = nullptr) : Ptr(AlObject) {}
+
+		bool IsValid() const { return Ptr && AlIsValid(Ptr); }
+
+		T* operator->() const
+		{
+			return Ptr;
+		}
+
+		T& operator*() const
+		{
+			return *Ptr;
+		}
+
+		T* Get() const
+		{
+			return Ptr;
+		}
+
+		TAlObjectPtr<T> operator=(TAlObjectPtr<T> Src)
+		{
+			Ptr = Src.IsValid() ? Src.Get() : nullptr;
+			return *this;
+		}
+
+		operator bool() const
+		{
+			return IsValid();
+		}
+
+	private:
+		T* Ptr = nullptr;
+	};
+
+	namespace OpenModelUtils
+	{
+		void SetActorTransform(TSharedPtr<IDatasmithActorElement>& OutActorElement, const AlDagNode& InDagNode);
+
+		bool IsValidActor(const TSharedPtr<IDatasmithActorElement>& ActorElement);
+
+		inline FString UuidToString(const uint32& Uuid)
+		{
+			return FString::Printf(TEXT("0x%08x"), Uuid);
+		}
+
+		inline uint32 GetTypeHash(AlPersistentID& GroupNodeId)
+		{
+			int IdA, IdB, IdC, IdD;
+			GroupNodeId.id(IdA, IdB, IdC, IdD);
+			return HashCombine(IdA, HashCombine(IdB, HashCombine(IdC, IdD)));
+		}
+
+		inline uint32 GetAlDagNodeUuid(AlDagNode& DagNode)
+		{
+			if (DagNode.hasPersistentID() == sSuccess)
+			{
+				AlPersistentID* PersistentID;
+				DagNode.persistentID(PersistentID);
+				return GetTypeHash(*PersistentID);
+			}
+			FString Label = UTF8_TO_TCHAR(DagNode.name());
+			return GetTypeHash(Label);
+		}
+
+		bool TransferAlMeshToMeshDescription(const AlMesh& Mesh, const TCHAR* SlotMaterialName, FMeshDescription& MeshDescription, CADLibrary::FMeshParameters& SymmetricParameters, const bool bMerge = false);
+
+		TAlObjectPtr<AlDagNode> TesselateDagLeaf(const AlDagNode& DagLeaf, ETesselatorType TessType, double Tolerance);
+	}
+
+}
+
+#endif
+
+
diff --git a/Engine/Plugins/Enterprise/DatasmithCADImporter/Source/WireInterface/Private/WireInterfaceImpl.cpp b/Engine/Plugins/Enterprise/DatasmithCADImporter/Source/WireInterface/Private/WireInterfaceImpl.cpp
new file mode 100644
--- /dev/null
+++ b/Engine/Plugins/Enterprise/DatasmithCADImporter/Source/WireInterface/Private/WireInterfaceImpl.cpp
@@ -0,0 +1,2391 @@
+// Copyright Epic Games, Inc. All Rights Reserved.
+
+#include "WireInterfaceImpl.h"
+
+#include "CADOptions.h"
+#include "Containers/List.h"
+#include "DatasmithImportOptions.h"
+#include "DatasmithPayload.h"
+#include "DatasmithSceneFactory.h"
+#include "DatasmithUtils.h"
+#include "GenericPlatform/GenericPlatformTLS.h"
+#include "HAL/ConsoleManager.h"
+#include "IDatasmithSceneElements.h"
+#include "Misc/FileHelper.h"
+#include "Misc/Paths.h"
+#include "Modules/ModuleInterface.h"
+#include "Modules/ModuleManager.h"
+#include "Utility/DatasmithMeshHelper.h"
+
+#include "StaticMeshDescription.h"
+#include "StaticMeshOperations.h"
+
+#if WITH_EDITOR
+#include "Editor.h"
+#include "IMessageLogListing.h"
+#include "Logging/TokenizedMessage.h"
+#include "MessageLogModule.h"
+#include "Modules/ModuleManager.h"
+#endif
+
+#include "AliasBrepConverter.h"
+#include "AliasModelToCADKernelConverter.h"
+#include "AliasModelToTechSoftConverter.h" // requires Techsoft as public dependency
+#include "CADInterfacesModule.h"
+
+#if PLATFORM_WINDOWS
+#include "Windows/AllowWindowsPlatformTypes.h"
+#endif
+
+#ifndef USE_OPENMODEL
+#define USE_OPENMODEL
+#endif
+
+#ifdef USE_OPENMODEL
+#include <AlChannel.h>
+#include <AlDagNode.h>
+#include <AlGroupNode.h>
+#include <AlLayer.h>
+#include <AlLinkItem.h>
+#include <AlList.h>
+#include <AlMesh.h>
+#include <AlMeshNode.h>
+#include <AlPersistentID.h>
+#include <AlRetrieveOptions.h>
+#include <AlShader.h>
+#include <AlShadingFieldItem.h>
+#include <AlShell.h>
+#include <AlShellNode.h>
+#include <AlSurface.h>
+#include <AlSurfaceNode.h>
+#include <AlTesselate.h>
+#include <AlTrimRegion.h>
+#include <AlTM.h>
+#include <AlUniverse.h>
+#endif
+
+#if PLATFORM_WINDOWS
+#include "Windows/HideWindowsPlatformTypes.h"
+#endif
+
+DEFINE_LOG_CATEGORY_STATIC(LogWireInterface, Log, All);
+
+#define LOCTEXT_NAMESPACE "WireInterface"
+
+#define WRONG_VERSION_TEXT "Unsupported version of Alias detected. Please upgrade to Alias 2021.3 (or later version)."
+
+#ifdef OPEN_MODEL_2023_0
+
+static bool bGAliasSewByMaterial = false;
+FAutoConsoleVariableRef GAliasSewByMaterial(
+	TEXT("ds.CADTranslator.Alias.SewByMaterial"),
+	bGAliasSewByMaterial,
+	TEXT("Enable Sew action merges BReps according to their material i.e. only BReps associated with same material can be merged together.\
+Default is disable\n"),
+ECVF_Default);
+
+#endif 
+
+namespace UE_DATASMITHWIRETRANSLATOR_NAMESPACE
+{
+
+#ifdef USE_OPENMODEL
+	static bool bGSewByMaterial = false;
+	static bool bGLayersAsActors = false;
+
+	static const FColor DefaultColor = FColor(200, 200, 200);
+
+	const uint64 LibAliasNext_Version = 0xffffffffffffffffull;
+	const uint64 LibAlias2023_1_0_Version = 8162778619576619;
+	const uint64 LibAlias2023_0_0_Version = 8162774324609149;
+	const uint64 LibAlias2022_2_0_Version = 7881307937833405;
+	const uint64 LibAlias2022_1_0_Version = 7881303642865885;
+	const uint64 LibAlias2022_0_1_Version = 7881299347964005;
+	const uint64 LibAlias2021_3_2_Version = 7599833027117059;
+	const uint64 LibAlias2021_3_1_Version = 7599824433840131;
+	const uint64 LibAlias2021_3_0_Version = 7599824424206339;
+	const uint64 LibAlias2021_Version = 7599824377020416;
+	const uint64 LibAlias2020_Version = 7318349414924288;
+	const uint64 LibAlias2019_Version = 5000000000000000;
+
+#if defined(OPEN_MODEL_2020)
+	const uint64 LibAliasVersionMin = LibAlias2019_Version;
+	const uint64 LibAliasVersionMax = LibAlias2021_3_0_Version;
+	const FString AliasSdkVersion   = TEXT("2020");
+#elif defined(OPEN_MODEL_2021_3)
+	const uint64 LibAliasVersionMin = LibAlias2021_3_0_Version;
+	const uint64 LibAliasVersionMax = LibAlias2022_0_1_Version;
+	const FString AliasSdkVersion   = TEXT("2021.3");
+#elif defined(OPEN_MODEL_2022)
+	const uint64 LibAliasVersionMin = LibAlias2022_0_1_Version;
+	const uint64 LibAliasVersionMax = LibAlias2022_1_0_Version;
+	const FString AliasSdkVersion   = TEXT("2022");
+#elif defined(OPEN_MODEL_2022_1)
+	const uint64 LibAliasVersionMin = LibAlias2022_1_0_Version;
+	const uint64 LibAliasVersionMax = LibAlias2022_2_0_Version;
+	const FString AliasSdkVersion   = TEXT("2022.1");
+#elif defined(OPEN_MODEL_2022_2)
+	const uint64 LibAliasVersionMin = LibAlias2022_2_0_Version;
+	const uint64 LibAliasVersionMax = LibAlias2023_0_0_Version;
+	const FString AliasSdkVersion   = TEXT("2022.2");
+#elif defined(OPEN_MODEL_2023_0)
+	const uint64 LibAliasVersionMin = LibAlias2023_0_0_Version;
+	const uint64 LibAliasVersionMax = LibAlias2023_1_0_Version;
+	const FString AliasSdkVersion   = TEXT("2023.0");
+#elif defined(OPEN_MODEL_2023_1)      
+	const uint64 LibAliasVersionMin = LibAlias2023_1_0_Version;
+	const uint64 LibAliasVersionMax = LibAliasNext_Version;
+	const FString AliasSdkVersion   = TEXT("2023.1");
+#endif
+
+	class FBodyData
+	{
+	private:
+		TArray<TPair<TAlObjectPtr<AlDagNode>, FColor>> Shells;
+		TSet<FString> ShaderNames;
+
+		FString Label;
+		FString LayerName;
+		bool bCadData = true;
+
+	public:
+
+		FBodyData(const FString& InLayerName, bool bInCadData)
+			: bCadData(bInCadData)
+		{
+			LayerName = InLayerName;
+		};
+
+		bool IsCadData() const
+		{
+			return bCadData;
+		}
+
+		void SetLabel(const FString& InLabel)
+		{
+			Label = InLabel;
+		}
+
+		FColor LastColor;
+		bool bWarningNotLogYet = true;
+		void Add(const TAlObjectPtr<AlDagNode>& Node, FColor Color, const FString& ShaderName)
+		{
+			using namespace CADLibrary;
+			if (ShaderNames.Num() == FImportParameters::GMaxMaterialCountPerMesh)
+			{
+				if (!ShaderNames.Contains(ShaderName))
+				{
+					UE_LOG(LogWireInterface, Warning, TEXT("The main UE5 rendering systems do not support more than 256 materials per mesh and the limit has been defined to %d materials. Only the first %d materials are kept. The others are replaced by the last one"), FImportParameters::GMaxMaterialCountPerMesh, FImportParameters::GMaxMaterialCountPerMesh);
+					bWarningNotLogYet = false;
+					Color = LastColor;
+				}
+			}
+			else
+			{
+				ShaderNames.Add(ShaderName);
+				LastColor = Color;
+			}
+			Shells.Emplace(Node, Color);
+		}
+
+		void Reserve(int32 MaxElement)
+		{
+			Shells.Reserve(MaxElement);
+			ShaderNames.Reserve(MaxElement);
+		}
+
+		const TSet<FString>& GetShaderNames() const
+		{
+			return ShaderNames;
+		}
+
+		const FString& GetLayerName() const
+		{
+			return LayerName;
+		}
+
+		const TArray<TPair<TAlObjectPtr<AlDagNode>, FColor>>& GetShells() const
+		{
+			return Shells;
+		}
+
+		// Generates FBodyData's unique id from AlDagNode objects
+		uint32 GetUuid(const uint32& ParentUuid)
+		{
+			if (Shells.Num() == 0)
+			{
+				return ParentUuid;
+			}
+
+			if (Shells.Num() > 1)
+			{
+				Shells.Sort([&](const TPair<TAlObjectPtr<AlDagNode>, FColor>& NodeA, const TPair<TAlObjectPtr<AlDagNode>, FColor>& NodeB) { return OpenModelUtils::GetAlDagNodeUuid(*NodeA.Key) < OpenModelUtils::GetAlDagNodeUuid(*NodeB.Key); });
+			}
+
+			uint32 BodyUUID = 0;
+			for (const TPair<TAlObjectPtr<AlDagNode>, FColor>& DagNode : Shells)
+			{
+				BodyUUID = HashCombine(BodyUUID, OpenModelUtils::GetAlDagNodeUuid(*DagNode.Key));
+			}
+
+			return HashCombine(ParentUuid, BodyUUID);
+		}
+	};
+
+	struct FDagNodeInfo
+	{
+		uint32 Uuid;  // Use for actor name
+		FString Label;
+		TSharedPtr<IDatasmithActorElement> ActorElement;
+		FString LayerName;
+		bool bLayerIsVisible = true;
+
+		TMap<FString, TSharedPtr<IDatasmithActorElement>> LayerToActorElement;
+		FDagNodeInfo* ParentInfo = nullptr;
+
+		FDagNodeInfo(FDagNodeInfo& InParentInfo)
+			: ParentInfo(&InParentInfo)
+		{}
+
+		FDagNodeInfo()
+			: ParentInfo(nullptr)
+		{}
+
+		void ExtractDagNodeInfo(AlDagNode& CurrentNode);
+		void ExtractDagNodeInfo(FBodyData& CurrentNode);
+		void ExtractDagNodeLayer(const AlDagNode& InDagNode);
+	};
+
+	// Alias material management (to allow sew of BReps of different materials):
+	// To be compatible with "Retessellate" function, Alias material management as to be the same than CAD (TechSoft) import. 
+	// As a reminder: the name and slot of UE Material from CAD is based on CAD material/color data i.e. RGBA Color components => "UE Material slot" (int32) and "UE Material name" (FString = FString::FromInt("UE Material slot"))
+	// UE Material Label is free.
+	// 
+	// During Retessellate step, Color/Material of each CAD Faces are known, so "UE Material slot" can be deduced.
+	// 
+	// For Alias import:
+	// Alias BRep is exported in CAD modeler (CADKernel, TechSoft, ...)
+	// Material is build in UE, 
+	// From an Alias Material, a unique Color is generated.
+	// This Color is associated to the BRep Shell/face in the CAD modeler
+	// The name and slot of the associated UE Material is defined from this color
+	// So at the retessellate step, nothing changes from the CAD Retessellate process
+	// 
+	// The unique Color of an Alias Material is defined as follows:
+	// TypeHash(Alias Material Name) => uint24 == 3 uint8 => RGB components of the color  
+
+	FColor CreateShaderColorFromShaderName(const FString& ShaderName)
+	{
+		const uint32 ShaderHash = FCrc::Strihash_DEPRECATED(*ShaderName);
+		const uint32 Red = (ShaderHash & 0xff000000) >> 24;
+		const uint32 Green = (ShaderHash & 0x00ff0000) >> 16;
+		const uint32 Blue = (ShaderHash & 0x0000ff00) >> 8;
+		return FColor(Red, Green, Blue);
+	}
+
+	int32 CreateShaderId(const FColor& ShaderColor)
+	{
+		return FMath::Abs((int32)GetTypeHash(ShaderColor));
+	}
+
+	uint32 GetSceneFileHash(const FString& FullPath, const FString& FileName)
+	{
+		FFileStatData FileStatData = IFileManager::Get().GetStatData(*FullPath);
+
+		int64 FileSize = FileStatData.FileSize;
+		FDateTime ModificationTime = FileStatData.ModificationTime;
+
+		uint32 FileHash = GetTypeHash(FileName);
+		FileHash = HashCombine(FileHash, GetTypeHash(FileSize));
+		FileHash = HashCombine(FileHash, GetTypeHash(ModificationTime));
+
+		return FileHash;
+	}
+
+	bool GetConsoleBoolValue(const TCHAR* CVarName, bool bDefault)
+	{
+		const IConsoleVariable* ConsoleVariable = IConsoleManager::Get().FindConsoleVariable(CVarName);
+		return ConsoleVariable ? ConsoleVariable->GetBool() : bDefault;
+	}
+
+	bool FWireTranslatorImpl::Initialize(const TCHAR* InSceneFullName)
+	{
+		bool bCanProcess = true;
+
+		// TODO: Check file can be loaded with current SDK
+
+		if (InSceneFullName)
+		{
+			SceneFullPath = InSceneFullName;
+			SceneName = FPaths::GetBaseFilename(SceneFullPath);
+			CurrentPath = FPaths::GetPath(SceneFullPath);
+		}
+
+		// TODO: Check whether this implementation can load it
+		return bCanProcess;
+	}
+
+	FWireTranslatorImpl::~FWireTranslatorImpl()
+	{
+		if (bSceneLoaded)
+		{
+			AlUniverse::deleteAll();
+		}
+		bSceneLoaded = false;
+	}
+
+	bool FWireTranslatorImpl::GetCommonParameters(int32 Field, double Value, FColor& Color, FColor& TransparencyColor, FColor& IncandescenceColor, double GlowIntensity)
+	{
+		switch (AlShadingFields(Field))
+		{
+		case AlShadingFields::kFLD_SHADING_COMMON_COLOR_R:
+			Color.R = (uint8)Value;
+			return true;
+		case AlShadingFields::kFLD_SHADING_COMMON_COLOR_G:
+			Color.G = (uint8)Value;
+			return true;
+		case  AlShadingFields::kFLD_SHADING_COMMON_COLOR_B:
+			Color.B = (uint8)Value;
+			return true;
+		case AlShadingFields::kFLD_SHADING_COMMON_INCANDESCENCE_R:
+			IncandescenceColor.R = (uint8)Value;
+			return true;
+		case AlShadingFields::kFLD_SHADING_COMMON_INCANDESCENCE_G:
+			IncandescenceColor.G = (uint8)Value;
+			return true;
+		case AlShadingFields::kFLD_SHADING_COMMON_INCANDESCENCE_B:
+			IncandescenceColor.B = (uint8)Value;
+			return true;
+		case  AlShadingFields::kFLD_SHADING_COMMON_TRANSPARENCY_R:
+			TransparencyColor.R = (uint8)Value;
+			return true;
+		case  AlShadingFields::kFLD_SHADING_COMMON_TRANSPARENCY_G:
+			TransparencyColor.G = (uint8)Value;
+			return true;
+		case  AlShadingFields::kFLD_SHADING_COMMON_TRANSPARENCY_B:
+			TransparencyColor.B = (uint8)Value;
+			return true;
+		case AlShadingFields::kFLD_SHADING_COMMON_GLOW_INTENSITY:
+			GlowIntensity = Value;
+			return true;
+		default:
+			return false;
+		}
+	}
+
+	void FWireTranslatorImpl::SetTessellationOptions(const FDatasmithTessellationOptions& Options)
+	{
+		TessellationOptions = Options;
+		if (CADModelConverter)
+		{
+			CADModelConverter->SetImportParameters(Options.ChordTolerance, Options.MaxEdgeLength, Options.NormalTolerance, (CADLibrary::EStitchingTechnique)Options.StitchingTechnique);
+		}
+		SceneFileHash = HashCombine(Options.GetHash(), GetSceneFileHash(SceneFullPath, SceneName));
+	}
+
+	bool FWireTranslatorImpl::Load(TSharedPtr<IDatasmithScene> InScene)
+	{
+		DatasmithScene = InScene;
+
+		const FString AliasProductVersion = FString::Printf(TEXT("Alias %s"), *AliasSdkVersion);
+		DatasmithScene->SetHost(TEXT("Alias"));
+		DatasmithScene->SetVendor(TEXT("Autodesk"));
+		DatasmithScene->SetProductName(TEXT("Alias Tools"));
+		DatasmithScene->SetExporterSDKVersion(*AliasSdkVersion);
+		DatasmithScene->SetProductVersion(*AliasProductVersion);
+
+		bAliasUseNative = GetConsoleBoolValue(TEXT("ds.Wiretranslator.UseNative"), false);
+		if (!bAliasUseNative)
+		{
+			CADLibrary::FImportParameters ImportParameters;
+			if (CADLibrary::FImportParameters::bGDisableCADKernelTessellation)
+			{
+				TSharedPtr<FAliasModelToTechSoftConverter> AliasToTechSoftConverter = MakeShared<FAliasModelToTechSoftConverter>(ImportParameters);
+				CADModelConverter = AliasToTechSoftConverter;
+				AliasBRepConverter = AliasToTechSoftConverter;
+			}
+			else
+			{
+				TSharedPtr<FAliasModelToCADKernelConverter> AliasToCADKernelConverter = MakeShared<FAliasModelToCADKernelConverter>(ImportParameters);
+				CADModelConverter = AliasToCADKernelConverter;
+				AliasBRepConverter = AliasToCADKernelConverter;
+			}
+		}
+
+		bGSewByMaterial = GetConsoleBoolValue(TEXT("ds.CADTranslator.Alias.SewByMaterial"), false);
+		bGLayersAsActors = GetConsoleBoolValue(TEXT("ds.CADTranslator.Alias.LayersAsActors"), false);
+
+		// Initialize Alias.
+		AlUniverse::initialize();
+
+		if (AlUniverse::retrieve(TCHAR_TO_UTF8(*SceneFullPath)) != sSuccess)
+		{
+			return false;
+		}
+
+		AlRetrieveOptions options;
+		AlUniverse::retrieveOptions(options);
+
+		bSceneLoaded = true;
+
+		// Make material
+		if (!GetShader())
+		{
+			return false;
+		}
+
+		// Parse and extract the DAG leaf nodes.
+		// Note that Alias file unit is cm like UE
+		return GetDagLeaves();
+	}
+
+	void FWireTranslatorImpl::AddAlBlinnParameters(const TAlObjectPtr<AlShader>& Shader, TSharedPtr<IDatasmithUEPbrMaterialElement> MaterialElement)
+	{
+		// Default values for a Blinn material
+		FColor Color(145, 148, 153);
+		FColor TransparencyColor(0, 0, 0);
+		FColor IncandescenceColor(0, 0, 0);
+		FColor SpecularColor(38, 38, 38);
+		double Diffuse = 1.0;
+		double GlowIntensity = 0.0;
+		double Gloss = 0.8;
+		double Eccentricity = 0.35;
+		double Specularity = 1.0;
+		double Reflectivity = 0.5;
+		double SpecularRolloff = 0.5;
+
+		AlList* List = Shader->fields();
+		for (AlShadingFieldItem* Item = static_cast<AlShadingFieldItem*>(List->first()); Item; Item = Item->nextField())
+		{
+			double Value = 0.0f;
+			statusCode ErrorCode = Shader->parameter(Item->field(), Value);
+			if (ErrorCode != 0)
+			{
+				continue;
+			}
+
+			if (GetCommonParameters(Item->field(), Value, Color, TransparencyColor, IncandescenceColor, GlowIntensity))
+			{
+				continue;
+			}
+
+			switch (Item->field())
+			{
+			case AlShadingFields::kFLD_SHADING_BLINN_DIFFUSE:
+				Diffuse = Value;
+				break;
+			case AlShadingFields::kFLD_SHADING_BLINN_GLOSS_:
+				Gloss = Value;
+				break;
+			case AlShadingFields::kFLD_SHADING_BLINN_SPECULAR_R:
+				SpecularColor.R = (uint8)(255.f * Value);
+				break;
+			case AlShadingFields::kFLD_SHADING_BLINN_SPECULAR_G:
+				SpecularColor.G = (uint8)(255.f * Value);;
+				break;
+			case AlShadingFields::kFLD_SHADING_BLINN_SPECULAR_B:
+				SpecularColor.B = (uint8)(255.f * Value);;
+				break;
+			case AlShadingFields::kFLD_SHADING_BLINN_SPECULARITY_:
+				Specularity = Value;
+				break;
+			case AlShadingFields::kFLD_SHADING_BLINN_SPECULAR_ROLLOFF:
+				SpecularRolloff = Value;
+				break;
+			case AlShadingFields::kFLD_SHADING_BLINN_ECCENTRICITY:
+				Eccentricity = Value;
+				break;
+			case AlShadingFields::kFLD_SHADING_BLINN_REFLECTIVITY:
+				Reflectivity = Value;
+				break;
+			}
+		}
+
+		bool bIsTransparent = IsTransparent(TransparencyColor);
+
+		// Construct parameter expressions
+		IDatasmithMaterialExpressionScalar* DiffuseExpression = MaterialElement->AddMaterialExpression<IDatasmithMaterialExpressionScalar>();
+		DiffuseExpression->GetScalar() = Diffuse;
+		DiffuseExpression->SetName(TEXT("Diffuse"));
+
+		IDatasmithMaterialExpressionScalar* GlossExpression = MaterialElement->AddMaterialExpression<IDatasmithMaterialExpressionScalar>();
+		GlossExpression->GetScalar() = Gloss;
+		GlossExpression->SetName(TEXT("Gloss"));
+
+		IDatasmithMaterialExpressionColor* SpecularColorExpression = MaterialElement->AddMaterialExpression<IDatasmithMaterialExpressionColor>();
+		SpecularColorExpression->SetName(TEXT("SpecularColor"));
+		SpecularColorExpression->GetColor() = FLinearColor::FromSRGBColor(SpecularColor);
+
+		IDatasmithMaterialExpressionScalar* SpecularityExpression = MaterialElement->AddMaterialExpression<IDatasmithMaterialExpressionScalar>();
+		SpecularityExpression->GetScalar() = Specularity * 0.3;
+		SpecularityExpression->SetName(TEXT("Specularity"));
+
+		IDatasmithMaterialExpressionScalar* SpecularRolloffExpression = MaterialElement->AddMaterialExpression<IDatasmithMaterialExpressionScalar>();
+		SpecularRolloffExpression->GetScalar() = SpecularRolloff;
+		SpecularRolloffExpression->SetName(TEXT("SpecularRolloff"));
+
+		IDatasmithMaterialExpressionScalar* EccentricityExpression = MaterialElement->AddMaterialExpression<IDatasmithMaterialExpressionScalar>();
+		EccentricityExpression->GetScalar() = Eccentricity;
+		EccentricityExpression->SetName(TEXT("Eccentricity"));
+
+		IDatasmithMaterialExpressionScalar* ReflectivityExpression = MaterialElement->AddMaterialExpression<IDatasmithMaterialExpressionScalar>();
+		ReflectivityExpression->GetScalar() = Reflectivity;
+		ReflectivityExpression->SetName(TEXT("Reflectivity"));
+
+		IDatasmithMaterialExpressionColor* ColorExpression = MaterialElement->AddMaterialExpression<IDatasmithMaterialExpressionColor>();
+		ColorExpression->SetName(TEXT("Color"));
+		ColorExpression->GetColor() = FLinearColor::FromSRGBColor(Color);
+
+		IDatasmithMaterialExpressionColor* IncandescenceColorExpression = MaterialElement->AddMaterialExpression<IDatasmithMaterialExpressionColor>();
+		IncandescenceColorExpression->SetName(TEXT("IncandescenceColor"));
+		IncandescenceColorExpression->GetColor() = FLinearColor::FromSRGBColor(IncandescenceColor);
+
+		IDatasmithMaterialExpressionColor* TransparencyColorExpression = MaterialElement->AddMaterialExpression<IDatasmithMaterialExpressionColor>();
+		TransparencyColorExpression->SetName(TEXT("TransparencyColor"));
+		TransparencyColorExpression->GetColor() = FLinearColor::FromSRGBColor(TransparencyColor);
+
+		IDatasmithMaterialExpressionScalar* GlowIntensityExpression = MaterialElement->AddMaterialExpression<IDatasmithMaterialExpressionScalar>();
+		GlowIntensityExpression->GetScalar() = GlowIntensity;
+		GlowIntensityExpression->SetName(TEXT("GlowIntensity"));
+
+		// Create aux expressions
+		IDatasmithMaterialExpressionGeneric* ColorSpecLerp = MaterialElement->AddMaterialExpression<IDatasmithMaterialExpressionGeneric>();
+		ColorSpecLerp->SetExpressionName(TEXT("LinearInterpolate"));
+
+		IDatasmithMaterialExpressionScalar* ColorSpecLerpValue = MaterialElement->AddMaterialExpression<IDatasmithMaterialExpressionScalar>();
+		ColorSpecLerpValue->GetScalar() = 0.96f;
+
+		IDatasmithMaterialExpressionGeneric* ColorMetallicLerp = MaterialElement->AddMaterialExpression<IDatasmithMaterialExpressionGeneric>();
+		ColorMetallicLerp->SetExpressionName(TEXT("LinearInterpolate"));
+
+		IDatasmithMaterialExpressionGeneric* DiffuseLerp = MaterialElement->AddMaterialExpression<IDatasmithMaterialExpressionGeneric>();
+		DiffuseLerp->SetExpressionName(TEXT("LinearInterpolate"));
+
+		IDatasmithMaterialExpressionScalar* DiffuseLerpA = MaterialElement->AddMaterialExpression<IDatasmithMaterialExpressionScalar>();
+		DiffuseLerpA->GetScalar() = 0.04f;
+
+		IDatasmithMaterialExpressionScalar* DiffuseLerpB = MaterialElement->AddMaterialExpression<IDatasmithMaterialExpressionScalar>();
+		DiffuseLerpB->GetScalar() = 1.0f;
+
+		IDatasmithMaterialExpressionGeneric* BaseColorMultiply = MaterialElement->AddMaterialExpression<IDatasmithMaterialExpressionGeneric>();
+		BaseColorMultiply->SetExpressionName(TEXT("Multiply"));
+
+		IDatasmithMaterialExpressionGeneric* BaseColorAdd = MaterialElement->AddMaterialExpression<IDatasmithMaterialExpressionGeneric>();
+		BaseColorAdd->SetExpressionName(TEXT("Add"));
+
+		IDatasmithMaterialExpressionGeneric* BaseColorTransparencyMultiply = MaterialElement->AddMaterialExpression<IDatasmithMaterialExpressionGeneric>();
+		BaseColorTransparencyMultiply->SetExpressionName(TEXT("Multiply"));
+
+		IDatasmithMaterialExpressionGeneric* IncandescenceMultiply = MaterialElement->AddMaterialExpression<IDatasmithMaterialExpressionGeneric>();
+		IncandescenceMultiply->SetExpressionName(TEXT("Multiply"));
+
+		IDatasmithMaterialExpressionGeneric* IncandescenceScaleMultiply = MaterialElement->AddMaterialExpression<IDatasmithMaterialExpressionGeneric>();
+		IncandescenceScaleMultiply->SetExpressionName(TEXT("Multiply"));
+
+		IDatasmithMaterialExpressionScalar* IncandescenceScale = MaterialElement->AddMaterialExpression<IDatasmithMaterialExpressionScalar>();
+		IncandescenceScale->GetScalar() = 100.0f;
+
+		IDatasmithMaterialExpressionGeneric* EccentricityMultiply = MaterialElement->AddMaterialExpression<IDatasmithMaterialExpressionGeneric>();
+		EccentricityMultiply->SetExpressionName(TEXT("Multiply"));
+
+		IDatasmithMaterialExpressionGeneric* EccentricityOneMinus = MaterialElement->AddMaterialExpression<IDatasmithMaterialExpressionGeneric>();
+		EccentricityOneMinus->SetExpressionName(TEXT("OneMinus"));
+
+		IDatasmithMaterialExpressionGeneric* RoughnessOneMinus = MaterialElement->AddMaterialExpression<IDatasmithMaterialExpressionGeneric>();
+		RoughnessOneMinus->SetExpressionName(TEXT("OneMinus"));
+
+		IDatasmithMaterialExpressionScalar* FresnelExponent = MaterialElement->AddMaterialExpression<IDatasmithMaterialExpressionScalar>();
+		FresnelExponent->GetScalar() = 4.0f;
+
+		IDatasmithMaterialExpressionFunctionCall* FresnelFunc = MaterialElement->AddMaterialExpression<IDatasmithMaterialExpressionFunctionCall>();
+		FresnelFunc->SetFunctionPathName(TEXT("/Engine/Functions/Engine_MaterialFunctions02/Fresnel_Function.Fresnel_Function"));
+
+		IDatasmithMaterialExpressionGeneric* FresnelLerp = MaterialElement->AddMaterialExpression<IDatasmithMaterialExpressionGeneric>();
+		FresnelLerp->SetExpressionName(TEXT("LinearInterpolate"));
+
+		IDatasmithMaterialExpressionScalar* FresnelLerpA = MaterialElement->AddMaterialExpression<IDatasmithMaterialExpressionScalar>();
+		FresnelLerpA->GetScalar() = 1.0f;
+
+		IDatasmithMaterialExpressionScalar* SpecularPowerExp = MaterialElement->AddMaterialExpression<IDatasmithMaterialExpressionScalar>();
+		SpecularPowerExp->GetScalar() = 0.5f;
+
+		IDatasmithMaterialExpressionGeneric* Power = MaterialElement->AddMaterialExpression<IDatasmithMaterialExpressionGeneric>();
+		Power->SetExpressionName(TEXT("Power"));
+
+		IDatasmithMaterialExpressionGeneric* FresnelMultiply = MaterialElement->AddMaterialExpression<IDatasmithMaterialExpressionGeneric>();
+		FresnelMultiply->SetExpressionName(TEXT("Multiply"));
+
+		IDatasmithMaterialExpressionGeneric* TransparencyOneMinus = MaterialElement->AddMaterialExpression<IDatasmithMaterialExpressionGeneric>();
+		TransparencyOneMinus->SetExpressionName(TEXT("OneMinus"));
+
+		IDatasmithMaterialExpressionFunctionCall* BreakFloat3 = nullptr;
+		IDatasmithMaterialExpressionGeneric* AddRG = nullptr;
+		IDatasmithMaterialExpressionGeneric* AddRGB = nullptr;
+		IDatasmithMaterialExpressionGeneric* Divide = nullptr;
+		IDatasmithMaterialExpressionScalar* DivideConstant = nullptr;
+		if (bIsTransparent)
+		{
+			BreakFloat3 = MaterialElement->AddMaterialExpression<IDatasmithMaterialExpressionFunctionCall>();
+			BreakFloat3->SetFunctionPathName(TEXT("/Engine/Functions/Engine_MaterialFunctions02/Utility/BreakFloat3Components.BreakFloat3Components"));
+
+			AddRG = MaterialElement->AddMaterialExpression<IDatasmithMaterialExpressionGeneric>();
+			AddRG->SetExpressionName(TEXT("Add"));
+
+			AddRGB = MaterialElement->AddMaterialExpression<IDatasmithMaterialExpressionGeneric>();
+			AddRGB->SetExpressionName(TEXT("Add"));
+
+			Divide = MaterialElement->AddMaterialExpression<IDatasmithMaterialExpressionGeneric>();
+			Divide->SetExpressionName(TEXT("Divide"));
+
+			DivideConstant = MaterialElement->AddMaterialExpression<IDatasmithMaterialExpressionScalar>();
+			DivideConstant->GetScalar() = 3.0f;
+		}
+
+		// Connect expressions
+		SpecularColorExpression->ConnectExpression(*ColorSpecLerp->GetInput(0));
+		ColorExpression->ConnectExpression(*ColorSpecLerp->GetInput(1));
+		ColorSpecLerpValue->ConnectExpression(*ColorSpecLerp->GetInput(2));
+
+		ColorExpression->ConnectExpression(*ColorMetallicLerp->GetInput(0));
+		ColorSpecLerp->ConnectExpression(*ColorMetallicLerp->GetInput(1));
+		GlossExpression->ConnectExpression(*ColorMetallicLerp->GetInput(2));
+
+		DiffuseLerpA->ConnectExpression(*DiffuseLerp->GetInput(0));
+		DiffuseLerpB->ConnectExpression(*DiffuseLerp->GetInput(1));
+		DiffuseExpression->ConnectExpression(*DiffuseLerp->GetInput(2));
+
+		ColorMetallicLerp->ConnectExpression(*BaseColorMultiply->GetInput(0));
+		DiffuseLerp->ConnectExpression(*BaseColorMultiply->GetInput(1));
+
+		BaseColorMultiply->ConnectExpression(*BaseColorAdd->GetInput(0));
+		IncandescenceColorExpression->ConnectExpression(*BaseColorAdd->GetInput(1));
+
+		BaseColorAdd->ConnectExpression(*BaseColorTransparencyMultiply->GetInput(0));
+		TransparencyOneMinus->ConnectExpression(*BaseColorTransparencyMultiply->GetInput(1));
+
+		GlowIntensityExpression->ConnectExpression(*IncandescenceScaleMultiply->GetInput(0));
+		IncandescenceScale->ConnectExpression(*IncandescenceScaleMultiply->GetInput(1));
+
+		BaseColorTransparencyMultiply->ConnectExpression(*IncandescenceMultiply->GetInput(0));
+		IncandescenceScaleMultiply->ConnectExpression(*IncandescenceMultiply->GetInput(1));
+
+		EccentricityExpression->ConnectExpression(*EccentricityOneMinus->GetInput(0));
+
+		EccentricityOneMinus->ConnectExpression(*EccentricityMultiply->GetInput(0));
+		SpecularityExpression->ConnectExpression(*EccentricityMultiply->GetInput(1));
+
+		EccentricityMultiply->ConnectExpression(*RoughnessOneMinus->GetInput(0));
+
+		FresnelExponent->ConnectExpression(*FresnelFunc->GetInput(3));
+
+		SpecularRolloffExpression->ConnectExpression(*Power->GetInput(0));
+		SpecularPowerExp->ConnectExpression(*Power->GetInput(1));
+
+		FresnelLerpA->ConnectExpression(*FresnelLerp->GetInput(0));
+		FresnelFunc->ConnectExpression(*FresnelLerp->GetInput(1));
+		Power->ConnectExpression(*FresnelLerp->GetInput(2));
+
+		FresnelLerp->ConnectExpression(*FresnelMultiply->GetInput(0));
+		ReflectivityExpression->ConnectExpression(*FresnelMultiply->GetInput(1));
+
+		TransparencyColorExpression->ConnectExpression(*TransparencyOneMinus->GetInput(0));
+
+		if (bIsTransparent)
+		{
+			TransparencyOneMinus->ConnectExpression(*BreakFloat3->GetInput(0));
+
+			BreakFloat3->ConnectExpression(*AddRG->GetInput(0), 0);
+			BreakFloat3->ConnectExpression(*AddRG->GetInput(1), 1);
+
+			AddRG->ConnectExpression(*AddRGB->GetInput(0));
+			BreakFloat3->ConnectExpression(*AddRGB->GetInput(1), 2);
+
+			AddRGB->ConnectExpression(*Divide->GetInput(0));
+			DivideConstant->ConnectExpression(*Divide->GetInput(1));
+		}
+
+		// Connect material outputs
+		MaterialElement->GetBaseColor().SetExpression(BaseColorTransparencyMultiply);
+		MaterialElement->GetMetallic().SetExpression(GlossExpression);
+		MaterialElement->GetSpecular().SetExpression(FresnelMultiply);
+		MaterialElement->GetRoughness().SetExpression(RoughnessOneMinus);
+		MaterialElement->GetEmissiveColor().SetExpression(IncandescenceMultiply);
+
+		if (bIsTransparent)
+		{
+			MaterialElement->GetOpacity().SetExpression(Divide);
+			MaterialElement->SetParentLabel(TEXT("M_DatasmithAliasBlinnTransparent"));
+		}
+		else
+		{
+			MaterialElement->SetParentLabel(TEXT("M_DatasmithAliasBlinn"));
+		}
+
+	}
+
+	void FWireTranslatorImpl::AddAlLambertParameters(const TAlObjectPtr<AlShader>& Shader, TSharedPtr<IDatasmithUEPbrMaterialElement> MaterialElement)
+	{
+		// Default values for a Lambert material
+		FColor Color(145, 148, 153);
+		FColor TransparencyColor(0, 0, 0);
+		FColor IncandescenceColor(0, 0, 0);
+		double Diffuse = 1.0;
+		double GlowIntensity = 0.0;
+
+		AlList* List = Shader->fields();
+		for (AlShadingFieldItem* Item = static_cast<AlShadingFieldItem*>(List->first()); Item; Item = Item->nextField())
+		{
+			double Value = 0.0f;
+			statusCode ErrorCode = Shader->parameter(Item->field(), Value);
+			if (ErrorCode != 0)
+			{
+				continue;
+			}
+
+			if (GetCommonParameters(Item->field(), Value, Color, TransparencyColor, IncandescenceColor, GlowIntensity))
+			{
+				continue;
+			}
+
+			switch (Item->field())
+			{
+			case AlShadingFields::kFLD_SHADING_LAMBERT_DIFFUSE:
+				Diffuse = Value;
+				break;
+			}
+		}
+
+		bool bIsTransparent = IsTransparent(TransparencyColor);
+
+		// Construct parameter expressions
+		IDatasmithMaterialExpressionScalar* DiffuseExpression = MaterialElement->AddMaterialExpression<IDatasmithMaterialExpressionScalar>();
+		DiffuseExpression->GetScalar() = Diffuse;
+		DiffuseExpression->SetName(TEXT("Diffuse"));
+
+		IDatasmithMaterialExpressionColor* ColorExpression = MaterialElement->AddMaterialExpression<IDatasmithMaterialExpressionColor>();
+		ColorExpression->SetName(TEXT("Color"));
+		ColorExpression->GetColor() = FLinearColor::FromSRGBColor(Color);
+
+		IDatasmithMaterialExpressionColor* IncandescenceColorExpression = MaterialElement->AddMaterialExpression<IDatasmithMaterialExpressionColor>();
+		IncandescenceColorExpression->SetName(TEXT("IncandescenceColor"));
+		IncandescenceColorExpression->GetColor() = FLinearColor::FromSRGBColor(IncandescenceColor);
+
+		IDatasmithMaterialExpressionColor* TransparencyColorExpression = MaterialElement->AddMaterialExpression<IDatasmithMaterialExpressionColor>();
+		TransparencyColorExpression->SetName(TEXT("TransparencyColor"));
+		TransparencyColorExpression->GetColor() = FLinearColor::FromSRGBColor(TransparencyColor);
+
+		IDatasmithMaterialExpressionScalar* GlowIntensityExpression = MaterialElement->AddMaterialExpression<IDatasmithMaterialExpressionScalar>();
+		GlowIntensityExpression->GetScalar() = GlowIntensity;
+		GlowIntensityExpression->SetName(TEXT("GlowIntensity"));
+
+		// Create aux expressions
+		IDatasmithMaterialExpressionGeneric* DiffuseLerp = MaterialElement->AddMaterialExpression<IDatasmithMaterialExpressionGeneric>();
+		DiffuseLerp->SetExpressionName(TEXT("LinearInterpolate"));
+
+		IDatasmithMaterialExpressionScalar* DiffuseLerpA = MaterialElement->AddMaterialExpression<IDatasmithMaterialExpressionScalar>();
+		DiffuseLerpA->GetScalar() = 0.04f;
+
+		IDatasmithMaterialExpressionScalar* DiffuseLerpB = MaterialElement->AddMaterialExpression<IDatasmithMaterialExpressionScalar>();
+		DiffuseLerpB->GetScalar() = 1.0f;
+
+		IDatasmithMaterialExpressionGeneric* BaseColorMultiply = MaterialElement->AddMaterialExpression<IDatasmithMaterialExpressionGeneric>();
+		BaseColorMultiply->SetExpressionName(TEXT("Multiply"));
+
+		IDatasmithMaterialExpressionGeneric* BaseColorAdd = MaterialElement->AddMaterialExpression<IDatasmithMaterialExpressionGeneric>();
+		BaseColorAdd->SetExpressionName(TEXT("Add"));
+
+		IDatasmithMaterialExpressionGeneric* BaseColorTransparencyMultiply = MaterialElement->AddMaterialExpression<IDatasmithMaterialExpressionGeneric>();
+		BaseColorTransparencyMultiply->SetExpressionName(TEXT("Multiply"));
+
+		IDatasmithMaterialExpressionGeneric* IncandescenceMultiply = MaterialElement->AddMaterialExpression<IDatasmithMaterialExpressionGeneric>();
+		IncandescenceMultiply->SetExpressionName(TEXT("Multiply"));
+
+		IDatasmithMaterialExpressionGeneric* IncandescenceScaleMultiply = MaterialElement->AddMaterialExpression<IDatasmithMaterialExpressionGeneric>();
+		IncandescenceScaleMultiply->SetExpressionName(TEXT("Multiply"));
+
+		IDatasmithMaterialExpressionScalar* IncandescenceScale = MaterialElement->AddMaterialExpression<IDatasmithMaterialExpressionScalar>();
+		IncandescenceScale->GetScalar() = 100.0f;
+
+		IDatasmithMaterialExpressionGeneric* TransparencyOneMinus = MaterialElement->AddMaterialExpression<IDatasmithMaterialExpressionGeneric>();
+		TransparencyOneMinus->SetExpressionName(TEXT("OneMinus"));
+
+		IDatasmithMaterialExpressionFunctionCall* BreakFloat3 = nullptr;
+		IDatasmithMaterialExpressionGeneric* AddRG = nullptr;
+		IDatasmithMaterialExpressionGeneric* AddRGB = nullptr;
+		IDatasmithMaterialExpressionGeneric* Divide = nullptr;
+		IDatasmithMaterialExpressionScalar* DivideConstant = nullptr;
+		if (bIsTransparent)
+		{
+			BreakFloat3 = MaterialElement->AddMaterialExpression<IDatasmithMaterialExpressionFunctionCall>();
+			BreakFloat3->SetFunctionPathName(TEXT("/Engine/Functions/Engine_MaterialFunctions02/Utility/BreakFloat3Components.BreakFloat3Components"));
+
+			AddRG = MaterialElement->AddMaterialExpression<IDatasmithMaterialExpressionGeneric>();
+			AddRG->SetExpressionName(TEXT("Add"));
+
+			AddRGB = MaterialElement->AddMaterialExpression<IDatasmithMaterialExpressionGeneric>();
+			AddRGB->SetExpressionName(TEXT("Add"));
+
+			Divide = MaterialElement->AddMaterialExpression<IDatasmithMaterialExpressionGeneric>();
+			Divide->SetExpressionName(TEXT("Divide"));
+
+			DivideConstant = MaterialElement->AddMaterialExpression<IDatasmithMaterialExpressionScalar>();
+			DivideConstant->GetScalar() = 3.0f;
+		}
+
+		// Connect expressions
+		DiffuseLerpA->ConnectExpression(*DiffuseLerp->GetInput(0));
+		DiffuseLerpB->ConnectExpression(*DiffuseLerp->GetInput(1));
+		DiffuseExpression->ConnectExpression(*DiffuseLerp->GetInput(2));
+
+		ColorExpression->ConnectExpression(*BaseColorMultiply->GetInput(0));
+		DiffuseLerp->ConnectExpression(*BaseColorMultiply->GetInput(1));
+
+		BaseColorMultiply->ConnectExpression(*BaseColorAdd->GetInput(0));
+		IncandescenceColorExpression->ConnectExpression(*BaseColorAdd->GetInput(1));
+
+		BaseColorAdd->ConnectExpression(*BaseColorTransparencyMultiply->GetInput(0));
+		TransparencyOneMinus->ConnectExpression(*BaseColorTransparencyMultiply->GetInput(1));
+
+		GlowIntensityExpression->ConnectExpression(*IncandescenceScaleMultiply->GetInput(0));
+		IncandescenceScale->ConnectExpression(*IncandescenceScaleMultiply->GetInput(1));
+
+		BaseColorTransparencyMultiply->ConnectExpression(*IncandescenceMultiply->GetInput(0));
+		IncandescenceScaleMultiply->ConnectExpression(*IncandescenceMultiply->GetInput(1));
+
+		TransparencyColorExpression->ConnectExpression(*TransparencyOneMinus->GetInput(0));
+
+		if (bIsTransparent)
+		{
+			TransparencyOneMinus->ConnectExpression(*BreakFloat3->GetInput(0));
+
+			BreakFloat3->ConnectExpression(*AddRG->GetInput(0), 0);
+			BreakFloat3->ConnectExpression(*AddRG->GetInput(1), 1);
+
+			AddRG->ConnectExpression(*AddRGB->GetInput(0));
+			BreakFloat3->ConnectExpression(*AddRGB->GetInput(1), 2);
+
+			AddRGB->ConnectExpression(*Divide->GetInput(0));
+			DivideConstant->ConnectExpression(*Divide->GetInput(1));
+		}
+
+		// Connect material outputs
+		MaterialElement->GetBaseColor().SetExpression(BaseColorTransparencyMultiply);
+		MaterialElement->GetEmissiveColor().SetExpression(IncandescenceMultiply);
+		if (bIsTransparent)
+		{
+			MaterialElement->GetOpacity().SetExpression(Divide);
+			MaterialElement->SetParentLabel(TEXT("M_DatasmithAliasLambertTransparent"));
+		}
+		else {
+			MaterialElement->SetParentLabel(TEXT("M_DatasmithAliasLambert"));
+		}
+
+	}
+
+	void FWireTranslatorImpl::AddAlLightSourceParameters(const TAlObjectPtr<AlShader>& Shader, TSharedPtr<IDatasmithUEPbrMaterialElement> MaterialElement)
+	{
+		// Default values for a LightSource material
+		FColor Color(145, 148, 153);
+		FColor TransparencyColor(0, 0, 0);
+		FColor IncandescenceColor(0, 0, 0);
+		double GlowIntensity = 0.0;
+
+		AlList* List = Shader->fields();
+		for (AlShadingFieldItem* Item = static_cast<AlShadingFieldItem*>(List->first()); Item; Item = Item->nextField())
+		{
+			double Value = 0.0f;
+			statusCode ErrorCode = Shader->parameter(Item->field(), Value);
+			if (ErrorCode != 0)
+			{
+				continue;
+			}
+
+			GetCommonParameters(Item->field(), Value, Color, TransparencyColor, IncandescenceColor, GlowIntensity);
+		}
+
+		bool bIsTransparent = IsTransparent(TransparencyColor);
+
+		// Construct parameter expressions
+		IDatasmithMaterialExpressionColor* ColorExpression = MaterialElement->AddMaterialExpression<IDatasmithMaterialExpressionColor>();
+		ColorExpression->SetName(TEXT("Color"));
+		ColorExpression->GetColor() = FLinearColor::FromSRGBColor(Color);
+
+		IDatasmithMaterialExpressionColor* IncandescenceColorExpression = MaterialElement->AddMaterialExpression<IDatasmithMaterialExpressionColor>();
+		IncandescenceColorExpression->SetName(TEXT("IncandescenceColor"));
+		IncandescenceColorExpression->GetColor() = FLinearColor::FromSRGBColor(IncandescenceColor);
+
+		IDatasmithMaterialExpressionColor* TransparencyColorExpression = MaterialElement->AddMaterialExpression<IDatasmithMaterialExpressionColor>();
+		TransparencyColorExpression->SetName(TEXT("TransparencyColor"));
+		TransparencyColorExpression->GetColor() = FLinearColor::FromSRGBColor(TransparencyColor);
+
+		IDatasmithMaterialExpressionScalar* GlowIntensityExpression = MaterialElement->AddMaterialExpression<IDatasmithMaterialExpressionScalar>();
+		GlowIntensityExpression->GetScalar() = GlowIntensity;
+		GlowIntensityExpression->SetName(TEXT("GlowIntensity"));
+
+		// Create aux expressions
+		IDatasmithMaterialExpressionGeneric* BaseColorAdd = MaterialElement->AddMaterialExpression<IDatasmithMaterialExpressionGeneric>();
+		BaseColorAdd->SetExpressionName(TEXT("Add"));
+
+		IDatasmithMaterialExpressionGeneric* BaseColorTransparencyMultiply = MaterialElement->AddMaterialExpression<IDatasmithMaterialExpressionGeneric>();
+		BaseColorTransparencyMultiply->SetExpressionName(TEXT("Multiply"));
+
+		IDatasmithMaterialExpressionGeneric* IncandescenceMultiply = MaterialElement->AddMaterialExpression<IDatasmithMaterialExpressionGeneric>();
+		IncandescenceMultiply->SetExpressionName(TEXT("Multiply"));
+
+		IDatasmithMaterialExpressionGeneric* IncandescenceScaleMultiply = MaterialElement->AddMaterialExpression<IDatasmithMaterialExpressionGeneric>();
+		IncandescenceScaleMultiply->SetExpressionName(TEXT("Multiply"));
+
+		IDatasmithMaterialExpressionScalar* IncandescenceScale = MaterialElement->AddMaterialExpression<IDatasmithMaterialExpressionScalar>();
+		IncandescenceScale->GetScalar() = 100.0f;
+
+		IDatasmithMaterialExpressionGeneric* TransparencyOneMinus = MaterialElement->AddMaterialExpression<IDatasmithMaterialExpressionGeneric>();
+		TransparencyOneMinus->SetExpressionName(TEXT("OneMinus"));
+
+		IDatasmithMaterialExpressionFunctionCall* BreakFloat3 = nullptr;
+		IDatasmithMaterialExpressionGeneric* AddRG = nullptr;
+		IDatasmithMaterialExpressionGeneric* AddRGB = nullptr;
+		IDatasmithMaterialExpressionGeneric* Divide = nullptr;
+		IDatasmithMaterialExpressionScalar* DivideConstant = nullptr;
+		if (bIsTransparent)
+		{
+			BreakFloat3 = MaterialElement->AddMaterialExpression<IDatasmithMaterialExpressionFunctionCall>();
+			BreakFloat3->SetFunctionPathName(TEXT("/Engine/Functions/Engine_MaterialFunctions02/Utility/BreakFloat3Components.BreakFloat3Components"));
+
+			AddRG = MaterialElement->AddMaterialExpression<IDatasmithMaterialExpressionGeneric>();
+			AddRG->SetExpressionName(TEXT("Add"));
+
+			AddRGB = MaterialElement->AddMaterialExpression<IDatasmithMaterialExpressionGeneric>();
+			AddRGB->SetExpressionName(TEXT("Add"));
+
+			Divide = MaterialElement->AddMaterialExpression<IDatasmithMaterialExpressionGeneric>();
+			Divide->SetExpressionName(TEXT("Divide"));
+
+			DivideConstant = MaterialElement->AddMaterialExpression<IDatasmithMaterialExpressionScalar>();
+			DivideConstant->GetScalar() = 3.0f;
+		}
+
+		// Connect expressions
+		ColorExpression->ConnectExpression(*BaseColorAdd->GetInput(0));
+		IncandescenceColorExpression->ConnectExpression(*BaseColorAdd->GetInput(1));
+
+		BaseColorAdd->ConnectExpression(*BaseColorTransparencyMultiply->GetInput(0));
+		TransparencyOneMinus->ConnectExpression(*BaseColorTransparencyMultiply->GetInput(1));
+
+		GlowIntensityExpression->ConnectExpression(*IncandescenceScaleMultiply->GetInput(0));
+		IncandescenceScale->ConnectExpression(*IncandescenceScaleMultiply->GetInput(1));
+
+		BaseColorTransparencyMultiply->ConnectExpression(*IncandescenceMultiply->GetInput(0));
+		IncandescenceScaleMultiply->ConnectExpression(*IncandescenceMultiply->GetInput(1));
+
+		TransparencyColorExpression->ConnectExpression(*TransparencyOneMinus->GetInput(0));
+
+		if (bIsTransparent)
+		{
+			TransparencyOneMinus->ConnectExpression(*BreakFloat3->GetInput(0));
+
+			BreakFloat3->ConnectExpression(*AddRG->GetInput(0), 0);
+			BreakFloat3->ConnectExpression(*AddRG->GetInput(1), 1);
+
+			AddRG->ConnectExpression(*AddRGB->GetInput(0));
+			BreakFloat3->ConnectExpression(*AddRGB->GetInput(1), 2);
+
+			AddRGB->ConnectExpression(*Divide->GetInput(0));
+			DivideConstant->ConnectExpression(*Divide->GetInput(1));
+		}
+
+		// Connect material outputs
+		MaterialElement->GetBaseColor().SetExpression(BaseColorTransparencyMultiply);
+		MaterialElement->GetEmissiveColor().SetExpression(IncandescenceMultiply);
+
+		if (bIsTransparent)
+		{
+			MaterialElement->GetOpacity().SetExpression(Divide);
+			MaterialElement->SetParentLabel(TEXT("M_DatasmithAliasLightSourceTransparent"));
+		}
+		else {
+			MaterialElement->SetParentLabel(TEXT("M_DatasmithAliasLightSource"));
+		}
+
+	}
+
+	void FWireTranslatorImpl::AddAlPhongParameters(const TAlObjectPtr<AlShader>& Shader, TSharedPtr<IDatasmithUEPbrMaterialElement> MaterialElement)
+	{
+		// Default values for a Phong material
+		FColor Color(145, 148, 153);
+		FColor TransparencyColor(0, 0, 0);
+		FColor IncandescenceColor(0, 0, 0);
+		FColor SpecularColor(38, 38, 38);
+		double Diffuse = 1.0;
+		double GlowIntensity = 0.0;
+		double Gloss = 0.8;
+		double Shinyness = 20.0;
+		double Specularity = 1.0;
+		double Reflectivity = 0.5;
+
+		AlList* List = Shader->fields();
+		for (AlShadingFieldItem* Item = static_cast<AlShadingFieldItem*>(List->first()); Item; Item = Item->nextField())
+		{
+			double Value = 0.0f;
+			statusCode ErrorCode = Shader->parameter(Item->field(), Value);
+			if (ErrorCode != 0)
+			{
+				continue;
+			}
+
+			if (GetCommonParameters(Item->field(), Value, Color, TransparencyColor, IncandescenceColor, GlowIntensity))
+			{
+				continue;
+			}
+
+			switch (Item->field())
+			{
+			case AlShadingFields::kFLD_SHADING_PHONG_DIFFUSE:
+				Diffuse = Value;
+				break;
+			case AlShadingFields::kFLD_SHADING_PHONG_GLOSS_:
+				Gloss = Value;
+				break;
+			case AlShadingFields::kFLD_SHADING_PHONG_SPECULAR_R:
+				SpecularColor.R = (uint8)(255.f * Value);;
+				break;
+			case AlShadingFields::kFLD_SHADING_PHONG_SPECULAR_G:
+				SpecularColor.G = (uint8)(255.f * Value);;
+				break;
+			case AlShadingFields::kFLD_SHADING_PHONG_SPECULAR_B:
+				SpecularColor.B = (uint8)(255.f * Value);;
+				break;
+			case AlShadingFields::kFLD_SHADING_PHONG_SPECULARITY_:
+				Specularity = Value;
+				break;
+			case AlShadingFields::kFLD_SHADING_PHONG_SHINYNESS:
+				Shinyness = Value;
+				break;
+			case AlShadingFields::kFLD_SHADING_PHONG_REFLECTIVITY:
+				Reflectivity = Value;
+				break;
+			}
+		}
+
+		bool bIsTransparent = IsTransparent(TransparencyColor);
+
+		// Construct parameter expressions
+		IDatasmithMaterialExpressionScalar* DiffuseExpression = MaterialElement->AddMaterialExpression<IDatasmithMaterialExpressionScalar>();
+		DiffuseExpression->GetScalar() = Diffuse;
+		DiffuseExpression->SetName(TEXT("Diffuse"));
+
+		IDatasmithMaterialExpressionScalar* GlossExpression = MaterialElement->AddMaterialExpression<IDatasmithMaterialExpressionScalar>();
+		GlossExpression->GetScalar() = Gloss;
+		GlossExpression->SetName(TEXT("Gloss"));
+
+		IDatasmithMaterialExpressionColor* SpecularColorExpression = MaterialElement->AddMaterialExpression<IDatasmithMaterialExpressionColor>();
+		SpecularColorExpression->SetName(TEXT("SpecularColor"));
+		SpecularColorExpression->GetColor() = FLinearColor::FromSRGBColor(SpecularColor);
+
+		IDatasmithMaterialExpressionScalar* SpecularityExpression = MaterialElement->AddMaterialExpression<IDatasmithMaterialExpressionScalar>();
+		SpecularityExpression->GetScalar() = Specularity * 0.3;
+		SpecularityExpression->SetName(TEXT("Specularity"));
+
+		IDatasmithMaterialExpressionScalar* ShinynessExpression = MaterialElement->AddMaterialExpression<IDatasmithMaterialExpressionScalar>();
+		ShinynessExpression->GetScalar() = Shinyness;
+		ShinynessExpression->SetName(TEXT("Shinyness"));
+
+		IDatasmithMaterialExpressionScalar* ReflectivityExpression = MaterialElement->AddMaterialExpression<IDatasmithMaterialExpressionScalar>();
+		ReflectivityExpression->GetScalar() = Reflectivity;
+		ReflectivityExpression->SetName(TEXT("Reflectivity"));
+
+		IDatasmithMaterialExpressionColor* ColorExpression = MaterialElement->AddMaterialExpression<IDatasmithMaterialExpressionColor>();
+		ColorExpression->SetName(TEXT("Color"));
+		ColorExpression->GetColor() = FLinearColor::FromSRGBColor(Color);
+
+		IDatasmithMaterialExpressionColor* IncandescenceColorExpression = MaterialElement->AddMaterialExpression<IDatasmithMaterialExpressionColor>();
+		IncandescenceColorExpression->SetName(TEXT("IncandescenceColor"));
+		IncandescenceColorExpression->GetColor() = FLinearColor::FromSRGBColor(IncandescenceColor);
+
+		IDatasmithMaterialExpressionColor* TransparencyColorExpression = MaterialElement->AddMaterialExpression<IDatasmithMaterialExpressionColor>();
+		TransparencyColorExpression->SetName(TEXT("TransparencyColor"));
+		TransparencyColorExpression->GetColor() = FLinearColor::FromSRGBColor(TransparencyColor);
+
+		IDatasmithMaterialExpressionScalar* GlowIntensityExpression = MaterialElement->AddMaterialExpression<IDatasmithMaterialExpressionScalar>();
+		GlowIntensityExpression->GetScalar() = GlowIntensity;
+		GlowIntensityExpression->SetName(TEXT("GlowIntensity"));
+
+		// Create aux expressions
+		IDatasmithMaterialExpressionGeneric* ColorSpecLerp = MaterialElement->AddMaterialExpression<IDatasmithMaterialExpressionGeneric>();
+		ColorSpecLerp->SetExpressionName(TEXT("LinearInterpolate"));
+
+		IDatasmithMaterialExpressionScalar* ColorSpecLerpValue = MaterialElement->AddMaterialExpression<IDatasmithMaterialExpressionScalar>();
+		ColorSpecLerpValue->GetScalar() = 0.96f;
+
+		IDatasmithMaterialExpressionGeneric* ColorMetallicLerp = MaterialElement->AddMaterialExpression<IDatasmithMaterialExpressionGeneric>();
+		ColorMetallicLerp->SetExpressionName(TEXT("LinearInterpolate"));
+
+		IDatasmithMaterialExpressionGeneric* DiffuseLerp = MaterialElement->AddMaterialExpression<IDatasmithMaterialExpressionGeneric>();
+		DiffuseLerp->SetExpressionName(TEXT("LinearInterpolate"));
+
+		IDatasmithMaterialExpressionScalar* DiffuseLerpA = MaterialElement->AddMaterialExpression<IDatasmithMaterialExpressionScalar>();
+		DiffuseLerpA->GetScalar() = 0.04f;
+
+		IDatasmithMaterialExpressionScalar* DiffuseLerpB = MaterialElement->AddMaterialExpression<IDatasmithMaterialExpressionScalar>();
+		DiffuseLerpB->GetScalar() = 1.0f;
+
+		IDatasmithMaterialExpressionGeneric* BaseColorMultiply = MaterialElement->AddMaterialExpression<IDatasmithMaterialExpressionGeneric>();
+		BaseColorMultiply->SetExpressionName(TEXT("Multiply"));
+
+		IDatasmithMaterialExpressionGeneric* BaseColorAdd = MaterialElement->AddMaterialExpression<IDatasmithMaterialExpressionGeneric>();
+		BaseColorAdd->SetExpressionName(TEXT("Add"));
+
+		IDatasmithMaterialExpressionGeneric* BaseColorTransparencyMultiply = MaterialElement->AddMaterialExpression<IDatasmithMaterialExpressionGeneric>();
+		BaseColorTransparencyMultiply->SetExpressionName(TEXT("Multiply"));
+
+		IDatasmithMaterialExpressionGeneric* IncandescenceMultiply = MaterialElement->AddMaterialExpression<IDatasmithMaterialExpressionGeneric>();
+		IncandescenceMultiply->SetExpressionName(TEXT("Multiply"));
+
+		IDatasmithMaterialExpressionGeneric* IncandescenceScaleMultiply = MaterialElement->AddMaterialExpression<IDatasmithMaterialExpressionGeneric>();
+		IncandescenceScaleMultiply->SetExpressionName(TEXT("Multiply"));
+
+		IDatasmithMaterialExpressionScalar* IncandescenceScale = MaterialElement->AddMaterialExpression<IDatasmithMaterialExpressionScalar>();
+		IncandescenceScale->GetScalar() = 100.0f;
+
+		IDatasmithMaterialExpressionGeneric* ShinynessSubtract = MaterialElement->AddMaterialExpression<IDatasmithMaterialExpressionGeneric>();
+		ShinynessSubtract->SetExpressionName(TEXT("Subtract"));
+
+		IDatasmithMaterialExpressionScalar* ShinynessSubtract2 = MaterialElement->AddMaterialExpression<IDatasmithMaterialExpressionScalar>();
+		ShinynessSubtract2->GetScalar() = 2.0f;
+
+		IDatasmithMaterialExpressionGeneric* ShinynessDivide = MaterialElement->AddMaterialExpression<IDatasmithMaterialExpressionGeneric>();
+		ShinynessDivide->SetExpressionName(TEXT("Divide"));
+
+		IDatasmithMaterialExpressionScalar* ShinynessDivide98 = MaterialElement->AddMaterialExpression<IDatasmithMaterialExpressionScalar>();
+		ShinynessDivide98->GetScalar() = 98.0f;
+
+		IDatasmithMaterialExpressionGeneric* SpecularityMultiply = MaterialElement->AddMaterialExpression<IDatasmithMaterialExpressionGeneric>();
+		SpecularityMultiply->SetExpressionName(TEXT("Multiply"));
+
+		IDatasmithMaterialExpressionGeneric* RoughnessOneMinus = MaterialElement->AddMaterialExpression<IDatasmithMaterialExpressionGeneric>();
+		RoughnessOneMinus->SetExpressionName(TEXT("OneMinus"));
+
+		IDatasmithMaterialExpressionGeneric* TransparencyOneMinus = MaterialElement->AddMaterialExpression<IDatasmithMaterialExpressionGeneric>();
+		TransparencyOneMinus->SetExpressionName(TEXT("OneMinus"));
+
+		IDatasmithMaterialExpressionFunctionCall* BreakFloat3 = nullptr;
+		IDatasmithMaterialExpressionGeneric* AddRG = nullptr;
+		IDatasmithMaterialExpressionGeneric* AddRGB = nullptr;
+		IDatasmithMaterialExpressionGeneric* Divide = nullptr;
+		IDatasmithMaterialExpressionScalar* DivideConstant = nullptr;
+		if (bIsTransparent)
+		{
+			BreakFloat3 = MaterialElement->AddMaterialExpression<IDatasmithMaterialExpressionFunctionCall>();
+			BreakFloat3->SetFunctionPathName(TEXT("/Engine/Functions/Engine_MaterialFunctions02/Utility/BreakFloat3Components.BreakFloat3Components"));
+
+			AddRG = MaterialElement->AddMaterialExpression<IDatasmithMaterialExpressionGeneric>();
+			AddRG->SetExpressionName(TEXT("Add"));
+
+			AddRGB = MaterialElement->AddMaterialExpression<IDatasmithMaterialExpressionGeneric>();
+			AddRGB->SetExpressionName(TEXT("Add"));
+
+			Divide = MaterialElement->AddMaterialExpression<IDatasmithMaterialExpressionGeneric>();
+			Divide->SetExpressionName(TEXT("Divide"));
+
+			DivideConstant = MaterialElement->AddMaterialExpression<IDatasmithMaterialExpressionScalar>();
+			DivideConstant->GetScalar() = 3.0f;
+		}
+
+		// Connect expressions
+		SpecularColorExpression->ConnectExpression(*ColorSpecLerp->GetInput(0));
+		ColorExpression->ConnectExpression(*ColorSpecLerp->GetInput(1));
+		ColorSpecLerpValue->ConnectExpression(*ColorSpecLerp->GetInput(2));
+
+		ColorExpression->ConnectExpression(*ColorMetallicLerp->GetInput(0));
+		ColorSpecLerp->ConnectExpression(*ColorMetallicLerp->GetInput(1));
+		GlossExpression->ConnectExpression(*ColorMetallicLerp->GetInput(2));
+
+		DiffuseLerpA->ConnectExpression(*DiffuseLerp->GetInput(0));
+		DiffuseLerpB->ConnectExpression(*DiffuseLerp->GetInput(1));
+		DiffuseExpression->ConnectExpression(*DiffuseLerp->GetInput(2));
+
+		ColorMetallicLerp->ConnectExpression(*BaseColorMultiply->GetInput(0));
+		DiffuseLerp->ConnectExpression(*BaseColorMultiply->GetInput(1));
+
+		BaseColorMultiply->ConnectExpression(*BaseColorAdd->GetInput(0));
+		IncandescenceColorExpression->ConnectExpression(*BaseColorAdd->GetInput(1));
+
+		BaseColorAdd->ConnectExpression(*BaseColorTransparencyMultiply->GetInput(0));
+		TransparencyOneMinus->ConnectExpression(*BaseColorTransparencyMultiply->GetInput(1));
+
+		GlowIntensityExpression->ConnectExpression(*IncandescenceScaleMultiply->GetInput(0));
+		IncandescenceScale->ConnectExpression(*IncandescenceScaleMultiply->GetInput(1));
+
+		BaseColorTransparencyMultiply->ConnectExpression(*IncandescenceMultiply->GetInput(0));
+		IncandescenceScaleMultiply->ConnectExpression(*IncandescenceMultiply->GetInput(1));
+
+		ShinynessExpression->ConnectExpression(*ShinynessSubtract->GetInput(0));
+		ShinynessSubtract2->ConnectExpression(*ShinynessSubtract->GetInput(1));
+
+		ShinynessSubtract->ConnectExpression(*ShinynessDivide->GetInput(0));
+		ShinynessDivide98->ConnectExpression(*ShinynessDivide->GetInput(1));
+
+		ShinynessDivide->ConnectExpression(*SpecularityMultiply->GetInput(0));
+		SpecularityExpression->ConnectExpression(*SpecularityMultiply->GetInput(1));
+
+		SpecularityMultiply->ConnectExpression(*RoughnessOneMinus->GetInput(0));
+
+		TransparencyColorExpression->ConnectExpression(*TransparencyOneMinus->GetInput(0));
+
+		if (bIsTransparent)
+		{
+			TransparencyOneMinus->ConnectExpression(*BreakFloat3->GetInput(0));
+
+			BreakFloat3->ConnectExpression(*AddRG->GetInput(0), 0);
+			BreakFloat3->ConnectExpression(*AddRG->GetInput(1), 1);
+
+			AddRG->ConnectExpression(*AddRGB->GetInput(0));
+			BreakFloat3->ConnectExpression(*AddRGB->GetInput(1), 2);
+
+			AddRGB->ConnectExpression(*Divide->GetInput(0));
+			DivideConstant->ConnectExpression(*Divide->GetInput(1));
+		}
+
+		// Connect material outputs
+		MaterialElement->GetBaseColor().SetExpression(BaseColorTransparencyMultiply);
+		MaterialElement->GetMetallic().SetExpression(GlossExpression);
+		MaterialElement->GetSpecular().SetExpression(ReflectivityExpression);
+		MaterialElement->GetRoughness().SetExpression(RoughnessOneMinus);
+		MaterialElement->GetEmissiveColor().SetExpression(IncandescenceMultiply);
+		if (bIsTransparent)
+		{
+			MaterialElement->GetOpacity().SetExpression(Divide);
+			MaterialElement->SetParentLabel(TEXT("M_DatasmithAliasPhongTransparent"));
+		}
+		else {
+			MaterialElement->SetParentLabel(TEXT("M_DatasmithAliasPhong"));
+		}
+	}
+
+	// Make material
+	bool FWireTranslatorImpl::GetShader()
+	{
+		for (TAlObjectPtr<AlShader> Shader(AlUniverse::firstShader()); Shader.IsValid(); Shader = TAlObjectPtr<AlShader>(AlUniverse::nextShader(Shader.Get())))
+		{
+			const FString ShaderName = UTF8_TO_TCHAR(Shader->name());
+			const FString ShaderModelName = Shader->shadingModel();
+
+			const FColor Color = CreateShaderColorFromShaderName(ShaderName);
+			ShaderNameToColor.Add(ShaderName, Color);
+
+			const int32 ShaderId = CreateShaderId(Color);
+			const FString MaterialElementName = FString::FromInt(ShaderId);
+
+			TSharedPtr<IDatasmithUEPbrMaterialElement> MaterialElement = FDatasmithSceneFactory::CreateUEPbrMaterial(*MaterialElementName);
+			MaterialElement->SetLabel(*ShaderName);
+
+			if (ShaderModelName.Equals(TEXT("BLINN")))
+			{
+				AddAlBlinnParameters(Shader, MaterialElement);
+			}
+			else if (ShaderModelName.Equals(TEXT("LAMBERT")))
+			{
+				AddAlLambertParameters(Shader, MaterialElement);
+			}
+			else if (ShaderModelName.Equals(TEXT("LIGHTSOURCE")))
+			{
+				AddAlLightSourceParameters(Shader, MaterialElement);
+			}
+			else if (ShaderModelName.Equals(TEXT("PHONG")))
+			{
+				AddAlPhongParameters(Shader, MaterialElement);
+			}
+
+			DatasmithScene->AddMaterial(MaterialElement);
+
+			TSharedPtr<IDatasmithMaterialIDElement> MaterialIDElement = FDatasmithSceneFactory::CreateMaterialId(MaterialElement->GetName());
+			ShaderNameToUEMaterialId.Add(*ShaderName, MaterialIDElement);
+		}
+		return true;
+	}
+
+	bool FWireTranslatorImpl::GetDagLeaves()
+	{
+		FDagNodeInfo RootContainer;
+		AlRootNode = TAlObjectPtr<AlDagNode>(AlUniverse::firstDagNode());
+		if (!AlRootNode.IsValid())
+		{
+			return false;
+		}
+
+		if (TessellationOptions.StitchingTechnique != EDatasmithCADStitchingTechnique::StitchingSew)
+		{
+			RecurseDagForLeavesNoMerge(AlRootNode, RootContainer);
+		}
+		else
+		{
+			RecurseDagForLeaves(AlRootNode, RootContainer);
+		}
+
+		// Delete empty branches
+		TArray<TSharedPtr<IDatasmithActorElement>> ActorsToRemove;
+		for (int32 Index = DatasmithScene->GetActorsCount() - 1; Index >= 0; --Index)
+		{
+			TSharedPtr<IDatasmithActorElement> Actor = DatasmithScene->GetActor(Index);
+			if (Actor.IsValid() && !Actor->IsA(EDatasmithElementType::StaticMeshActor))
+			{
+				RecurseDeleteEmptyActor(Actor);
+				if (Actor->GetChildrenCount() == 0)
+				{
+					ActorsToRemove.Add(Actor);
+				}
+			}
+		}
+
+		for (const TSharedPtr<IDatasmithActorElement>& Actor : ActorsToRemove)
+		{
+			DatasmithScene->RemoveActor(Actor, EDatasmithActorRemovalRule::RemoveChildren);
+		}
+
+		return true;
+	}
+
+	void FWireTranslatorImpl::RecurseDeleteEmptyActor(TSharedPtr<IDatasmithActorElement> Actor)
+	{
+		for (int32 Index = Actor->GetChildrenCount() - 1; Index >= 0; --Index)
+		{
+			TSharedPtr<IDatasmithActorElement> ActorChild = Actor->GetChild(Index);
+			if (ActorChild.IsValid() && !ActorChild->IsA(EDatasmithElementType::StaticMeshActor))
+			{
+				RecurseDeleteEmptyActor(ActorChild);
+				if (ActorChild->GetChildrenCount() == 0)
+				{
+					Actor->RemoveChild(ActorChild);
+				}
+			}
+		}
+	}
+
+
+	void FDagNodeInfo::ExtractDagNodeLayer(const AlDagNode& InDagNode)
+	{
+		AlLayer* LayerPtr = InDagNode.layer();
+		if (LayerPtr)
+		{
+			TAlObjectPtr<AlLayer> Layer(LayerPtr);
+			LayerName = UTF8_TO_TCHAR(Layer->name());
+			bLayerIsVisible = !Layer->invisible();
+			if (ActorElement.IsValid())
+			{
+				ActorElement->SetLayer(*LayerName);
+				LayerToActorElement.Add(LayerName, ActorElement);
+			}
+		}
+	}
+
+	void FDagNodeInfo::ExtractDagNodeInfo(AlDagNode& CurrentNode)
+	{
+		Label = UTF8_TO_TCHAR(CurrentNode.name());
+		uint32 ThisGroupNodeUuid = OpenModelUtils::GetAlDagNodeUuid(CurrentNode);
+		Uuid = HashCombine(ParentInfo ? ParentInfo->Uuid : 1, ThisGroupNodeUuid);
+	}
+
+	void FDagNodeInfo::ExtractDagNodeInfo(FBodyData& CurrentNode)
+	{
+		if (!ParentInfo)
+		{
+			return;
+		}
+		Label = ParentInfo->Label;
+		CurrentNode.SetLabel(ParentInfo->Label);
+		Uuid = CurrentNode.GetUuid(ParentInfo->Uuid);
+	}
+
+	TSharedPtr<IDatasmithActorElement> FWireTranslatorImpl::FindOrAddLayerActor(const FString& LayerName)
+	{
+		if (!bGLayersAsActors || LayerName.IsEmpty())
+		{
+			return TSharedPtr<IDatasmithActorElement>();
+		}
+
+		TSharedPtr<IDatasmithActorElement>* LayerActor = LayerNameToActor.Find(LayerName);
+		if (LayerActor)
+		{
+			return *LayerActor;
+		}
+
+		uint32 LayerUuid = GetTypeHash(LayerName);
+		TSharedPtr<IDatasmithActorElement> NewLayerActor = FDatasmithSceneFactory::CreateActor(*OpenModelUtils::UuidToString(LayerUuid));
+		NewLayerActor->SetLabel(*LayerName);
+		NewLayerActor->SetLayer(*LayerName);
+		DatasmithScene->AddActor(NewLayerActor);
+		LayerNameToActor.Add(LayerName, NewLayerActor);
+		return NewLayerActor;
+	}
+
+	TSharedPtr<IDatasmithActorElement> FWireTranslatorImpl::FindOrAddParentActor(FDagNodeInfo& ParentInfo, const FString& LayerName)
+	{
+		if (!ParentInfo.ActorElement.IsValid())
+		{
+			if (!LayerName.IsEmpty())
+			{
+				return FindOrAddLayerActor(LayerName);
+			}
+			return TSharedPtr<IDatasmithActorElement>();
+		}
+
+		if (!bGLayersAsActors)
+		{
+			return ParentInfo.ActorElement;
+		}
+
+		TSharedPtr<IDatasmithActorElement>* ActorParent = ParentInfo.LayerToActorElement.Find(LayerName);
+		if (ActorParent)
+		{
+			return *ActorParent;
+		}
+
+		uint32 LayerActorUuid = HashCombine(ParentInfo.Uuid, GetTypeHash(LayerName));
+
+		TSharedPtr<IDatasmithActorElement> LayerActorParent = FDatasmithSceneFactory::CreateActor(*OpenModelUtils::UuidToString(LayerActorUuid));
+		LayerActorParent->SetLabel(*ParentInfo.Label);
+		LayerActorParent->SetLayer(*LayerName);
+
+		ParentInfo.LayerToActorElement.Add(LayerName, LayerActorParent);
+
+		TSharedPtr<IDatasmithActorElement> GrandParentActor = FindOrAddParentActor(*ParentInfo.ParentInfo, LayerName);
+		if (GrandParentActor.IsValid())
+		{
+			GrandParentActor->AddChild(LayerActorParent);
+		}
+		else
+		{
+			TSharedPtr<IDatasmithActorElement> LayerActor = FindOrAddLayerActor(LayerName);
+			LayerActor->AddChild(LayerActorParent);
+		}
+
+		return LayerActorParent;
+	}
+
+	bool FWireTranslatorImpl::ProcessAlGroupNode(AlDagNode& GroupNode, FDagNodeInfo& ParentInfo)
+	{
+		AlGroupNode* AlGroup = GroupNode.asGroupNodePtr();
+		if (!AlIsValid(AlGroup))
+		{
+			return false;
+		}
+
+		AlDagNode* AlChildPtr = AlGroup->childNode();
+		if (!AlIsValid(AlChildPtr))
+		{
+			return false;
+		}
+
+		TAlObjectPtr<AlDagNode> ChildNode(AlChildPtr);
+
+		FDagNodeInfo ThisGroupNodeInfo(ParentInfo);
+		ThisGroupNodeInfo.ExtractDagNodeInfo(GroupNode);
+
+		ThisGroupNodeInfo.ActorElement = FDatasmithSceneFactory::CreateActor(*OpenModelUtils::UuidToString(ThisGroupNodeInfo.Uuid));
+		ThisGroupNodeInfo.ActorElement->SetLabel(*ThisGroupNodeInfo.Label);
+		ThisGroupNodeInfo.ExtractDagNodeLayer(GroupNode);
+
+		TSharedPtr<IDatasmithActorElement> ParentActorElement = FindOrAddParentActor(ParentInfo, ThisGroupNodeInfo.LayerName);
+		if (ParentActorElement.IsValid())
+		{
+			ParentActorElement->AddChild(ThisGroupNodeInfo.ActorElement);
+		}
+		else
+		{
+			DatasmithScene->AddActor(ThisGroupNodeInfo.ActorElement);
+		}
+
+		RecurseDagForLeaves(ChildNode, ThisGroupNodeInfo);
+
+		// Apply local transform to actor element
+		OpenModelUtils::SetActorTransform(ThisGroupNodeInfo.ActorElement, GroupNode);
+
+		return true;
+	}
+
+	TSharedPtr<IDatasmithMeshElement> FWireTranslatorImpl::FindOrAddMeshElement(const TSharedPtr<FBodyData>& Body, const FDagNodeInfo& NodeInfo)
+	{
+		TSharedPtr<IDatasmithMeshElement>* MeshElementPtr = BodyUuidToMeshElementMap.Find(NodeInfo.Uuid);
+		if (MeshElementPtr != nullptr)
+		{
+			return *MeshElementPtr;
+		}
+
+		TSharedPtr<IDatasmithMeshElement> MeshElement = FDatasmithSceneFactory::CreateMesh(*OpenModelUtils::UuidToString(NodeInfo.Uuid));
+		MeshElement->SetLabel(*NodeInfo.Label);
+		MeshElement->SetLightmapSourceUV(-1);
+
+		const TSet<FString>& ShaderNames = Body->GetShaderNames();
+		for (const FString& ShaderName : ShaderNames)
+		{
+			TSharedPtr<IDatasmithMaterialIDElement> MaterialElement = ShaderNameToUEMaterialId[ShaderName];
+			if (MaterialElement.IsValid())
+			{
+				const FColor* ShaderColor = ShaderNameToColor.Find(ShaderName);
+				if (ShaderColor)
+				{
+					int32 ShaderId = CreateShaderId(*ShaderColor);
+					MeshElement->SetMaterial(MaterialElement->GetName(), ShaderId);
+				}
+			}
+		}
+
+		DatasmithScene->AddMesh(MeshElement);
+
+		ShellUuidToMeshElementMap.Add(NodeInfo.Uuid, MeshElement);
+		MeshElementToBodyMap.Add(MeshElement.Get(), Body);
+
+		BodyUuidToMeshElementMap.Add(NodeInfo.Uuid, MeshElement);
+
+		return MeshElement;
+	}
+
+	TSharedPtr<IDatasmithMeshElement> FWireTranslatorImpl::FindOrAddMeshElement(const TAlObjectPtr<AlDagNode>& ShellNode, const FDagNodeInfo& ShellNodeInfo, const FString& ShaderName)
+	{
+		uint32 ShellUuid = OpenModelUtils::GetAlDagNodeUuid(*ShellNode);
+
+		// Look if geometry has not been already processed, return it if found
+		TSharedPtr<IDatasmithMeshElement>* MeshElementPtr = ShellUuidToMeshElementMap.Find(ShellUuid);
+		if (MeshElementPtr != nullptr)
+		{
+			return *MeshElementPtr;
+		}
+
+		TSharedPtr<IDatasmithMeshElement> MeshElement = FDatasmithSceneFactory::CreateMesh(*OpenModelUtils::UuidToString(ShellNodeInfo.Uuid));
+		MeshElement->SetLabel(*ShellNodeInfo.Label);
+		MeshElement->SetLightmapSourceUV(-1);
+
+		// Set MeshElement FileHash used for re-import task
+		FMD5 MD5; // unique Value that define the mesh
+		MD5.Update(reinterpret_cast<const uint8*>(&SceneFileHash), sizeof SceneFileHash);
+		// MeshActor Name
+		MD5.Update(reinterpret_cast<const uint8*>(&ShellUuid), sizeof ShellUuid);
+		FMD5Hash Hash;
+		Hash.Set(MD5);
+		MeshElement->SetFileHash(Hash);
+
+		if (ShaderName.Len())
+		{
+			TSharedPtr<IDatasmithMaterialIDElement> MaterialElement = ShaderNameToUEMaterialId[ShaderName];
+			if (MaterialElement.IsValid())
+			{
+				const FColor* ShaderColor = ShaderNameToColor.Find(ShaderName);
+				if (ShaderColor)
+				{
+					int32 ShaderId = CreateShaderId(*ShaderColor);
+					MeshElement->SetMaterial(MaterialElement->GetName(), ShaderId);
+				}
+			}
+
+		}
+
+		DatasmithScene->AddMesh(MeshElement);
+
+		ShellUuidToMeshElementMap.Add(ShellUuid, MeshElement);
+		MeshElementToAlDagNodeMap.Add(MeshElement.Get(), ShellNode);
+		MeshElementToShaderName.Add(MeshElement.Get(), ShaderName);
+
+		return MeshElement;
+	}
+
+	void FWireTranslatorImpl::ProcessAlShellNode(const TAlObjectPtr<AlDagNode>& ShellNode, FDagNodeInfo& ParentInfo, const FString& ShaderName)
+	{
+		FDagNodeInfo ShellInfo(ParentInfo);
+		ShellInfo.ExtractDagNodeInfo(*ShellNode);
+
+		TSharedPtr<IDatasmithMeshElement> MeshElement = FindOrAddMeshElement(ShellNode, ShellInfo, ShaderName);
+		if (!MeshElement.IsValid())
+		{
+			return;
+		}
+
+		TSharedPtr<IDatasmithMeshActorElement> ActorElement = FDatasmithSceneFactory::CreateMeshActor(*OpenModelUtils::UuidToString(ShellInfo.Uuid));
+		if (!ActorElement.IsValid())
+		{
+			return;
+		}
+
+		ActorElement->SetLabel(*ShellInfo.Label);
+		ActorElement->SetStaticMeshPathName(MeshElement->GetName());
+		ShellInfo.ActorElement = ActorElement;
+
+		ShellInfo.ExtractDagNodeLayer(*ShellNode);
+		if (!ShellInfo.bLayerIsVisible)
+		{
+			return;
+		}
+
+		// Apply materials on the current part
+		if (ShaderName.Len())
+		{
+			const TSharedPtr<IDatasmithMaterialIDElement>& MaterialIDElement = ShaderNameToUEMaterialId[ShaderName];
+			if (MaterialIDElement.IsValid())
+			{
+				for (int32 Index = 0; Index < MeshElement->GetMaterialSlotCount(); ++Index)
+				{
+					MaterialIDElement->SetId(MeshElement->GetMaterialSlotAt(Index)->GetId());
+					ActorElement->AddMaterialOverride(MaterialIDElement);
+				}
+			}
+		}
+
+		TSharedPtr<IDatasmithActorElement> ParentElement = FindOrAddParentActor(ParentInfo, ShellInfo.LayerName);
+		if (ParentElement.IsValid())
+		{
+			ParentElement->AddChild(ActorElement);
+		}
+		else
+		{
+			DatasmithScene->AddActor(ActorElement);
+		}
+
+		OpenModelUtils::SetActorTransform(ShellInfo.ActorElement, *ShellNode);
+	}
+
+	void FWireTranslatorImpl::ProcessBodyNode(const TSharedPtr<FBodyData>& Body, FDagNodeInfo& ParentInfo)
+	{
+
+		if (!Body.IsValid())
+		{
+			return;
+		}
+
+		if (Body->GetShells().Num() == 1)
+		{
+			DagForLeavesNoMerge(Body->GetShells()[0].Key, ParentInfo);
+			return;
+		}
+
+		FDagNodeInfo ShellInfo(ParentInfo);
+		ShellInfo.ExtractDagNodeInfo(*Body);
+
+		TSharedPtr<IDatasmithMeshElement> MeshElement = FindOrAddMeshElement(Body, ShellInfo);
+		if (!MeshElement.IsValid())
+		{
+			return;
+		}
+
+		TSharedPtr<IDatasmithMeshActorElement> ActorElement = FDatasmithSceneFactory::CreateMeshActor(*OpenModelUtils::UuidToString(ShellInfo.Uuid));
+		if (!ActorElement.IsValid())
+		{
+			return;
+		}
+
+		ActorElement->SetLabel(ShellInfo.Label.IsEmpty() ? *Body->GetLayerName() : *ShellInfo.Label);
+		ActorElement->SetStaticMeshPathName(MeshElement->GetName());
+		ShellInfo.ActorElement = ActorElement;
+
+		ActorElement->SetLayer(*Body->GetLayerName());
+
+		// Apply materials on the current part
+		int32 Index = 0;
+		const TSet<FString>& ShaderNames = Body->GetShaderNames();
+		for (const FString& ShaderName : ShaderNames)
+		{
+			TSharedPtr<IDatasmithMaterialIDElement> MaterialIDElement = ShaderNameToUEMaterialId[ShaderName];
+			if (MaterialIDElement.IsValid())
+			{
+				MaterialIDElement->SetId(MeshElement->GetMaterialSlotAt(Index)->GetId());
+				ActorElement->AddMaterialOverride(MaterialIDElement);
+				Index++;
+			}
+		}
+
+		TSharedPtr<IDatasmithActorElement> ParentElement = FindOrAddParentActor(ParentInfo, Body->GetLayerName());
+		if (ParentElement.IsValid())
+		{
+			ParentElement->AddChild(ActorElement);
+		}
+		else
+		{
+			DatasmithScene->AddActor(ActorElement);
+		}
+
+		return;
+	}
+
+	TAlObjectPtr<AlDagNode> GetNextNode(const TAlObjectPtr<AlDagNode>& DagNode)
+	{
+		// Grab the next sibling before deleting the node.
+		AlDagNode* SiblingNode = DagNode->nextNode();
+		if (AlIsValid(SiblingNode))
+		{
+			return TAlObjectPtr<AlDagNode>(SiblingNode);
+		}
+		else
+		{
+			return TAlObjectPtr<AlDagNode>();
+		}
+	}
+
+	uint32 GetBodyGroupUuid(const FString& ShaderName, const FString& LayerName, bool bCadData)
+	{
+		uint32 Uuid = HashCombine(GetTypeHash(LayerName), GetTypeHash(bCadData));
+		if (bGSewByMaterial)
+		{
+			Uuid = HashCombine(GetTypeHash(ShaderName), Uuid);
+		}
+		return Uuid;
+	}
+
+	uint32 GetPatchCount(const TAlObjectPtr<AlShell>& Shell)
+	{
+		uint32 PatchCount = 0;
+		TAlObjectPtr<AlTrimRegion> TrimRegion(Shell->firstTrimRegion());
+		while (TrimRegion.IsValid())
+		{
+			PatchCount++;
+			TrimRegion = TAlObjectPtr<AlTrimRegion>(TrimRegion->nextRegion());
+		}
+		return PatchCount;
+	}
+
+	void FWireTranslatorImpl::AddNodeInBodyGroup(TAlObjectPtr<AlDagNode>& DagNode, const FString& ShaderName, TMap<uint32, TSharedPtr<FBodyData>>& ShellToProcess, bool bIsAPatch, uint32 MaxSize)
+	{
+		FString LayerName;
+
+		AlLayer* LayerPtr = DagNode->layer();
+		if (AlIsValid(LayerPtr))
+		{
+			const TAlObjectPtr<AlLayer> Layer(LayerPtr);
+			LayerName = UTF8_TO_TCHAR(Layer->name());
+			if (Layer->invisible())
+			{
+				return;
+			}
+		}
+
+		uint32 SetId = GetBodyGroupUuid(ShaderName, LayerName, bIsAPatch);
+
+		TSharedPtr<FBodyData> Body;
+		if (TSharedPtr<FBodyData>* PBody = ShellToProcess.Find(SetId))
+		{
+			Body = *PBody;
+		}
+		else
+		{
+			TSharedPtr<FBodyData> BodyRef = MakeShared<FBodyData>(LayerName, bIsAPatch);
+			ShellToProcess.Add(SetId, BodyRef);
+			BodyRef->Reserve(MaxSize);
+			Body = BodyRef;
+		}
+
+		const FColor& Color = ShaderNameToColor[ShaderName];
+		Body->Add(DagNode, Color, ShaderName);
+	}
+
+	void FWireTranslatorImpl::RecurseDagForLeaves(const TAlObjectPtr<AlDagNode>& FirstDagNode, FDagNodeInfo& ParentInfo)
+	{
+		TAlObjectPtr<AlDagNode> DagNode = FirstDagNode;
+		uint32 MaxSize = 0;
+		while (DagNode)
+		{
+			MaxSize++;
+			DagNode = GetNextNode(DagNode);
+		}
+
+		DagNode = FirstDagNode;
+
+		TMap<uint32, TSharedPtr<FBodyData>> ShellToProcess;
+
+		const char* ShaderName = nullptr;
+
+		while (DagNode)
+		{
+			AlObjectType objectType = DagNode->type();
+
+			// Process the current node.
+			switch (objectType)
+			{
+				// Push all leaf nodes into 'leaves'
+			case kShellNodeType:
+			{
+				AlShellNode* ShellNode = DagNode->asShellNodePtr();
+				if (!AlIsValid(ShellNode))
+				{
+					break;
+				}
+
+				AlShell* ShellPtr = ShellNode->shell();
+				if (!AlIsValid(ShellPtr))
+				{
+					break;
+				}
+
+				TAlObjectPtr<AlShell> Shell(ShellPtr);
+				uint32 NbPatch = GetPatchCount(Shell);
+
+				TAlObjectPtr<AlShader> Shader(Shell->firstShader());
+				if (Shader.IsValid())
+				{
+					ShaderName = Shader->name();
+				}
+
+				if (NbPatch == 1)
+				{
+					AddNodeInBodyGroup(DagNode, ShaderName, ShellToProcess, true, MaxSize);
+				}
+				else
+				{
+					ProcessAlShellNode(DagNode, ParentInfo, ShaderName);
+				}
+				break;
+			}
+
+			case kSurfaceNodeType:
+			{
+				AlSurfaceNode* SurfaceNode = DagNode->asSurfaceNodePtr();
+				TAlObjectPtr<AlSurface> Surface(SurfaceNode->surface());
+				if (Surface.IsValid())
+				{
+					TAlObjectPtr<AlShader> Shader(Surface->firstShader());
+					if (Shader.IsValid())
+					{
+						ShaderName = Shader->name();
+					}
+				}
+				AddNodeInBodyGroup(DagNode, ShaderName, ShellToProcess, true, MaxSize);
+				break;
+			}
+
+			case kMeshNodeType:
+			{
+				AlMeshNode* MeshNode = DagNode->asMeshNodePtr();
+				TAlObjectPtr<AlMesh> Mesh(MeshNode->mesh());
+				if (Mesh.IsValid())
+				{
+					TAlObjectPtr<AlShader> Shader(Mesh->firstShader());
+					if (Shader.IsValid())
+					{
+						ShaderName = Shader->name();
+					}
+				}
+				AddNodeInBodyGroup(DagNode, ShaderName, ShellToProcess, false, MaxSize);
+				break;
+			}
+
+			// Traverse down through groups
+			case kGroupNodeType:
+			{
+				ProcessAlGroupNode(*DagNode, ParentInfo);
+				break;
+			}
+
+			default:
+			{
+				break;
+			}
+			}
+
+			DagNode = GetNextNode(DagNode);
+		}
+
+		for (const TPair<uint32, TSharedPtr<FBodyData>>& Body : ShellToProcess)
+		{
+			ProcessBodyNode(Body.Value, ParentInfo);
+		}
+	}
+
+	void FWireTranslatorImpl::RecurseDagForLeavesNoMerge(const TAlObjectPtr<AlDagNode>& FirstDagNode, FDagNodeInfo& ParentInfo)
+	{
+		TAlObjectPtr<AlDagNode> DagNode = FirstDagNode;
+		while (DagNode)
+		{
+			DagForLeavesNoMerge(DagNode, ParentInfo);
+			DagNode = GetNextNode(DagNode);
+		}
+	}
+
+	void FWireTranslatorImpl::DagForLeavesNoMerge(const TAlObjectPtr<AlDagNode>& DagNode, FDagNodeInfo& ParentInfo)
+	{
+		const char* ShaderName = nullptr;
+
+		// Process the current node.
+		AlObjectType objectType = DagNode->type();
+		switch (objectType)
+		{
+			// Push all leaf nodes into 'leaves'
+		case kShellNodeType:
+		{
+			AlShellNode* ShellNode = DagNode->asShellNodePtr();
+			TAlObjectPtr<AlShell> Shell(ShellNode->shell());
+			if (Shell.IsValid())
+			{
+				TAlObjectPtr<AlShader> Shader(Shell->firstShader());
+				if (Shader.IsValid())
+				{
+					ShaderName = Shader->name();
+				}
+
+				ProcessAlShellNode(DagNode, ParentInfo, ShaderName);
+			}
+			break;
+		}
+		case kSurfaceNodeType:
+		{
+			AlSurfaceNode* SurfaceNode = DagNode->asSurfaceNodePtr();
+			TAlObjectPtr<AlSurface> Surface(SurfaceNode->surface());
+			if (Surface.IsValid())
+			{
+				TAlObjectPtr<AlShader> Shader(Surface->firstShader());
+				if (Shader.IsValid())
+				{
+					ShaderName = Shader->name();
+				}
+			}
+			ProcessAlShellNode(DagNode, ParentInfo, ShaderName);
+			break;
+		}
+
+		case kMeshNodeType:
+		{
+			AlMeshNode* MeshNode = DagNode->asMeshNodePtr();
+			TAlObjectPtr<AlMesh> Mesh(MeshNode->mesh());
+			if (Mesh.IsValid())
+			{
+				TAlObjectPtr<AlShader> Shader(Mesh->firstShader());
+				if (Shader.IsValid())
+				{
+					ShaderName = Shader->name();
+				}
+			}
+			ProcessAlShellNode(DagNode, ParentInfo, ShaderName);
+			break;
+		}
+
+		// Traverse down through groups
+		case kGroupNodeType:
+		{
+			ProcessAlGroupNode(*DagNode, ParentInfo);
+			break;
+		}
+
+		default:
+		{
+			break;
+		}
+
+		}
+	}
+
+	TSharedPtr<CADLibrary::ICADModelConverter> FWireTranslatorImpl::GetModelConverter(IAliasBRepConverter*& BRepConverter) const
+	{
+		TSharedPtr<CADLibrary::ICADModelConverter> ModelConverter;
+
+		CADLibrary::FImportParameters ImportParameters;
+		if (CADLibrary::FImportParameters::bGDisableCADKernelTessellation)
+		{
+			TSharedPtr<FAliasModelToTechSoftConverter> AliasToTechSoftConverter = MakeShared<FAliasModelToTechSoftConverter>(ImportParameters);
+			ModelConverter = AliasToTechSoftConverter;
+			BRepConverter = AliasToTechSoftConverter.Get();
+		}
+		else
+		{
+			TSharedPtr<FAliasModelToCADKernelConverter> AliasToCADKernelConverter = MakeShared<FAliasModelToCADKernelConverter>(ImportParameters);
+			ModelConverter = AliasToCADKernelConverter;
+			BRepConverter = AliasToCADKernelConverter.Get();
+		}
+
+		ModelConverter->SetImportParameters(TessellationOptions.ChordTolerance, TessellationOptions.MaxEdgeLength, TessellationOptions.NormalTolerance, (CADLibrary::EStitchingTechnique)TessellationOptions.StitchingTechnique);
+
+		return ModelConverter;
+	}
+
+	TOptional<FMeshDescription> FWireTranslatorImpl::MeshDagNodeWithExternalMesher(AlDagNode& DagNode, TSharedPtr<IDatasmithMeshElement> MeshElement, CADLibrary::FMeshParameters& MeshParameters)
+	{
+		// Wire unit is cm
+		IAliasBRepConverter* BRepConverter = nullptr;
+		TSharedPtr<CADLibrary::ICADModelConverter> ModelConverter = GetModelConverter(BRepConverter);
+
+		ModelConverter->InitializeProcess();
+
+		EAliasObjectReference ObjectReference = EAliasObjectReference::LocalReference;
+
+		if (MeshParameters.bIsSymmetric)
+		{
+			// All actors of a Alias symmetric layer are defined in the world Reference i.e. they have identity transform. So Mesh actor has to be defined in the world reference.
+			ObjectReference = EAliasObjectReference::WorldReference;
+		}
+
+		const FColor* ColorPtr = nullptr;
+		const FString* ShaderName = MeshElementToShaderName.Find(MeshElement.Get());
+		if (ShaderName)
+		{
+			ColorPtr = ShaderNameToColor.Find(*ShaderName);
+		}
+		if (!ColorPtr) // Should never happen
+		{
+			ColorPtr = &DefaultColor;
+		}
+
+		if (!BRepConverter->AddBRep(DagNode, *ColorPtr, ObjectReference))
+		{
+			return TOptional<FMeshDescription>();
+		}
+
+		ModelConverter->RepairTopology();
+
+		CADModelConverter->SaveModel(*OutputPath, MeshElement);
+
+		FMeshDescription MeshDescription;
+		DatasmithMeshHelper::PrepareAttributeForStaticMesh(MeshDescription);
+
+		bool bRet = ModelConverter->Tessellate(MeshParameters, MeshDescription);
+		if (!bRet)
+		{
+			const TCHAR* StaticMeshLable = MeshElement->GetLabel();
+			const TCHAR* StaticMeshName = MeshElement->GetName();
+			UE_LOG(LogWireInterface, Warning, TEXT("Failed to generate the mesh of \"%s\" (%s) StaticMesh."), StaticMeshLable, StaticMeshName);
+		}
+
+		return MoveTemp(MeshDescription);
+	}
+
+	TOptional<FMeshDescription> FWireTranslatorImpl::MeshDagNodeWithExternalMesher(TSharedPtr<FBodyData> Body, TSharedPtr<IDatasmithMeshElement> MeshElement, CADLibrary::FMeshParameters& MeshParameters)
+	{
+		// Wire unit is cm
+		IAliasBRepConverter* BRepConverter = nullptr;
+		TSharedPtr<CADLibrary::ICADModelConverter> ModelConverter = GetModelConverter(BRepConverter);
+
+		ModelConverter->InitializeProcess();
+
+		EAliasObjectReference ObjectReference = EAliasObjectReference::LocalReference;
+		if (MeshParameters.bIsSymmetric)
+		{
+			// All actors of a Alias symmetric layer are defined in the world Reference i.e. they have identity transform. So Mesh actor has to be defined in the world reference.
+			ObjectReference = EAliasObjectReference::WorldReference;
+		}
+		else if (TessellationOptions.StitchingTechnique == EDatasmithCADStitchingTechnique::StitchingSew)
+		{
+			// In the case of StitchingSew, AlDagNode children of a GroupNode are merged together. To be merged, they have to be defined in the reference of parent GroupNode.
+			ObjectReference = EAliasObjectReference::ParentReference;
+		}
+
+		for (const TPair<TAlObjectPtr<AlDagNode>, FColor>& DagNode : Body->GetShells())
+		{
+			BRepConverter->AddBRep(*DagNode.Key, DagNode.Value, ObjectReference);
+		}
+
+		ModelConverter->RepairTopology();
+
+		ModelConverter->SaveModel(*OutputPath, MeshElement);
+
+		FMeshDescription MeshDescription;
+		DatasmithMeshHelper::PrepareAttributeForStaticMesh(MeshDescription);
+
+		ModelConverter->Tessellate(MeshParameters, MeshDescription);
+
+		return MoveTemp(MeshDescription);
+	}
+
+	TOptional<FMeshDescription> FWireTranslatorImpl::GetMeshOfShellNode(AlDagNode& DagNode, TSharedPtr<IDatasmithMeshElement> MeshElement, CADLibrary::FMeshParameters& MeshParameters)
+	{
+		if (CADModelConverter.IsValid() && CADModelConverter->IsSessionValid())
+		{
+			TOptional<FMeshDescription> UEMesh = MeshDagNodeWithExternalMesher(DagNode, MeshElement, MeshParameters);
+			return UEMesh;
+		}
+		else
+		{
+			AlMatrix4x4 AlMatrix;
+			DagNode.inverseGlobalTransformationMatrix(AlMatrix);
+			// TODO: the best way, should be to don't have to apply inverse global transform to the generated mesh
+			TAlObjectPtr<AlDagNode> TesselatedNode = OpenModelUtils::TesselateDagLeaf(DagNode, ETesselatorType::Fast, TessellationOptions.ChordTolerance);
+			if (TesselatedNode.IsValid())
+			{
+				// Get the meshes from the dag nodes. Note that removing the mesh's DAG.
+				// will also removes the meshes, so we have to do it later.
+				TOptional<FMeshDescription> UEMesh = GetMeshOfNodeMesh(*TesselatedNode, MeshElement, MeshParameters, &AlMatrix);
+				return UEMesh;
+			}
+		}
+		return TOptional<FMeshDescription>();
+	}
+
+	TOptional<FMeshDescription> FWireTranslatorImpl::GetMeshOfShellBody(TSharedPtr<FBodyData> Body, TSharedPtr<IDatasmithMeshElement> MeshElement, CADLibrary::FMeshParameters& MeshParameters)
+	{
+		TOptional<FMeshDescription> UEMesh = MeshDagNodeWithExternalMesher(Body, MeshElement, MeshParameters);
+		return UEMesh;
+	}
+
+	TOptional<FMeshDescription> FWireTranslatorImpl::GetMeshOfMeshBody(TSharedPtr<FBodyData> Body, TSharedPtr<IDatasmithMeshElement> MeshElement, CADLibrary::FMeshParameters& MeshParameters)
+	{
+		FMeshDescription MeshDescription;
+		DatasmithMeshHelper::PrepareAttributeForStaticMesh(MeshDescription);
+		MeshDescription.Empty();
+
+		for (const TPair<TAlObjectPtr<AlDagNode>, FColor>& DagNode : Body->GetShells())
+		{
+			AlMeshNode* MeshNode = DagNode.Key->asMeshNodePtr();
+			if (!AlIsValid(MeshNode))
+			{
+				continue;
+			}
+
+			AlMesh* MeshPtr = MeshNode->mesh();
+			if (!AlIsValid(MeshPtr))
+			{
+				continue;
+			}
+			TAlObjectPtr<AlMesh> Mesh(MeshPtr);
+
+			AlMatrix4x4 AlMatrix;
+			DagNode.Key->localTransformationMatrix(AlMatrix);
+
+			Mesh->transform(AlMatrix);
+
+			const FColor& ShaderColor = DagNode.Value;
+			const int32 ShaderId = CreateShaderId(ShaderColor);
+			const FString SlotMaterialName = FString::FromInt(ShaderId);
+
+			const bool bMerge = true;
+			OpenModelUtils::TransferAlMeshToMeshDescription(*MeshPtr, *SlotMaterialName, MeshDescription, MeshParameters, bMerge);
+		}
+
+		// Build edge meta data
+		FStaticMeshOperations::DetermineEdgeHardnessesFromVertexInstanceNormals(MeshDescription);
+
+		return MoveTemp(MeshDescription);
+	}
+
+	TOptional<FMeshDescription> FWireTranslatorImpl::GetMeshOfNodeMesh(AlDagNode& TesselatedNode, TSharedPtr<IDatasmithMeshElement> MeshElement, CADLibrary::FMeshParameters& MeshParameters, AlMatrix4x4* AlMeshInvGlobalMatrix)
+	{
+		AlMeshNode* MeshNode = TesselatedNode.asMeshNodePtr();
+		if (!AlIsValid(MeshNode))
+		{
+			return TOptional<FMeshDescription>();
+		}
+
+		AlMesh* Mesh = MeshNode->mesh();
+		if (!AlIsValid(Mesh))
+		{
+			return TOptional<FMeshDescription>();
+		}
+
+		TAlObjectPtr<AlMesh> SharedMesh(Mesh);
+		if (AlMeshInvGlobalMatrix != nullptr)
+		{
+			SharedMesh->transform(*AlMeshInvGlobalMatrix);
+		}
+
+		return ImportMesh(*Mesh, MeshElement, MeshParameters);
+	}
+
+	TOptional<FMeshDescription> FWireTranslatorImpl::GetMeshDescription(TSharedPtr<IDatasmithMeshElement> MeshElement, CADLibrary::FMeshParameters& MeshParameters, TSharedPtr<FBodyData> Body)
+	{
+		if (Body->GetShells().Num() == 0)
+		{
+			return TOptional<FMeshDescription>();
+		}
+
+		const TArray<TPair<TAlObjectPtr<AlDagNode>, FColor>>& Shells = Body->GetShells();
+		TAlObjectPtr<AlDagNode> DagNode = Shells[0].Key;
+		AlLayer* LayerPtr = DagNode->layer();
+		if (AlIsValid(LayerPtr))
+		{
+			const TAlObjectPtr<AlLayer> Layer(LayerPtr);
+			if (LayerPtr->isSymmetric())
+			{
+				MeshParameters.bIsSymmetric = true;
+				double Normal[3], Origin[3];
+				LayerPtr->symmetricNormal(Normal[0], Normal[1], Normal[2]);
+				LayerPtr->symmetricOrigin(Origin[0], Origin[1], Origin[2]);
+
+				MeshParameters.SymmetricOrigin.X = (float)Origin[0];
+				MeshParameters.SymmetricOrigin.Y = (float)Origin[1];
+				MeshParameters.SymmetricOrigin.Z = (float)Origin[2];
+				MeshParameters.SymmetricNormal.X = (float)Normal[0];
+				MeshParameters.SymmetricNormal.Y = (float)Normal[1];
+				MeshParameters.SymmetricNormal.Z = (float)Normal[2];
+			}
+		}
+
+		if (Body->IsCadData())
+		{
+			return GetMeshOfShellBody(Body, MeshElement, MeshParameters);
+		}
+		else
+		{
+			return GetMeshOfMeshBody(Body, MeshElement, MeshParameters);
+		}
+	}
+
+	TOptional<FMeshDescription> FWireTranslatorImpl::GetMeshDescription(TSharedPtr<IDatasmithMeshElement> MeshElement, CADLibrary::FMeshParameters& MeshParameters)
+	{
+		TAlObjectPtr<AlDagNode>* DagNodeTemp = MeshElementToAlDagNodeMap.Find(MeshElement.Get());
+		if (DagNodeTemp == nullptr || !DagNodeTemp->IsValid())
+		{
+			TSharedPtr<FBodyData>* BodyTemp = MeshElementToBodyMap.Find(MeshElement.Get());
+			if (BodyTemp != nullptr && BodyTemp->IsValid())
+			{
+				return GetMeshDescription(MeshElement, MeshParameters, (*BodyTemp));
+			}
+			return TOptional<FMeshDescription>();
+		}
+
+		AlDagNode& DagNode = **DagNodeTemp;
+		AlObjectType objectType = DagNode.type();
+
+		if (objectType == kShellNodeType || objectType == kSurfaceNodeType || objectType == kMeshNodeType)
+		{
+			boolean bAlOrientation;
+			DagNode.getSurfaceOrientation(bAlOrientation);
+			MeshParameters.bNeedSwapOrientation = (objectType == kMeshNodeType) ? (bool)bAlOrientation : false; 
+
+			AlLayer* LayerPtr = DagNode.layer();
+			if (AlIsValid(LayerPtr))
+			{
+				const TAlObjectPtr<AlLayer> Layer(LayerPtr);
+				if (LayerPtr->isSymmetric())
+				{
+					MeshParameters.bIsSymmetric = true;
+					double Normal[3], Origin[3];
+					LayerPtr->symmetricNormal(Normal[0], Normal[1], Normal[2]);
+					LayerPtr->symmetricOrigin(Origin[0], Origin[1], Origin[2]);
+
+					MeshParameters.SymmetricOrigin.X = (float)Origin[0];
+					MeshParameters.SymmetricOrigin.Y = (float)Origin[1];
+					MeshParameters.SymmetricOrigin.Z = (float)Origin[2];
+					MeshParameters.SymmetricNormal.X = (float)Normal[0];
+					MeshParameters.SymmetricNormal.Y = (float)Normal[1];
+					MeshParameters.SymmetricNormal.Z = (float)Normal[2];
+				}
+			}
+		}
+
+		switch (objectType)
+		{
+		case kShellNodeType:
+		case kSurfaceNodeType:
+		{
+			return GetMeshOfShellNode(DagNode, MeshElement, MeshParameters);
+			break;
+		}
+
+		case kMeshNodeType:
+		{
+			return GetMeshOfNodeMesh(DagNode, MeshElement, MeshParameters);
+			break;
+		}
+
+		default:
+			break;
+		}
+
+		return TOptional<FMeshDescription>();
+	}
+
+
+	// Note that Alias file unit is cm like UE
+	TOptional<FMeshDescription> FWireTranslatorImpl::ImportMesh(AlMesh& InMesh, TSharedPtr<IDatasmithMeshElement> MeshElement, CADLibrary::FMeshParameters& InMeshParameters)
+	{
+		FMeshDescription MeshDescription;
+		DatasmithMeshHelper::PrepareAttributeForStaticMesh(MeshDescription);
+
+		const FColor* ShaderColorPtr = nullptr;
+		const FString* ShaderName = MeshElementToShaderName.Find(MeshElement.Get());
+		if (ShaderName)
+		{
+			ShaderColorPtr = ShaderNameToColor.Find(*ShaderName);
+		}
+		if (!ShaderColorPtr) // Should never happen
+		{
+			ShaderColorPtr = &DefaultColor;
+		}
+		int32 ShaderId = CreateShaderId(*ShaderColorPtr);
+		FString SlotMaterialName = FString::FromInt(ShaderId);
+
+		const bool bMerge = false;
+		OpenModelUtils::TransferAlMeshToMeshDescription(InMesh, *SlotMaterialName, MeshDescription, InMeshParameters, bMerge);
+
+		// Build edge meta data
+		FStaticMeshOperations::DetermineEdgeHardnessesFromVertexInstanceNormals(MeshDescription);
+
+		return MoveTemp(MeshDescription);
+	}
+
+	#define TRACK_MESH 0
+
+	bool FWireTranslatorImpl::LoadStaticMesh(const TSharedPtr<IDatasmithMeshElement> MeshElement, FDatasmithMeshElementPayload& OutMeshPayload, const FDatasmithTessellationOptions& InTessellationOptions)
+	{
+	#if TRACK_MESH
+		static const FString MeshName(TEXT("shell_30181"));
+		bool bSkip = true;
+		if (MeshName.Equals(MeshElement->GetLabel()))
+		{
+			bSkip = false;
+		}
+		if (bSkip)
+		{
+			return false;
+		}
+	#endif
+		CADLibrary::FMeshParameters MeshParameters;
+
+		uint64 StartTime = FPlatformTime::Cycles64();
+
+		if (TOptional<FMeshDescription> Mesh = GetMeshDescription(MeshElement, MeshParameters))
+		{
+			OutMeshPayload.LodMeshes.Add(MoveTemp(Mesh.GetValue()));
+			const TCHAR* MeshFilename = MeshElement->GetFile();
+			if (!bAliasUseNative && FPaths::FileExists(MeshFilename))
+			{
+				CADModelConverter->AddSurfaceDataForMesh(MeshFilename, MeshParameters, InTessellationOptions, OutMeshPayload);
+
+				// Remove the file because it is temporary since caching is disabled.
+				if (!CADLibrary::FImportParameters::bGEnableCADCache)
+				{
+					IFileManager::Get().Delete(MeshFilename);
+				}
+			}
+		}
+
+		double ElapsedSeconds = FPlatformTime::ToSeconds64(FPlatformTime::Cycles64() - StartTime);
+
+		UE_LOG(LogWireInterface, Display, TEXT("Mesh description of %s retrieved in %.3f s"), MeshElement->GetLabel(), ElapsedSeconds);
+
+		return OutMeshPayload.LodMeshes.Num() > 0;
+	}
+#endif
+
+	class FWireInterfaceModule : public IModuleInterface
+	{
+	public:
+		virtual void StartupModule() override
+		{
+			uint64 AliasVersion = IWireInterface::GetRequiredAliasVersion();
+
+#ifdef OPEN_MODEL_2020
+			// Check installed version of Alias Tools because binaries before 2021.3 are not compatible with Alias 2022
+			if (LibAlias2020_Version < AliasVersion && AliasVersion < LibAlias2021_Version)
+			{
+				static const bool bIsDisplay = []() -> bool
+					{
+						UE_LOG(LogWireInterface, Warning, TEXT(WRONG_VERSION_TEXT));
+						return true;
+					}();
+				return;
+			}
+#endif
+
+			if (LibAliasVersionMin <= AliasVersion && AliasVersion < LibAliasVersionMax)
+			{
+				auto MakeInterfaceFunc = []() -> TSharedPtr<IWireInterface>
+					{
+						return MakeShared<FWireTranslatorImpl>();
+					};
+				IWireInterface::RegisterInterface(UE_OPENMODEL_MAJOR_VERSION, UE_OPENMODEL_MAJOR_VERSION, MoveTemp(MakeInterfaceFunc));
+			}
+		}
+
+		virtual void ShutdownModule() override
+		{
+		}
+	};
+
+} // namespace
+
+// need this macro wrapper to expand the UE_DATASMITHWIRETRANSLATOR_MODULE_NAME macro and not create symbols with "UE_DATASMITHWIRETRANSLATOR_MODULE_NAME" in the token
+#define IMPLEMENT_MODULE_WRAPPER(ModuleName) IMPLEMENT_MODULE(UE_DATASMITHWIRETRANSLATOR_NAMESPACE::FWireInterfaceModule, ModuleName);
+IMPLEMENT_MODULE_WRAPPER(UE_DATASMITHWIRETRANSLATOR_MODULE_NAME)
+
+#undef LOCTEXT_NAMESPACE // "DatasmithWireTranslator"
diff --git a/Engine/Plugins/Enterprise/DatasmithCADImporter/Source/WireInterface/Private/WireInterfaceImpl.h b/Engine/Plugins/Enterprise/DatasmithCADImporter/Source/WireInterface/Private/WireInterfaceImpl.h
new file mode 100644
--- /dev/null
+++ b/Engine/Plugins/Enterprise/DatasmithCADImporter/Source/WireInterface/Private/WireInterfaceImpl.h
@@ -0,0 +1,137 @@
+// Copyright Epic Games, Inc. All Rights Reserved.
+
+#pragma once
+
+#include "CoreMinimal.h"
+
+#include "OpenModelUtils.h"
+
+#include "IWireInterface.h"
+
+#include "DatasmithImportOptions.h"
+
+#include "Templates/SharedPointer.h"
+
+class AlDagNode;
+class AlMesh;
+class AlShaderNode;
+class IDatasmithActorElement;
+class IDatasmithMaterialIDElement;
+class IDatasmithUEPbrMaterialElement;
+
+struct FMeshDescription;
+struct FDatasmithTessellationOptions;
+
+namespace CADLibrary
+{
+	struct FMeshParameters;
+	class ICADModelConverter;
+}
+
+namespace UE_DATASMITHWIRETRANSLATOR_NAMESPACE
+{
+	class FBodyData;
+	struct FDagNodeInfo;
+	class IAliasBRepConverter;
+
+	class FWireTranslatorImpl : public IWireInterface
+	{
+	public:
+		FWireTranslatorImpl() {}
+
+		~FWireTranslatorImpl();
+
+		/** Begin IWrieInterface */
+		virtual bool Initialize(const TCHAR* InSceneFullName) override;
+		virtual bool Load(TSharedPtr<IDatasmithScene> InScene) override;
+		virtual void SetTessellationOptions(const FDatasmithTessellationOptions& Options) override;
+		virtual void SetOutputPath(const FString& Path) { OutputPath = Path; }
+
+		bool LoadStaticMesh(const TSharedPtr<IDatasmithMeshElement> MeshElement, FDatasmithMeshElementPayload& OutMeshPayload, const FDatasmithTessellationOptions& InTessellationOptions);
+		/** End IWrieInterface */
+
+	private:
+		TOptional<FMeshDescription> GetMeshDescription(TSharedPtr<IDatasmithMeshElement> MeshElement, CADLibrary::FMeshParameters& OutMeshParameters);
+		TOptional<FMeshDescription> GetMeshDescription(TSharedPtr<IDatasmithMeshElement> MeshElement, CADLibrary::FMeshParameters& MeshParameters, TSharedPtr<FBodyData> BodyTemp);
+
+
+	private:
+
+		bool GetDagLeaves();
+		bool GetShader();
+		void RecurseDagForLeaves(const TAlObjectPtr<AlDagNode>& DagNode, FDagNodeInfo& ParentInfo);
+		void RecurseDagForLeavesNoMerge(const TAlObjectPtr<AlDagNode>& DagNode, FDagNodeInfo& ParentInfo);
+		void DagForLeavesNoMerge(const TAlObjectPtr<AlDagNode>& DagNode, FDagNodeInfo& ParentInfo);
+		bool ProcessAlGroupNode(AlDagNode& GroupNode, FDagNodeInfo& ParentInfo);
+		void ProcessAlShellNode(const TAlObjectPtr<AlDagNode>& ShellNode, FDagNodeInfo& ParentInfo, const FString& ShaderName);
+		void ProcessBodyNode(const TSharedPtr<FBodyData>& Body, FDagNodeInfo& ParentInfo);
+		TSharedPtr<IDatasmithMeshElement> FindOrAddMeshElement(const TSharedPtr<FBodyData>& Body, const FDagNodeInfo& ParentInfo);
+		TSharedPtr<IDatasmithMeshElement> FindOrAddMeshElement(const TAlObjectPtr<AlDagNode>& ShellNode, const FDagNodeInfo& ParentInfo, const FString& ShaderName);
+		TSharedPtr<IDatasmithActorElement> FindOrAddParentActor(FDagNodeInfo& ParentInfo, const FString& LayerName);
+		TSharedPtr<IDatasmithActorElement> FindOrAddLayerActor(const FString& LayerName);
+
+		void RecurseDeleteEmptyActor(TSharedPtr<IDatasmithActorElement> Actor);
+
+		TOptional<FMeshDescription> GetMeshOfShellNode(AlDagNode& DagNode, TSharedPtr<IDatasmithMeshElement> MeshElement, CADLibrary::FMeshParameters& MeshParameters);
+		TOptional<FMeshDescription> GetMeshOfNodeMesh(AlDagNode& DagNode, TSharedPtr<IDatasmithMeshElement> MeshElement, CADLibrary::FMeshParameters& MeshParameters, AlMatrix4x4* AlMeshInvGlobalMatrix = nullptr);
+		TOptional<FMeshDescription> GetMeshOfShellBody(TSharedPtr<FBodyData> DagNode, TSharedPtr<IDatasmithMeshElement> MeshElement, CADLibrary::FMeshParameters& MeshParameters);
+		TOptional<FMeshDescription> GetMeshOfMeshBody(TSharedPtr<FBodyData> DagNode, TSharedPtr<IDatasmithMeshElement> MeshElement, CADLibrary::FMeshParameters& MeshParameters);
+
+		void AddNodeInBodyGroup(TAlObjectPtr<AlDagNode>& DagNode, const FString& ShaderName, TMap<uint32, TSharedPtr<FBodyData>>& ShellToProcess, bool bIsAPatch, uint32 MaxSize);
+
+		TSharedPtr<CADLibrary::ICADModelConverter> GetModelConverter(IAliasBRepConverter*& BRepConverter) const;
+
+		TOptional<FMeshDescription> MeshDagNodeWithExternalMesher(AlDagNode& DagNode, TSharedPtr<IDatasmithMeshElement> MeshElement, CADLibrary::FMeshParameters& MeshParameters);
+		TOptional<FMeshDescription> MeshDagNodeWithExternalMesher(TSharedPtr<FBodyData> DagNode, TSharedPtr<IDatasmithMeshElement> MeshElement, CADLibrary::FMeshParameters& MeshParameters);
+
+		TOptional<FMeshDescription> ImportMesh(AlMesh& Mesh, TSharedPtr<IDatasmithMeshElement> MeshElement, CADLibrary::FMeshParameters& MeshParameters);
+
+		bool IsTransparent(FColor& TransparencyColor)
+		{
+			float Opacity = 1.0f - ((float)(TransparencyColor.R + TransparencyColor.G + TransparencyColor.B)) / 765.0f;
+			return !FMath::IsNearlyEqual(Opacity, 1.0f);
+		}
+
+		bool GetCommonParameters(int32 Field, double Value, FColor& Color, FColor& TransparencyColor, FColor& IncandescenceColor, double GlowIntensity);
+
+		void AddAlBlinnParameters(const TAlObjectPtr<AlShader>& Shader, TSharedPtr<IDatasmithUEPbrMaterialElement> MaterialElement);
+		void AddAlLambertParameters(const TAlObjectPtr<AlShader>& Shader, TSharedPtr<IDatasmithUEPbrMaterialElement> MaterialElement);
+		void AddAlLightSourceParameters(const TAlObjectPtr<AlShader>& Shader, TSharedPtr<IDatasmithUEPbrMaterialElement> MaterialElement);
+		void AddAlPhongParameters(const TAlObjectPtr<AlShader>& Shader, TSharedPtr<IDatasmithUEPbrMaterialElement> MaterialElement);
+
+	private:
+		TSharedPtr<IDatasmithScene> DatasmithScene;
+		FString SceneName;
+		FString CurrentPath;
+		FString OutputPath;
+		FString SceneFullPath;
+
+		// Hash value of the scene file used to check if the file has been modified for re-import
+		uint32 SceneFileHash = 0;
+
+		TAlObjectPtr<AlDagNode> AlRootNode;
+
+		/** Table of correspondence between mesh identifier and associated Datasmith mesh element */
+		TMap<uint32, TSharedPtr<IDatasmithMeshElement>> ShellUuidToMeshElementMap;
+		TMap<uint32, TSharedPtr<IDatasmithMeshElement>> BodyUuidToMeshElementMap;
+
+		/** Datasmith mesh elements to OpenModel objects */
+		TMap<IDatasmithMeshElement*, TAlObjectPtr<AlDagNode>> MeshElementToAlDagNodeMap;
+
+		TMap<IDatasmithMeshElement*, TSharedPtr<FBodyData>> MeshElementToBodyMap;
+
+		TMap<IDatasmithMeshElement*, FString> MeshElementToShaderName;
+		TMap<FString, FColor> ShaderNameToColor;
+		TMap<FString, TSharedPtr<IDatasmithMaterialIDElement>> ShaderNameToUEMaterialId;
+
+		TMap<FString, TSharedPtr<IDatasmithActorElement>> LayerNameToActor;
+
+		FDatasmithTessellationOptions TessellationOptions;
+
+		TSharedPtr<CADLibrary::ICADModelConverter> CADModelConverter;
+		TSharedPtr<IAliasBRepConverter> AliasBRepConverter;
+
+		bool bAliasUseNative = false;
+		bool bSceneLoaded = false;
+	};
+} // namespace
diff --git a/Engine/Plugins/Enterprise/DatasmithCADImporter/Source/WireInterface/Public/WireInterfaceModule.h b/Engine/Plugins/Enterprise/DatasmithCADImporter/Source/WireInterface/Public/WireInterfaceModule.h
new file mode 100644
--- /dev/null
+++ b/Engine/Plugins/Enterprise/DatasmithCADImporter/Source/WireInterface/Public/WireInterfaceModule.h
@@ -0,0 +1,39 @@
+// Copyright Epic Games, Inc. All Rights Reserved.
+
+#pragma once
+
+#include "CoreMinimal.h"
+#include "Modules/ModuleInterface.h"
+#include "Modules/ModuleManager.h"
+
+namespace UE_DATASMITHWIRETRANSLATOR_NAMESPACE
+{
+
+/**
+ * Datasmith Translator for .wire files.
+ */
+class FDatasmithWireTranslatorModule : public IModuleInterface
+{
+public:
+
+    static FDatasmithWireTranslatorModule& Get()
+    {
+        return FModuleManager::LoadModuleChecked< FDatasmithWireTranslatorModule >(PREPROCESSOR_TO_STRING(UE_DATASMITHWIRETRANSLATOR_MODULE_NAME));
+	}
+
+    static bool IsAvailable()
+    {
+        return FModuleManager::Get().IsModuleLoaded(PREPROCESSOR_TO_STRING(UE_DATASMITHWIRETRANSLATOR_MODULE_NAME));
+    }
+
+	virtual void StartupModule() override;
+
+	virtual void ShutdownModule() override;
+
+	FString GetTempDir() const;
+
+private:
+	FString TempDir;
+};
+
+}
\ No newline at end of file
diff --git a/Engine/Plugins/Enterprise/DatasmithCADImporter/Source/WireInterface/WireInterface2020.Build.cs b/Engine/Plugins/Enterprise/DatasmithCADImporter/Source/WireInterface/WireInterface2020.Build.cs
new file mode 100644
--- /dev/null
+++ b/Engine/Plugins/Enterprise/DatasmithCADImporter/Source/WireInterface/WireInterface2020.Build.cs
@@ -0,0 +1,99 @@
+// Copyright Epic Games, Inc. All Rights Reserved.
+
+using System.IO;
+using UnrealBuildTool;
+
+public abstract class WireInterfaceBase : ModuleRules
+{
+	public WireInterfaceBase(ReadOnlyTargetRules Target) : base(Target)
+	{
+		// TODO: investigate to remove that (Jira UETOOL-4975)
+		bUseUnity = false;
+		//OptimizeCode = CodeOptimization.Never;
+		//PCHUsage = ModuleRules.PCHUsageMode.NoPCHs;
+
+		PublicDependencyModuleNames.AddRange(
+			new string[]
+			{
+				"Core",
+				"CoreUObject",
+			}
+		);
+
+		PrivateDependencyModuleNames.AddRange(
+			new string[]
+			{
+				"CADInterfaces",
+				"CADKernel",
+				"CADKernelSurface",
+				"CADLibrary",
+				"CADTools",
+				"DatasmithContent",
+				"DatasmithCore",
+				"DatasmithTranslator",
+				"DatasmithWireTranslator",
+				"Engine",
+				"MeshDescription",
+				"ParametricSurface",
+				"StaticMeshDescription",
+			}
+		);
+
+		if (Target.bBuildEditor == true)
+		{
+			PrivateDependencyModuleNames.AddRange(
+				new string[]
+				{
+ 					"MessageLog",
+ 					"UnrealEd",
+				}
+			);
+		}
+
+		PublicDefinitions.Add(GetAliasDefinition());
+		PublicDefinitions.Add($"UE_DATASMITHWIRETRANSLATOR_NAMESPACE={GetAliasDefinition()}Namespace");
+		PublicDefinitions.Add($"UE_DATASMITHWIRETRANSLATOR_MODULE_NAME={GetType().Name}");
+		PublicDefinitions.Add($"UE_OPENMODEL_MAJOR_VERSION={GetMajorVersion()}");
+		PublicDefinitions.Add($"UE_OPENMODEL_MINOR_VERSION={GetMinorVersion()}");
+
+		if (System.Type.GetType(GetAliasVersion()) != null)
+		{
+			PrivateDependencyModuleNames.Add(GetAliasVersion());
+		}
+	}
+
+	public abstract string GetAliasVersion();
+	public abstract string GetAliasDefinition();
+	public abstract int GetMajorVersion();
+	public abstract int GetMinorVersion();
+
+}
+
+public class WireInterface2020 : WireInterfaceBase
+{
+	public WireInterface2020(ReadOnlyTargetRules Target) 
+		: base(Target)
+	{
+		PublicDefinitions.Add("IS_MAIN_MODULE");
+	}
+
+	public override string GetAliasVersion()
+	{
+		return "OpenModel2020";
+	}
+	
+	public override string GetAliasDefinition()
+	{
+		return "OPEN_MODEL_2020";
+	}
+
+	public override int GetMajorVersion()
+	{
+		return 2020;
+	}
+
+	public override int GetMinorVersion()
+	{
+		return 0;
+	}
+}
diff --git a/Engine/Plugins/Enterprise/DatasmithCADImporter/Source/WireInterface/WireInterface2021_3.Build.cs b/Engine/Plugins/Enterprise/DatasmithCADImporter/Source/WireInterface/WireInterface2021_3.Build.cs
new file mode 100644
--- /dev/null
+++ b/Engine/Plugins/Enterprise/DatasmithCADImporter/Source/WireInterface/WireInterface2021_3.Build.cs
@@ -0,0 +1,32 @@
+// Copyright Epic Games, Inc. All Rights Reserved.
+
+using System.IO;
+using UnrealBuildTool;
+
+public class WireInterface2021_3 : WireInterfaceBase
+{
+	public WireInterface2021_3(ReadOnlyTargetRules Target) 
+		: base(Target)
+	{
+	}
+
+	public override string GetAliasVersion()
+	{
+		return "OpenModel2021_3";
+	}
+	
+	public override string GetAliasDefinition()
+	{
+		return "OPEN_MODEL_2021_3";
+	}
+
+	public override int GetMajorVersion()
+	{
+		return 2021;
+	}
+
+	public override int GetMinorVersion()
+	{
+		return 3;
+	}
+}
diff --git a/Engine/Plugins/Enterprise/DatasmithCADImporter/Source/WireInterface/WireInterface2022.Build.cs b/Engine/Plugins/Enterprise/DatasmithCADImporter/Source/WireInterface/WireInterface2022.Build.cs
new file mode 100644
--- /dev/null
+++ b/Engine/Plugins/Enterprise/DatasmithCADImporter/Source/WireInterface/WireInterface2022.Build.cs
@@ -0,0 +1,32 @@
+// Copyright Epic Games, Inc. All Rights Reserved.
+
+using System.IO;
+using UnrealBuildTool;
+
+public class WireInterface2022 : WireInterfaceBase
+{
+	public WireInterface2022(ReadOnlyTargetRules Target) 
+		: base(Target)
+	{
+	}
+
+	public override string GetAliasVersion()
+	{
+		return "OpenModel2022";
+	}
+	
+	public override string GetAliasDefinition()
+	{
+	return "OPEN_MODEL_2022";
+	}
+
+	public override int GetMajorVersion()
+	{
+		return 2022;
+	}
+
+	public override int GetMinorVersion()
+	{
+		return 0;
+	}
+}
\ No newline at end of file
diff --git a/Engine/Plugins/Enterprise/DatasmithCADImporter/Source/WireInterface/WireInterface2022_1.Build.cs b/Engine/Plugins/Enterprise/DatasmithCADImporter/Source/WireInterface/WireInterface2022_1.Build.cs
new file mode 100644
--- /dev/null
+++ b/Engine/Plugins/Enterprise/DatasmithCADImporter/Source/WireInterface/WireInterface2022_1.Build.cs
@@ -0,0 +1,32 @@
+// Copyright Epic Games, Inc. All Rights Reserved.
+
+using System.IO;
+using UnrealBuildTool;
+
+public class WireInterface2022_1 : WireInterfaceBase
+{
+	public WireInterface2022_1(ReadOnlyTargetRules Target) 
+		: base(Target)
+	{
+	}
+
+	public override string GetAliasVersion()
+	{
+		return "OpenModel2022_1";
+	}
+	
+	public override string GetAliasDefinition()
+	{
+		return "OPEN_MODEL_2022_1";
+	}
+
+	public override int GetMajorVersion()
+	{
+		return 2022;
+	}
+
+	public override int GetMinorVersion()
+	{
+		return 1;
+	}
+}
\ No newline at end of file
diff --git a/Engine/Plugins/Enterprise/DatasmithCADImporter/Source/WireInterface/WireInterface2022_2.Build.cs b/Engine/Plugins/Enterprise/DatasmithCADImporter/Source/WireInterface/WireInterface2022_2.Build.cs
new file mode 100644
--- /dev/null
+++ b/Engine/Plugins/Enterprise/DatasmithCADImporter/Source/WireInterface/WireInterface2022_2.Build.cs
@@ -0,0 +1,32 @@
+// Copyright Epic Games, Inc. All Rights Reserved.
+
+using System.IO;
+using UnrealBuildTool;
+
+public class WireInterface2022_2 : WireInterfaceBase
+{
+	public WireInterface2022_2(ReadOnlyTargetRules Target) 
+		: base(Target)
+	{
+	}
+
+	public override string GetAliasVersion()
+	{
+		return "OpenModel2022_2";
+	}
+
+	public override string GetAliasDefinition()
+	{
+		return "OPEN_MODEL_2022_2";
+	}
+
+	public override int GetMajorVersion()
+	{
+		return 2022;
+	}
+
+	public override int GetMinorVersion()
+	{
+		return 2;
+	}
+}
\ No newline at end of file
diff --git a/Engine/Plugins/Enterprise/DatasmithCADImporter/Source/WireInterface/WireInterface2023_0.Build.cs b/Engine/Plugins/Enterprise/DatasmithCADImporter/Source/WireInterface/WireInterface2023_0.Build.cs
new file mode 100644
--- /dev/null
+++ b/Engine/Plugins/Enterprise/DatasmithCADImporter/Source/WireInterface/WireInterface2023_0.Build.cs
@@ -0,0 +1,32 @@
+// Copyright Epic Games, Inc. All Rights Reserved.
+
+using System.IO;
+using UnrealBuildTool;
+
+public class WireInterface2023_0 : WireInterfaceBase
+{
+	public WireInterface2023_0(ReadOnlyTargetRules Target) 
+		: base(Target)
+	{
+	}
+
+	public override string GetAliasVersion()
+	{
+		return "OpenModel2023_0";
+	}
+
+	public override string GetAliasDefinition()
+	{
+		return "OPEN_MODEL_2023_0";
+	}
+
+	public override int GetMajorVersion()
+	{
+		return 2023;
+	}
+
+	public override int GetMinorVersion()
+	{
+		return 0;
+	}
+}
\ No newline at end of file
diff --git a/Engine/Plugins/Enterprise/DatasmithCADImporter/Source/WireInterface/WireInterface2023_1.Build.cs b/Engine/Plugins/Enterprise/DatasmithCADImporter/Source/WireInterface/WireInterface2023_1.Build.cs
new file mode 100644
--- /dev/null
+++ b/Engine/Plugins/Enterprise/DatasmithCADImporter/Source/WireInterface/WireInterface2023_1.Build.cs
@@ -0,0 +1,32 @@
+// Copyright Epic Games, Inc. All Rights Reserved.
+
+using System.IO;
+using UnrealBuildTool;
+
+public class WireInterface2023_1 : WireInterfaceBase
+{
+	public WireInterface2023_1(ReadOnlyTargetRules Target) 
+		: base(Target)
+	{
+	}
+
+	public override string GetAliasVersion()
+	{
+		return "OpenModel2023_1";
+	}
+
+	public override string GetAliasDefinition()
+	{
+		return "OPEN_MODEL_2023_1";
+	}
+
+	public override int GetMajorVersion()
+	{
+		return 2023;
+	}
+
+	public override int GetMinorVersion()
+	{
+		return 1;
+	}
+}
\ No newline at end of file
