Fixed crash when tessellating bezier curve

#rnx

diff --git a/Engine/Plugins/Enterprise/DatasmithCADImporter/Source/CADInterfaces/Private/TechSoftBridge.cpp b/Engine/Plugins/Enterprise/DatasmithCADImporter/Source/CADInterfaces/Private/TechSoftBridge.cpp
--- a/Engine/Plugins/Enterprise/DatasmithCADImporter/Source/CADInterfaces/Private/TechSoftBridge.cpp
+++ b/Engine/Plugins/Enterprise/DatasmithCADImporter/Source/CADInterfaces/Private/TechSoftBridge.cpp
@@ -161,10 +161,10 @@ UE::CADKernel::FSurfacicBoundary GetSurfacicBoundary(A3DDomainData& Domain, cons
 	UE::CADKernel::EIso VIndex = UVReparameterization.GetSwapUV() ? UE::CADKernel::EIso::IsoU : UE::CADKernel::EIso::IsoV;
 
 	UE::CADKernel::FSurfacicBoundary Boundary;
-	Boundary[UIndex].Min = Min.U;
-	Boundary[VIndex].Min = Min.V;
-	Boundary[UIndex].Max = Max.U;
-	Boundary[VIndex].Max = Max.V;
+	Boundary[UIndex].Min = FMath::Min(Min.U, Max.U);
+	Boundary[UIndex].Max = FMath::Max(Min.U, Max.U);
+	Boundary[VIndex].Min = FMath::Min(Min.V, Max.V);
+	Boundary[VIndex].Max = FMath::Max(Min.V, Max.V);
 
 	return Boundary;
 }
diff --git a/Engine/Source/Runtime/Datasmith/CADKernel/Private/CADKernel/Geo/Curves/BezierCurve.cpp b/Engine/Source/Runtime/Datasmith/CADKernel/Private/CADKernel/Geo/Curves/BezierCurve.cpp
--- a/Engine/Source/Runtime/Datasmith/CADKernel/Private/CADKernel/Geo/Curves/BezierCurve.cpp
+++ b/Engine/Source/Runtime/Datasmith/CADKernel/Private/CADKernel/Geo/Curves/BezierCurve.cpp
@@ -52,6 +52,14 @@ FBezierCurve::FBezierCurve(const FNurbsCurveData& NurbsCurveData)
 
 	Poles = NurbsCurveData.Poles;
 	Weights = NurbsCurveData.Weights;
+	if (Weights.Num() != Poles.Num())
+	{
+		Weights.SetNum(Poles.Num());
+		for (double& Weight : Weights)
+		{
+			Weight = 1.;
+		}
+	}
 
 	// Validate the curve is actually rational
 	if (NurbsCurveData.bIsRational)
diff --git a/Engine/Source/Runtime/Datasmith/CADKernel/Private/CADKernel/Geo/Curves/BezierCurve.h b/Engine/Source/Runtime/Datasmith/CADKernel/Private/CADKernel/Geo/Curves/BezierCurve.h
--- a/Engine/Source/Runtime/Datasmith/CADKernel/Private/CADKernel/Geo/Curves/BezierCurve.h
+++ b/Engine/Source/Runtime/Datasmith/CADKernel/Private/CADKernel/Geo/Curves/BezierCurve.h
@@ -28,8 +28,11 @@ protected:
 		NumSegments = 1;
 		NodalVector.Add(0.);
 		NodalVector.Add(1.);
-		Weights.Add(1.);
-		Weights.Add(1.);
+		Weights.SetNum(Poles.Num());
+		for (double& Weight : Weights)
+		{
+			Weight = 1.;
+		}
 	}
 
 	FBezierCurve(const FNurbsCurveData& NurbsCurveData);
