Fixed compilation issue in AliasModelToCADKernelConverter.h and AliasModelToTechSoftConverter.h

#rnx

diff --git a/Engine/Plugins/Enterprise/DatasmithCADImporter/Source/DatasmithWireTranslator/Private/DatasmithWireTranslator.cpp b/Engine/Plugins/Enterprise/DatasmithCADImporter/Source/DatasmithWireTranslator/Private/DatasmithWireTranslator.cpp
--- a/Engine/Plugins/Enterprise/DatasmithCADImporter/Source/DatasmithWireTranslator/Private/DatasmithWireTranslator.cpp
+++ b/Engine/Plugins/Enterprise/DatasmithCADImporter/Source/DatasmithWireTranslator/Private/DatasmithWireTranslator.cpp
@@ -14,6 +14,8 @@
 
 DEFINE_LOG_CATEGORY_STATIC(LogDatasmithWireTranslator, Log, All);
 
+#define LOCTEXT_NAMESPACE "DatasmithWireTranslator"
+
 static TAutoConsoleVariable<bool> CVarAliasThreadSafe(
 	TEXT("ds.WireTranslator.ThreadSafe"),
 	false,
diff --git a/Engine/Plugins/Enterprise/DatasmithCADImporter/Source/WireInterface/Private/AliasBrepConverter.h b/Engine/Plugins/Enterprise/DatasmithCADImporter/Source/WireInterface/Private/AliasBrepConverter.h
--- a/Engine/Plugins/Enterprise/DatasmithCADImporter/Source/WireInterface/Private/AliasBrepConverter.h
+++ b/Engine/Plugins/Enterprise/DatasmithCADImporter/Source/WireInterface/Private/AliasBrepConverter.h
@@ -5,6 +5,8 @@
 class AlDagNode;
 struct FColor;
 
+typedef double AlMatrix4x4[4][4];
+
 namespace UE_DATASMITHWIRETRANSLATOR_NAMESPACE
 {
 
diff --git a/Engine/Plugins/Enterprise/DatasmithCADImporter/Source/WireInterface/Private/OpenModelUtils.cpp b/Engine/Plugins/Enterprise/DatasmithCADImporter/Source/WireInterface/Private/OpenModelUtils.cpp
--- a/Engine/Plugins/Enterprise/DatasmithCADImporter/Source/WireInterface/Private/OpenModelUtils.cpp
+++ b/Engine/Plugins/Enterprise/DatasmithCADImporter/Source/WireInterface/Private/OpenModelUtils.cpp
@@ -1,6 +1,8 @@
 // Copyright Epic Games, Inc. All Rights Reserved.
 #include "OpenModelUtils.h"
 
+#include "AliasBRepConverter.h"
+
 #ifdef USE_OPENMODEL
 
 #include "CADOptions.h"
diff --git a/Engine/Plugins/Enterprise/DatasmithCADImporter/Source/WireInterface/Private/OpenModelUtils.h b/Engine/Plugins/Enterprise/DatasmithCADImporter/Source/WireInterface/Private/OpenModelUtils.h
--- a/Engine/Plugins/Enterprise/DatasmithCADImporter/Source/WireInterface/Private/OpenModelUtils.h
+++ b/Engine/Plugins/Enterprise/DatasmithCADImporter/Source/WireInterface/Private/OpenModelUtils.h
@@ -27,8 +27,6 @@ class AlMesh;
 
 struct FMeshDescription;
 
-typedef double AlMatrix4x4[4][4];
-
 namespace UE_DATASMITHWIRETRANSLATOR_NAMESPACE
 {
 	//enum class EAliasObjectReference : uint8
diff --git a/Engine/Plugins/Enterprise/DatasmithCADImporter/Source/WireInterface/Private/WireInterfaceImpl.cpp b/Engine/Plugins/Enterprise/DatasmithCADImporter/Source/WireInterface/Private/WireInterfaceImpl.cpp
--- a/Engine/Plugins/Enterprise/DatasmithCADImporter/Source/WireInterface/Private/WireInterfaceImpl.cpp
+++ b/Engine/Plugins/Enterprise/DatasmithCADImporter/Source/WireInterface/Private/WireInterfaceImpl.cpp
@@ -28,7 +28,6 @@
 #include "Modules/ModuleManager.h"
 #endif
 
-#include "AliasBrepConverter.h"
 #include "AliasModelToCADKernelConverter.h"
 #include "AliasModelToTechSoftConverter.h" // requires Techsoft as public dependency
 #include "CADInterfacesModule.h"
diff --git a/Engine/Plugins/Enterprise/DatasmithCADImporter/Source/WireInterface/Private/WireInterfaceImpl.h b/Engine/Plugins/Enterprise/DatasmithCADImporter/Source/WireInterface/Private/WireInterfaceImpl.h
--- a/Engine/Plugins/Enterprise/DatasmithCADImporter/Source/WireInterface/Private/WireInterfaceImpl.h
+++ b/Engine/Plugins/Enterprise/DatasmithCADImporter/Source/WireInterface/Private/WireInterfaceImpl.h
@@ -4,6 +4,7 @@
 
 #include "CoreMinimal.h"
 
+#include "AliasBrepConverter.h"
 #include "OpenModelUtils.h"
 
 #include "IWireInterface.h"
