Fixed compilation issue when OpenModel SDK is not avaialble

#rnx

diff --git a/Engine/Plugins/Enterprise/DatasmithCADImporter/Source/WireInterface/Private/AliasBrepConverter.h b/Engine/Plugins/Enterprise/DatasmithCADImporter/Source/WireInterface/Private/AliasBrepConverter.h
--- a/Engine/Plugins/Enterprise/DatasmithCADImporter/Source/WireInterface/Private/AliasBrepConverter.h
+++ b/Engine/Plugins/Enterprise/DatasmithCADImporter/Source/WireInterface/Private/AliasBrepConverter.h
@@ -2,6 +2,7 @@
 
 #pragma once
 
+#ifdef USE_OPENMODEL
 class AlDagNode;
 struct FColor;
 
@@ -24,4 +25,5 @@ public:
 	virtual bool AddBRep(AlDagNode& DagNode, const FColor& Color, EAliasObjectReference ObjectReference) = 0;
 };
 
-}
\ No newline at end of file
+}
+#endif
\ No newline at end of file
diff --git a/Engine/Plugins/Enterprise/DatasmithCADImporter/Source/WireInterface/Private/AliasModelToCADKernelConverter.h b/Engine/Plugins/Enterprise/DatasmithCADImporter/Source/WireInterface/Private/AliasModelToCADKernelConverter.h
--- a/Engine/Plugins/Enterprise/DatasmithCADImporter/Source/WireInterface/Private/AliasModelToCADKernelConverter.h
+++ b/Engine/Plugins/Enterprise/DatasmithCADImporter/Source/WireInterface/Private/AliasModelToCADKernelConverter.h
@@ -4,6 +4,7 @@
 
 #include "CoreMinimal.h"
 
+#ifdef USE_OPENMODEL
 #include "AliasBRepConverter.h"
 #include "OpenModelUtils.h"
 
@@ -67,4 +68,5 @@ protected:
 	TMap<void*, TSharedPtr<UE::CADKernel::FTopologicalEdge>>  AlEdge2CADKernelEdge;
 };
 
-}
\ No newline at end of file
+}
+#endif
\ No newline at end of file
diff --git a/Engine/Plugins/Enterprise/DatasmithCADImporter/Source/WireInterface/Private/AliasModelToTechSoftConverter.h b/Engine/Plugins/Enterprise/DatasmithCADImporter/Source/WireInterface/Private/AliasModelToTechSoftConverter.h
--- a/Engine/Plugins/Enterprise/DatasmithCADImporter/Source/WireInterface/Private/AliasModelToTechSoftConverter.h
+++ b/Engine/Plugins/Enterprise/DatasmithCADImporter/Source/WireInterface/Private/AliasModelToTechSoftConverter.h
@@ -4,6 +4,7 @@
 
 #include "CoreMinimal.h"
 
+#ifdef USE_OPENMODEL
 #include "AliasBRepConverter.h"
 #include "OpenModelUtils.h"
 
@@ -51,4 +52,5 @@ protected:
 };
 
 }
+#endif
 
diff --git a/Engine/Plugins/Enterprise/DatasmithCADImporter/Source/WireInterface/Private/WireInterfaceImpl.cpp b/Engine/Plugins/Enterprise/DatasmithCADImporter/Source/WireInterface/Private/WireInterfaceImpl.cpp
--- a/Engine/Plugins/Enterprise/DatasmithCADImporter/Source/WireInterface/Private/WireInterfaceImpl.cpp
+++ b/Engine/Plugins/Enterprise/DatasmithCADImporter/Source/WireInterface/Private/WireInterfaceImpl.cpp
@@ -2,6 +2,9 @@
 
 #include "WireInterfaceImpl.h"
 
+#include "Modules/ModuleManager.h"
+
+#ifdef USE_OPENMODEL
 #include "CADOptions.h"
 #include "Containers/List.h"
 #include "DatasmithImportOptions.h"
@@ -25,7 +28,6 @@
 #include "IMessageLogListing.h"
 #include "Logging/TokenizedMessage.h"
 #include "MessageLogModule.h"
-#include "Modules/ModuleManager.h"
 #endif
 
 #include "AliasModelToCADKernelConverter.h"
@@ -36,11 +38,6 @@
 #include "Windows/AllowWindowsPlatformTypes.h"
 #endif
 
-#ifndef USE_OPENMODEL
-#define USE_OPENMODEL
-#endif
-
-#ifdef USE_OPENMODEL
 #include <AlChannel.h>
 #include <AlDagNode.h>
 #include <AlGroupNode.h>
@@ -61,7 +58,6 @@
 #include <AlTrimRegion.h>
 #include <AlTM.h>
 #include <AlUniverse.h>
-#endif
 
 #if PLATFORM_WINDOWS
 #include "Windows/HideWindowsPlatformTypes.h"
@@ -83,12 +79,11 @@ FAutoConsoleVariableRef GAliasSewByMaterial(
 Default is disable\n"),
 ECVF_Default);
 
-#endif 
+#endif
 
 namespace UE_DATASMITHWIRETRANSLATOR_NAMESPACE
 {
 
-#ifdef USE_OPENMODEL
 	static bool bGSewByMaterial = false;
 	static bool bGLayersAsActors = false;
 
@@ -2344,13 +2339,13 @@ namespace UE_DATASMITHWIRETRANSLATOR_NAMESPACE
 
 		return OutMeshPayload.LodMeshes.Num() > 0;
 	}
-#endif
 
 	class FWireInterfaceModule : public IModuleInterface
 	{
 	public:
 		virtual void StartupModule() override
 		{
+#ifdef USE_OPENMODEL
 			uint64 AliasVersion = IWireInterface::GetRequiredAliasVersion();
 
 #ifdef OPEN_MODEL_2020
@@ -2374,6 +2369,7 @@ namespace UE_DATASMITHWIRETRANSLATOR_NAMESPACE
 					};
 				IWireInterface::RegisterInterface(UE_OPENMODEL_MAJOR_VERSION, UE_OPENMODEL_MAJOR_VERSION, MoveTemp(MakeInterfaceFunc));
 			}
+#endif
 		}
 
 		virtual void ShutdownModule() override
@@ -2383,8 +2379,25 @@ namespace UE_DATASMITHWIRETRANSLATOR_NAMESPACE
 
 } // namespace
 
+#undef LOCTEXT_NAMESPACE // "DatasmithWireTranslator"
+
+#else
+namespace UE_DATASMITHWIRETRANSLATOR_NAMESPACE
+{
+	class FWireInterfaceModule : public IModuleInterface
+	{
+	public:
+		virtual void StartupModule() override
+		{
+		}
+
+		virtual void ShutdownModule() override
+		{
+		}
+	};
+}
+#endif
+
 // need this macro wrapper to expand the UE_DATASMITHWIRETRANSLATOR_MODULE_NAME macro and not create symbols with "UE_DATASMITHWIRETRANSLATOR_MODULE_NAME" in the token
 #define IMPLEMENT_MODULE_WRAPPER(ModuleName) IMPLEMENT_MODULE(UE_DATASMITHWIRETRANSLATOR_NAMESPACE::FWireInterfaceModule, ModuleName);
 IMPLEMENT_MODULE_WRAPPER(UE_DATASMITHWIRETRANSLATOR_MODULE_NAME)
-
-#undef LOCTEXT_NAMESPACE // "DatasmithWireTranslator"
diff --git a/Engine/Plugins/Enterprise/DatasmithCADImporter/Source/WireInterface/Private/WireInterfaceImpl.h b/Engine/Plugins/Enterprise/DatasmithCADImporter/Source/WireInterface/Private/WireInterfaceImpl.h
--- a/Engine/Plugins/Enterprise/DatasmithCADImporter/Source/WireInterface/Private/WireInterfaceImpl.h
+++ b/Engine/Plugins/Enterprise/DatasmithCADImporter/Source/WireInterface/Private/WireInterfaceImpl.h
@@ -4,6 +4,8 @@
 
 #include "CoreMinimal.h"
 
+#ifdef USE_OPENMODEL
+
 #include "AliasBrepConverter.h"
 #include "OpenModelUtils.h"
 
@@ -136,3 +138,4 @@ namespace UE_DATASMITHWIRETRANSLATOR_NAMESPACE
 		bool bSceneLoaded = false;
 	};
 } // namespace
+#endif
