Wire Improvements:
  - Refactored the code for better readability and maintenance
  - Decoupled the Datasmith Wire translator itself from the actual Wire code (a new module created). This was needed to be able to have Wire specific options.
  - Fixed hierarchy issue when importing layers as actors
  - Added two new options, "Use Layer As Actor" and "Merge Geomtry By Group" to get rid of the console variable. Both are set to true by default.

#rnx

diff --git a/Engine/Plugins/Enterprise/DatasmithCADImporter/Source/DatasmithWireTranslator/Private/DatasmithWireTranslator.cpp b/Engine/Plugins/Enterprise/DatasmithCADImporter/Source/DatasmithWireTranslator/Private/DatasmithWireTranslator.cpp
--- a/Engine/Plugins/Enterprise/DatasmithCADImporter/Source/DatasmithWireTranslator/Private/DatasmithWireTranslator.cpp
+++ b/Engine/Plugins/Enterprise/DatasmithCADImporter/Source/DatasmithWireTranslator/Private/DatasmithWireTranslator.cpp
@@ -186,7 +186,6 @@ bool FDatasmithWireTranslator::LoadScene(TSharedRef<IDatasmithScene> OutScene)
 	IFileManager::Get().MakeDirectory(*OutputPath, true);
 
 	WireInterface->SetOutputPath(OutputPath);
-	WireInterface->SetTessellationOptions(CommonTessellationOptions);
 
 	return WireInterface->Load(OutScene);
 }
@@ -201,13 +200,28 @@ bool FDatasmithWireTranslator::LoadStaticMesh(const TSharedRef<IDatasmithMeshEle
 	return ensure(WireInterface.IsValid()) ? WireInterface->LoadStaticMesh(MeshElement, OutMeshPayload, CommonTessellationOptions) : false;
 }
 
+void FDatasmithWireTranslator::GetSceneImportOptions(TArray<TObjectPtr<UDatasmithOptionsBase>>& Options)
+{
+	UDatasmithWireOptions* ImportOptions = Datasmith::MakeOptionsObjectPtr<UDatasmithWireOptions>();
+	ImportOptions->LoadConfig();
+	Options.Add(ImportOptions);
+}
+
 void FDatasmithWireTranslator::SetSceneImportOptions(const TArray<TObjectPtr<UDatasmithOptionsBase>>& Options)
 {
 	FParametricSurfaceTranslator::SetSceneImportOptions(Options);
 
 	if (ensure(WireInterface.IsValid()))
 	{
-		WireInterface->SetTessellationOptions(CommonTessellationOptions);
+		for (const TObjectPtr<UDatasmithOptionsBase>& OptionPtr : Options)
+		{
+			if (UDatasmithWireOptions* ImportOptions = Cast<UDatasmithWireOptions>(OptionPtr))
+			{
+				ImportOptions->SaveConfig(CPF_Config);
+				WireInterface->SetImportSettings(ImportOptions->Settings);
+				CommonTessellationOptions = ImportOptions->Settings;
+			}
+		}
 	}
 }
 
diff --git a/Engine/Plugins/Enterprise/DatasmithCADImporter/Source/DatasmithWireTranslator/Private/DatasmithWireTranslator.h b/Engine/Plugins/Enterprise/DatasmithCADImporter/Source/DatasmithWireTranslator/Private/DatasmithWireTranslator.h
--- a/Engine/Plugins/Enterprise/DatasmithCADImporter/Source/DatasmithWireTranslator/Private/DatasmithWireTranslator.h
+++ b/Engine/Plugins/Enterprise/DatasmithCADImporter/Source/DatasmithWireTranslator/Private/DatasmithWireTranslator.h
@@ -4,39 +4,21 @@
 
 #include "CoreMinimal.h"
 
-#include "ParametricSurfaceTranslator.h"
-
-#include "DatasmithImportOptions.h"
+#include "IWireInterface.h"
 
-#include "UObject/ObjectMacros.h"
+#include "ParametricSurfaceTranslator.h"
 
 #include "DatasmithWireTranslator.generated.h"
 
-class IWireInterface;
-
-USTRUCT(BlueprintType)
-struct FDatasmithWireOptions : public FDatasmithTessellationOptions
-{
-	GENERATED_BODY()
-
-	UPROPERTY(config, EditAnywhere, BlueprintReadWrite, Category = "Wire Translation Options")
-	bool bUseLayerAsActor = false;
-
-	uint32 GetHash() const
-	{
-		return HashCombine(FDatasmithTessellationOptions::GetHash(), GetTypeHash(false));
-	}
-};
-
 
 UCLASS(BlueprintType, config = EditorPerProjectUserSettings)
-class UDatasmithWireTranslationOptions : public UDatasmithOptionsBase
+class UDatasmithWireOptions : public UDatasmithOptionsBase
 {
 	GENERATED_BODY()
 
 public:
 	UPROPERTY(config, EditAnywhere, BlueprintReadWrite, Category = "Wire Translation Options", meta = (ShowOnlyInnerProperties))
-	FDatasmithWireOptions Options;
+	FWireSettings Settings;
 };
 
 class FDatasmithWireTranslator : public FParametricSurfaceTranslator
@@ -58,6 +40,7 @@ public:
 	// End IDatasmithTranslator overrides
 
 	// Begin ADatasmithCoreTechTranslator overrides
+	virtual void GetSceneImportOptions(TArray<TObjectPtr<UDatasmithOptionsBase>>& Options) override;
 	virtual void SetSceneImportOptions(const TArray<TObjectPtr<UDatasmithOptionsBase>>& Options) override;
 
 protected:
diff --git a/Engine/Plugins/Enterprise/DatasmithCADImporter/Source/DatasmithWireTranslator/Public/IWireInterface.h b/Engine/Plugins/Enterprise/DatasmithCADImporter/Source/DatasmithWireTranslator/Public/IWireInterface.h
--- a/Engine/Plugins/Enterprise/DatasmithCADImporter/Source/DatasmithWireTranslator/Public/IWireInterface.h
+++ b/Engine/Plugins/Enterprise/DatasmithCADImporter/Source/DatasmithWireTranslator/Public/IWireInterface.h
@@ -4,6 +4,12 @@
 
 #include "CoreMinimal.h"
 
+#include "DatasmithImportOptions.h"
+
+#include "UObject/ObjectMacros.h"
+
+#include "IWireInterface.generated.h"
+
 class IDatasmithMeshElement;
 class IDatasmithScene;
 
@@ -12,6 +18,25 @@ struct FDatasmithMeshElementPayload;
 
 typedef TFunction<TSharedPtr<class IWireInterface>()> FInterfaceMaker;
 
+USTRUCT(BlueprintType)
+struct FWireSettings : public FDatasmithTessellationOptions
+{
+	GENERATED_BODY()
+
+	UPROPERTY(config, EditAnywhere, BlueprintReadWrite, Category = "Wire Translation Options", meta = (ToolTip = "If set to true, the first level's actors in the outliner are the layers. Default is true"))
+	bool bUseLayerAsActor = true;
+
+	UPROPERTY(config, EditAnywhere, BlueprintReadWrite, Category = "Wire Translation Options", meta = (ToolTip = "If set to true, all geometry nodes under a group are merged. Default is true."))
+	bool bMergeGeomtryByGroup = true;
+
+	uint32 GetHash() const
+	{
+		return HashCombine(FDatasmithTessellationOptions::GetHash(), GetTypeHash(bUseLayerAsActor));
+	}
+
+	bool bAliasUseNative = false;
+};
+
 class IWireInterface
 {
 public:
@@ -21,7 +46,7 @@ public:
 
 	virtual bool Load(TSharedPtr<IDatasmithScene> Scene) = 0;
 
-	virtual void SetTessellationOptions(const FDatasmithTessellationOptions& Options) = 0;
+	virtual void SetImportSettings(const FWireSettings& Settings) = 0;
 	virtual void SetOutputPath(const FString& Path) = 0;
 	virtual bool LoadStaticMesh(const TSharedPtr<IDatasmithMeshElement> MeshElement, FDatasmithMeshElementPayload& OutMeshPayload, const FDatasmithTessellationOptions& InTessellationOptions) = 0;
 
diff --git a/Engine/Plugins/Enterprise/DatasmithCADImporter/Source/WireInterface/Private/AliasBrepConverter.h b/Engine/Plugins/Enterprise/DatasmithCADImporter/Source/WireInterface/Private/AliasBrepConverter.h
--- a/Engine/Plugins/Enterprise/DatasmithCADImporter/Source/WireInterface/Private/AliasBrepConverter.h
+++ b/Engine/Plugins/Enterprise/DatasmithCADImporter/Source/WireInterface/Private/AliasBrepConverter.h
@@ -2,6 +2,8 @@
 
 #pragma once
 
+#include "CoreMinimal.h"
+
 #ifdef USE_OPENMODEL
 class AlDagNode;
 struct FColor;
@@ -23,6 +25,7 @@ class IAliasBRepConverter
 {
 public:
 	virtual bool AddBRep(AlDagNode& DagNode, const FColor& Color, EAliasObjectReference ObjectReference) = 0;
+	virtual bool AddBRep(AlDagNode& DagNode, uint32 SlotID, EAliasObjectReference ObjectReference) = 0;
 };
 
 }
diff --git a/Engine/Plugins/Enterprise/DatasmithCADImporter/Source/WireInterface/Private/AliasModelToCADKernelConverter.cpp b/Engine/Plugins/Enterprise/DatasmithCADImporter/Source/WireInterface/Private/AliasModelToCADKernelConverter.cpp
--- a/Engine/Plugins/Enterprise/DatasmithCADImporter/Source/WireInterface/Private/AliasModelToCADKernelConverter.cpp
+++ b/Engine/Plugins/Enterprise/DatasmithCADImporter/Source/WireInterface/Private/AliasModelToCADKernelConverter.cpp
@@ -297,12 +297,17 @@ void FAliasModelToCADKernelConverter::AddShell(const AlShell& InShell, EAliasObj
 }
 
 bool FAliasModelToCADKernelConverter::AddBRep(AlDagNode& DagNode, const FColor& Color, EAliasObjectReference InObjectReference)
+{
+	uint32 ColorId = (uint32)CADLibrary::BuildColorUId(Color);
+	return AddBRep(DagNode, ColorId, InObjectReference);
+}
+
+bool FAliasModelToCADKernelConverter::AddBRep(AlDagNode& DagNode, uint32 SlotID, EAliasObjectReference InObjectReference)
 {
 	AlEdge2CADKernelEdge.Empty();
 
 	TSharedRef<UE::CADKernel::FBody> CADKernelBody = UE::CADKernel::FEntity::MakeShared<UE::CADKernel::FBody>();
-	uint32 ColorId = (uint32) CADLibrary::BuildColorUId(Color);
-	CADKernelBody->SetColorId(ColorId);
+	CADKernelBody->SetColorId(SlotID);
 
 	TSharedRef<UE::CADKernel::FShell> CADKernelShell = UE::CADKernel::FEntity::MakeShared<UE::CADKernel::FShell>();
 	CADKernelBody->AddShell(CADKernelShell);
diff --git a/Engine/Plugins/Enterprise/DatasmithCADImporter/Source/WireInterface/Private/AliasModelToCADKernelConverter.h b/Engine/Plugins/Enterprise/DatasmithCADImporter/Source/WireInterface/Private/AliasModelToCADKernelConverter.h
--- a/Engine/Plugins/Enterprise/DatasmithCADImporter/Source/WireInterface/Private/AliasModelToCADKernelConverter.h
+++ b/Engine/Plugins/Enterprise/DatasmithCADImporter/Source/WireInterface/Private/AliasModelToCADKernelConverter.h
@@ -47,6 +47,7 @@ public:
 	// End FCADModelToCADKernelConverterBase overrides
 
 	virtual bool AddBRep(AlDagNode& DagNode, const FColor& Color, EAliasObjectReference ObjectReference) override;
+	virtual bool AddBRep(AlDagNode& DagNode, uint32 SlotID, EAliasObjectReference ObjectReference) override;
 
 protected:
 	TSharedPtr<UE::CADKernel::FTopologicalEdge> AddEdge(const AlTrimCurve& TrimCurve, TSharedPtr<UE::CADKernel::FSurface>& CarrierSurface);
diff --git a/Engine/Plugins/Enterprise/DatasmithCADImporter/Source/WireInterface/Private/AliasModelToTechSoftConverter.cpp b/Engine/Plugins/Enterprise/DatasmithCADImporter/Source/WireInterface/Private/AliasModelToTechSoftConverter.cpp
--- a/Engine/Plugins/Enterprise/DatasmithCADImporter/Source/WireInterface/Private/AliasModelToTechSoftConverter.cpp
+++ b/Engine/Plugins/Enterprise/DatasmithCADImporter/Source/WireInterface/Private/AliasModelToTechSoftConverter.cpp
@@ -354,6 +354,13 @@ A3DTopoFace* FAliasModelToTechSoftConverter::AddTrimRegion(const AlTrimRegion& I
 }
 #endif
 
+bool FAliasModelToTechSoftConverter::AddBRep(AlDagNode& DagNode, uint32 SlotID, EAliasObjectReference InObjectReference)
+{
+	FColor Color(SlotID);
+
+	return AddBRep(DagNode, Color, InObjectReference);
+}
+
 bool FAliasModelToTechSoftConverter::AddBRep(AlDagNode& DagNode, const FColor& Color, EAliasObjectReference InObjectReference)
 {
 #ifdef USE_TECHSOFT_SDK
diff --git a/Engine/Plugins/Enterprise/DatasmithCADImporter/Source/WireInterface/Private/AliasModelToTechSoftConverter.h b/Engine/Plugins/Enterprise/DatasmithCADImporter/Source/WireInterface/Private/AliasModelToTechSoftConverter.h
--- a/Engine/Plugins/Enterprise/DatasmithCADImporter/Source/WireInterface/Private/AliasModelToTechSoftConverter.h
+++ b/Engine/Plugins/Enterprise/DatasmithCADImporter/Source/WireInterface/Private/AliasModelToTechSoftConverter.h
@@ -39,6 +39,7 @@ public:
 	}
 
 	virtual bool AddBRep(AlDagNode& DagNode, const FColor& Color, EAliasObjectReference ObjectReference) override;
+	virtual bool AddBRep(AlDagNode& DagNode, uint32 SlotID, EAliasObjectReference ObjectReference) override;
 
 protected:
 	TMap<void*, A3DTopoCoEdge*> AlEdgeToTSCoEdge;
diff --git a/Engine/Plugins/Enterprise/DatasmithCADImporter/Source/WireInterface/Private/OpenModelUtils.cpp b/Engine/Plugins/Enterprise/DatasmithCADImporter/Source/WireInterface/Private/OpenModelUtils.cpp
--- a/Engine/Plugins/Enterprise/DatasmithCADImporter/Source/WireInterface/Private/OpenModelUtils.cpp
+++ b/Engine/Plugins/Enterprise/DatasmithCADImporter/Source/WireInterface/Private/OpenModelUtils.cpp
@@ -18,301 +18,573 @@
 #include "Windows/AllowWindowsPlatformTypes.h"
 #endif
 
+#include "AlChannel.h"
 #include "AlDagNode.h"
-#include "AlMesh.h"
+#include "AlGroupNode.h"
 #include "AlLayer.h"
+#include "AlLinkItem.h"
+#include "AlList.h"
+#include "AlMesh.h"
+#include "AlMeshNode.h"
+#include "AlPersistentID.h"
+#include "AlRetrieveOptions.h"
+#include "AlShader.h"
+#include "AlShadingFieldItem.h"
+#include "AlShell.h"
+#include "AlShellNode.h"
+#include "AlSurface.h"
+#include "AlSurfaceNode.h"
 #include "AlTesselate.h"
+#include "AlTrimRegion.h"
+#include "AlTM.h"
+#include "AlUniverse.h"
 
-namespace UE_DATASMITHWIRETRANSLATOR_NAMESPACE
-{
+#if PLATFORM_WINDOWS
+#include "Windows/HideWindowsPlatformTypes.h"
+#endif
 
-void OpenModelUtils::SetActorTransform(TSharedPtr<IDatasmithActorElement>& OutActorElement, const AlDagNode& InDagNode)
+
+namespace UE_DATASMITHWIRETRANSLATOR_NAMESPACE
 {
-	if (OutActorElement)
+	bool FPatchMesh::Initialize()
 	{
-		AlLayer* LayerPtr = InDagNode.layer();
-		if (AlIsValid(LayerPtr))
+		if (MeshNodes.IsEmpty())
+		{
+			return false;
+		}
+
+		if (!bInitialized)
 		{
-			const TUniquePtr<AlLayer> Layer(LayerPtr);
+			Hash = GetTypeHash(Name);
 
-			if (LayerPtr->isSymmetric())
+			for (const TAlDagNodePtr<AlMeshNode>& MeshNode : MeshNodes)
 			{
-				return;
+				Hash = HashCombine(Hash, MeshNode.GetHash());
 			}
+
+			UniqueID = TEXT("PatchMesh") + FString::FromInt(Hash);
+
+			bInitialized = true;
 		}
 
-		AlMatrix4x4 AlMatrix;
-		InDagNode.localTransformationMatrix(AlMatrix);
+		return bInitialized;
+	}
 
-		FMatrix Matrix;
-		double* MatrixFloats = (double*)Matrix.M;
-		for (int32 IndexI = 0; IndexI < 4; ++IndexI)
+	int32 FBodyNode::GetSlotIndex(const TAlDagNodePtr<AlDagNode>& DagNode)
+	{
+		// TODO: Add support for AlShell
+		if (TAlDagNodePtr<AlSurfaceNode> SurfaceNode = DagNode->asSurfaceNodePtr())
 		{
-			for (int32 IndexJ = 0; IndexJ < 4; ++IndexJ)
+			if (TAlObjectPtr<AlSurface> Surface = SurfaceNode->surface())
 			{
-				MatrixFloats[IndexI * 4 + IndexJ] = AlMatrix[IndexI][IndexJ];
+				TAlObjectPtr<AlShader> Shader(Surface->firstShader());
+				return Shader ? ShaderNameToSlotIndex[Shader.GetName()] : 0;
 			}
 		}
-		FTransform LocalTransform(Matrix);
-		FTransform LocalUETransform = FDatasmithUtils::ConvertTransform(FDatasmithUtils::EModelCoordSystem::ZUp_RightHanded, LocalTransform);
 
-		FVector Translation = LocalUETransform.GetTranslation();
-		if (!FMath::IsNearlyEqual(Translation.X, 0.0f))
-		{
-			printf(">>> %f", Translation.X);
-		}
-		//OutActorElement->SetTranslation(LocalUETransform.GetTranslation());
-		//OutActorElement->SetScale(LocalUETransform.GetScale3D());
-		//OutActorElement->SetRotation(LocalUETransform.GetRotation());
+		ensure(false);
+		return 0;
 	}
-}
 
-bool OpenModelUtils::IsValidActor(const TSharedPtr<IDatasmithActorElement>& ActorElement)
-{
-	if (ActorElement != nullptr)
+	bool FBodyNode::Initialize()
 	{
-		if (ActorElement->GetChildrenCount() > 0)
+		if (SurfaceNodes.IsEmpty() && ShellNodes.IsEmpty())
 		{
-			return true;
+			return false;
 		}
-		else if (ActorElement->IsA(EDatasmithElementType::StaticMeshActor))
+
+		if (!bInitialized)
 		{
-			const TSharedPtr<IDatasmithMeshActorElement>& MeshActorElement = StaticCastSharedPtr<IDatasmithMeshActorElement>(ActorElement);
-			return FCString::Strlen(MeshActorElement->GetStaticMeshPathName()) > 0;
+			Hash = GetTypeHash(Name);
+
+			int32 SlotIndex = 0;
+			TSet<TAlObjectPtr<AlLayer>> LayerSet;
+
+			for (const TAlDagNodePtr<AlSurfaceNode>& SurfaceNode : SurfaceNodes)
+			{
+				TAlObjectPtr<AlSurface> Surface(SurfaceNode->surface());
+				if (!Surface)
+				{
+					continue;
+				}
+
+				TAlObjectPtr<AlShader> Shader(Surface->firstShader());
+				if (Shader && !ShaderNameToSlotIndex.Contains(Shader.GetName()))
+				{
+					ShaderNameToSlotIndex.Add(Shader.GetName(), SlotIndex);
+					SlotIndexToShader.Add(SlotIndex++, Shader);
+				}
+
+				if (TAlObjectPtr<AlLayer> SLayer = SurfaceNode->layer())
+				{
+					LayerSet.Add(TAlObjectPtr<AlLayer >(SLayer));
+				}
+
+				Hash = HashCombine(Hash, SurfaceNode.GetHash());
+			}
+
+			for (const TAlDagNodePtr<AlShellNode>& ShellNode : ShellNodes)
+			{
+				// TODO: Extract all shaders
+				TAlObjectPtr<AlShell> Shell(ShellNode->shell());
+				if (!Shell)
+				{
+					continue;
+				}
+
+				TAlObjectPtr<AlShader> Shader(Shell->firstShader());
+				if (Shader && !ShaderNameToSlotIndex.Contains(Shader.GetName()))
+				{
+					ShaderNameToSlotIndex.Add(Shader.GetName(), SlotIndex);
+					SlotIndexToShader.Add(SlotIndex++, Shader);
+				}
+
+				if (TAlObjectPtr<AlLayer> SLayer = ShellNode->layer())
+				{
+					LayerSet.Add(TAlObjectPtr<AlLayer >(SLayer));
+				}
+				Hash = HashCombine(Hash, ShellNode.GetHash());
+			}
+
+			ensure(LayerSet.Num() == 1);
+			TAlObjectPtr<AlLayer> GeomLayer = LayerSet.Array()[0];
+			if (GeomLayer != Layer)
+			{
+				Layer = GeomLayer;
+			}
+
+			UniqueID = TEXT("BodyNode") + FString::FromInt(Hash);
+
+			bInitialized = true;
 		}
-	}
-	return false;
-}
 
-bool OpenModelUtils::TransferAlMeshToMeshDescription(const AlMesh& AliasMesh, const TCHAR* SlotMaterialId, FMeshDescription& MeshDescription, CADLibrary::FMeshParameters& MeshParameters, const bool bMerge)
-{
-	if (AliasMesh.numberOfVertices() == 0 || AliasMesh.numberOfTriangles() == 0)
-	{
-		return false;
+		return bInitialized;
 	}
 
-	if (!bMerge)
+	void FBodyNode::AddSurfaceNode(const TAlDagNodePtr<AlSurfaceNode>& SurfaceNode)
 	{
-		MeshDescription.Empty();
+		if (OpenModelUtils::IsGeometryValid(SurfaceNode.Get()))
+		{
+			SurfaceNodes.Add(SurfaceNode);
+		}
 	}
 
-	int32 NbStep = 1;
-	FMatrix44f SymmetricMatrix;
-	bool bIsSymmetricMesh = MeshParameters.bIsSymmetric;
-	if (bIsSymmetricMesh)
+	void FBodyNode::AddShellNode(const TAlDagNodePtr<AlShellNode>& ShellNode)
 	{
-		NbStep = 2;
-		SymmetricMatrix = FDatasmithUtils::GetSymmetricMatrix(MeshParameters.SymmetricOrigin, MeshParameters.SymmetricNormal);
+		if (OpenModelUtils::IsGeometryValid(ShellNode.Get()))
+		{
+			ShellNodes.Add(ShellNode);
+		}
 	}
 
-	// Gather all array data
-	FStaticMeshAttributes Attributes(MeshDescription);
-	TVertexInstanceAttributesRef<FVector3f> VertexInstanceNormals = Attributes.GetVertexInstanceNormals();
-	TVertexInstanceAttributesRef<FVector2f> VertexInstanceUVs = Attributes.GetVertexInstanceUVs();
-	TPolygonGroupAttributesRef<FName> PolygonGroupImportedMaterialSlotNames = Attributes.GetPolygonGroupMaterialSlotNames();
-	TVertexAttributesRef<FVector3f> VertexPositions = MeshDescription.GetVertexPositions();
-
-	// Prepared for static mesh usage ?
-	if (!VertexPositions.IsValid() || !VertexInstanceNormals.IsValid() || !VertexInstanceUVs.IsValid() || !PolygonGroupImportedMaterialSlotNames.IsValid())
+	template<>
+	uint32 TAlObjectPtr<AlLayer>::GetHash() const
 	{
-		return false;
+		uint32 Hash = GetTypeHash(GetName());
+		Hash = HashCombine(Hash, GetTypeHash(Ptr->number()));
+		Hash = HashCombine(Hash, GetTypeHash(Ptr->color()));
+		Hash = HashCombine(Hash, (bool)Ptr->invisible() ? 1 : 0);
+		Hash = HashCombine(Hash, (bool)Ptr->isSymmetric() ? 1 : 0);
+
+		return Hash;
 	}
 
-	bool bHasUVData = (AliasMesh.uvs() != nullptr);
+	bool OpenModelUtils::GetCsvLayerString(const TAlObjectPtr<AlLayer>& Layer, FString& CsvString)
+	{
+		if (!Layer)
+		{
+			return false;
+		}
 
-	int VertexCount = AliasMesh.numberOfVertices();
-	int TriangleCount = AliasMesh.numberOfTriangles();
-	const int32 VertexInstanceCount = 3 * TriangleCount;
+		CsvString = Layer.GetName();
 
-	TArray<FVertexID> VertexPositionIDs;
-	VertexPositionIDs.SetNum(VertexCount * NbStep);
+		TAlObjectPtr<AlLayer> ParentLayer(Layer->parentLayer());
+		while (ParentLayer)
+		{
+			FString ParentLayerName = ParentLayer.GetName();
+			if (!ParentLayerName.IsEmpty())
+			{
+				CsvString += TEXT(",") + ParentLayerName;
+			}
 
-	// Reserve space for attributes
-	// At this point, all the faces are triangles
-	MeshDescription.ReserveNewVertices(VertexCount * NbStep);
-	MeshDescription.ReserveNewVertexInstances(VertexInstanceCount * NbStep);
-	MeshDescription.ReserveNewEdges(VertexInstanceCount * NbStep);
-	MeshDescription.ReserveNewPolygons(TriangleCount * NbStep);
+			ParentLayer = ParentLayer->parentLayer();
+		}
 
-	MeshDescription.ReserveNewPolygonGroups(1);
-	FPolygonGroupID PolyGroupId = MeshDescription.CreatePolygonGroup();
-	FName ImportedSlotName = SlotMaterialId;
-	PolygonGroupImportedMaterialSlotNames[PolyGroupId] = ImportedSlotName;
+		return !CsvString.IsEmpty();
+	}
 
-	// At least one UV set must exist.
-	if (VertexInstanceUVs.GetNumChannels() == 0)
+	bool OpenModelUtils::ActorHasContent(const TSharedPtr<IDatasmithActorElement>& ActorElement)
 	{
-		VertexInstanceUVs.SetNumChannels(1);
-	}
+		if (!ActorElement.IsValid())
+		{
+			return false;
+		}
 
-	// Get Alias mesh info
-	const float* AlVertices = AliasMesh.vertices();
+		return ActorElement->IsA(EDatasmithElementType::StaticMeshActor) || ActorElement->GetChildrenCount() > 0;
+	}
 
-	for (int32 Step = 0; Step < NbStep; Step++)
+	void OpenModelUtils::SetActorTransform(IDatasmithActorElement& ActorElement, const AlDagNode& InDagNode)
 	{
-		// Fill the vertex array
-		if (Step == 0)
+		// TODO: Find why no transform applied if layer is symmetrical???
+		TAlObjectPtr<AlLayer> Layer(InDagNode.layer());
+		if (Layer && Layer->isSymmetric())
 		{
-			FVertexID* VertexPositionIDPtr = VertexPositionIDs.GetData();
-			for (int Index = 0; Index < VertexCount; ++Index, ++VertexPositionIDPtr)
+			return;
+		}
+
+		AlMatrix4x4 AlGlobalMatrix;
+		InDagNode.globalTransformationMatrix(AlGlobalMatrix);
+
+		FMatrix GlobalMatrix;
+		double* MatrixFloats = (double*)GlobalMatrix.M;
+		for (int32 IndexI = 0; IndexI < 4; ++IndexI)
+		{
+			for (int32 IndexJ = 0; IndexJ < 4; ++IndexJ)
 			{
-				const float* CurVertex = AlVertices + 3 * Index;
-				*VertexPositionIDPtr = MeshDescription.CreateVertex();
-				// ConvertVector_ZUp_RightHanded
-				VertexPositions[*VertexPositionIDPtr] = FVector3f(-CurVertex[0], CurVertex[1], CurVertex[2]);
+				MatrixFloats[IndexI * 4 + IndexJ] = AlGlobalMatrix[IndexI][IndexJ];
 			}
 		}
-		else
+
+		FTransform GlobalTransform = FDatasmithUtils::ConvertTransform(FDatasmithUtils::EModelCoordSystem::ZUp_RightHanded, FTransform(GlobalMatrix));
+
+		ActorElement.SetTranslation(GlobalTransform.GetTranslation());
+		ActorElement.SetScale(GlobalTransform.GetScale3D());
+		ActorElement.SetRotation(GlobalTransform.GetRotation());
+	}
+
+	bool OpenModelUtils::IsValidActor(const TSharedPtr<IDatasmithActorElement>& ActorElement)
+	{
+		if (ActorElement != nullptr)
 		{
-			FVertexID* VertexPositionIDPtr = VertexPositionIDs.GetData() + VertexCount;
-			for (int Index = 0, PositionIndex = VertexCount; Index < VertexCount; ++Index, ++VertexPositionIDPtr)
+			if (ActorElement->GetChildrenCount() > 0)
 			{
-				const float* CurVertex = AlVertices + 3 * Index;
-				*VertexPositionIDPtr = MeshDescription.CreateVertex();
-				// ConvertVector_ZUp_RightHanded
-				VertexPositions[*VertexPositionIDPtr] = SymmetricMatrix.TransformPosition(FVector3f(-CurVertex[0], CurVertex[1], CurVertex[2]));
+				return true;
 			}
+			else if (ActorElement->IsA(EDatasmithElementType::StaticMeshActor))
+			{
+				const TSharedPtr<IDatasmithMeshActorElement>& MeshActorElement = StaticCastSharedPtr<IDatasmithMeshActorElement>(ActorElement);
+				return FCString::Strlen(MeshActorElement->GetStaticMeshPathName()) > 0;
+			}
+		}
+		return false;
+	}
+
+	bool OpenModelUtils::TransferAlMeshToMeshDescription(const AlMesh& AliasMesh, const TCHAR* SlotMaterialId, FMeshDescription& MeshDescription, CADLibrary::FMeshParameters& MeshParameters, const bool bMerge)
+	{
+		if (AliasMesh.numberOfVertices() == 0 || AliasMesh.numberOfTriangles() == 0)
+		{
+			return false;
 		}
 
-		FBox UVBBox(FVector(MAX_FLT), FVector(-MAX_FLT));
+		if (!bMerge)
+		{
+			MeshDescription.Empty();
+		}
 
-		const int32 CornerCount = 3; // only triangles
-		FVertexID CornerVertexIDs[3];
-		TArray<FVertexInstanceID> CornerVertexInstanceIDs;
-		CornerVertexInstanceIDs.SetNum(3);
+		int32 NbStep = 1;
+		FMatrix44f SymmetricMatrix;
+		bool bIsSymmetricMesh = MeshParameters.bIsSymmetric;
+		if (bIsSymmetricMesh)
+		{
+			NbStep = 2;
+			SymmetricMatrix = FDatasmithUtils::GetSymmetricMatrix(MeshParameters.SymmetricOrigin, MeshParameters.SymmetricNormal);
+		}
+
+		// Gather all array data
+		FStaticMeshAttributes Attributes(MeshDescription);
+		TVertexInstanceAttributesRef<FVector3f> VertexInstanceNormals = Attributes.GetVertexInstanceNormals();
+		TVertexInstanceAttributesRef<FVector2f> VertexInstanceUVs = Attributes.GetVertexInstanceUVs();
+		TPolygonGroupAttributesRef<FName> PolygonGroupImportedMaterialSlotNames = Attributes.GetPolygonGroupMaterialSlotNames();
+		TVertexAttributesRef<FVector3f> VertexPositions = MeshDescription.GetVertexPositions();
+
+		// Prepared for static mesh usage ?
+		if (!VertexPositions.IsValid() || !VertexInstanceNormals.IsValid() || !VertexInstanceUVs.IsValid() || !PolygonGroupImportedMaterialSlotNames.IsValid())
+		{
+			return false;
+		}
+
+		bool bHasUVData = (AliasMesh.uvs() != nullptr);
+
+		int VertexCount = AliasMesh.numberOfVertices();
+		int TriangleCount = AliasMesh.numberOfTriangles();
+		const int32 VertexInstanceCount = 3 * TriangleCount;
+
+		TArray<FVertexID> VertexPositionIDs;
+		VertexPositionIDs.SetNum(VertexCount * NbStep);
+
+		// Reserve space for attributes
+		// At this point, all the faces are triangles
+		MeshDescription.ReserveNewVertices(VertexCount * NbStep);
+		MeshDescription.ReserveNewVertexInstances(VertexInstanceCount * NbStep);
+		MeshDescription.ReserveNewEdges(VertexInstanceCount * NbStep);
+		MeshDescription.ReserveNewPolygons(TriangleCount * NbStep);
+
+		MeshDescription.ReserveNewPolygonGroups(1);
+		FPolygonGroupID PolyGroupId = MeshDescription.CreatePolygonGroup();
+		FName ImportedSlotName = SlotMaterialId;
+		PolygonGroupImportedMaterialSlotNames[PolyGroupId] = ImportedSlotName;
+
+		// At least one UV set must exist.
+		if (VertexInstanceUVs.GetNumChannels() == 0)
+		{
+			VertexInstanceUVs.SetNumChannels(1);
+		}
 
 		// Get Alias mesh info
-		const int* Triangles = AliasMesh.triangles();
-		const float* AlNormals = AliasMesh.normals();
-		const float* AlUVs = AliasMesh.uvs();
+		const float* AlVertices = AliasMesh.vertices();
 
-		// Get per-triangle data: indices, normals and uvs
-		if (!MeshParameters.bNeedSwapOrientation == ((bool)Step))
+		for (int32 Step = 0; Step < NbStep; Step++)
 		{
-			for (int32 FaceIndex = 0; FaceIndex < TriangleCount; ++FaceIndex, Triangles += 3)
+			// Fill the vertex array
+			if (Step == 0)
 			{
-				// Create Vertex instances and set their attributes
-				for (int32 VertexIndex = 0, TIndex = 2; VertexIndex < CornerCount; ++VertexIndex, --TIndex)
+				FVertexID* VertexPositionIDPtr = VertexPositionIDs.GetData();
+				for (int Index = 0; Index < VertexCount; ++Index, ++VertexPositionIDPtr)
 				{
-					CornerVertexIDs[VertexIndex] = VertexPositionIDs[Triangles[TIndex] + VertexCount * Step];
-					CornerVertexInstanceIDs[VertexIndex] = MeshDescription.CreateVertexInstance(CornerVertexIDs[VertexIndex]);
-
-					// Set the normal
-					const float* CurNormal = &AlNormals[3 * Triangles[TIndex]];
+					const float* CurVertex = AlVertices + 3 * Index;
+					*VertexPositionIDPtr = MeshDescription.CreateVertex();
 					// ConvertVector_ZUp_RightHanded
-					FVector3f UENormal(-CurNormal[0], CurNormal[1], CurNormal[2]);
-					UENormal = UENormal.GetSafeNormal();
-					if (Step > 0)
-					{
-						UENormal = SymmetricMatrix.TransformVector(UENormal);
-					}
-					else
-					{
-						UENormal *= -1.;
-					}
-					VertexInstanceNormals[CornerVertexInstanceIDs[VertexIndex]] = UENormal;
+					VertexPositions[*VertexPositionIDPtr] = FVector3f(-CurVertex[0], CurVertex[1], CurVertex[2]);
 				}
-				if (CornerVertexIDs[0] == CornerVertexIDs[1] || CornerVertexIDs[0] == CornerVertexIDs[2] || CornerVertexIDs[1] == CornerVertexIDs[2])
+			}
+			else
+			{
+				FVertexID* VertexPositionIDPtr = VertexPositionIDs.GetData() + VertexCount;
+				for (int Index = 0, PositionIndex = VertexCount; Index < VertexCount; ++Index, ++VertexPositionIDPtr)
 				{
-					continue;
+					const float* CurVertex = AlVertices + 3 * Index;
+					*VertexPositionIDPtr = MeshDescription.CreateVertex();
+					// ConvertVector_ZUp_RightHanded
+					VertexPositions[*VertexPositionIDPtr] = SymmetricMatrix.TransformPosition(FVector3f(-CurVertex[0], CurVertex[1], CurVertex[2]));
 				}
+			}
+
+			FBox UVBBox(FVector(MAX_FLT), FVector(-MAX_FLT));
 
-				// Set the UV
-				if (bHasUVData)
+			const int32 CornerCount = 3; // only triangles
+			FVertexID CornerVertexIDs[3];
+			TArray<FVertexInstanceID> CornerVertexInstanceIDs;
+			CornerVertexInstanceIDs.SetNum(3);
+
+			// Get Alias mesh info
+			const int* Triangles = AliasMesh.triangles();
+			const float* AlNormals = AliasMesh.normals();
+			const float* AlUVs = AliasMesh.uvs();
+
+			// Get per-triangle data: indices, normals and uvs
+			if (!MeshParameters.bNeedSwapOrientation == ((bool)Step))
+			{
+				for (int32 FaceIndex = 0; FaceIndex < TriangleCount; ++FaceIndex, Triangles += 3)
 				{
-					//for (int32 VertexIndex = 2; VertexIndex >= 0; --VertexIndex)
+					// Create Vertex instances and set their attributes
 					for (int32 VertexIndex = 0, TIndex = 2; VertexIndex < CornerCount; ++VertexIndex, --TIndex)
 					{
-						FVector2D UVValues(AlUVs[2 * Triangles[TIndex] + 0], AlUVs[2 * Triangles[TIndex] + 1]);
-						UVBBox += FVector(UVValues, 0.0f);
-						VertexInstanceUVs.Set(CornerVertexInstanceIDs[VertexIndex], 0, FVector2f(UVValues));
+						CornerVertexIDs[VertexIndex] = VertexPositionIDs[Triangles[TIndex] + VertexCount * Step];
+						CornerVertexInstanceIDs[VertexIndex] = MeshDescription.CreateVertexInstance(CornerVertexIDs[VertexIndex]);
+
+						// Set the normal
+						const float* CurNormal = &AlNormals[3 * Triangles[TIndex]];
+						// ConvertVector_ZUp_RightHanded
+						FVector3f UENormal(-CurNormal[0], CurNormal[1], CurNormal[2]);
+						UENormal = UENormal.GetSafeNormal();
+						if (Step > 0)
+						{
+							UENormal = SymmetricMatrix.TransformVector(UENormal);
+						}
+						else
+						{
+							UENormal *= -1.;
+						}
+						VertexInstanceNormals[CornerVertexInstanceIDs[VertexIndex]] = UENormal;
+					}
+					if (CornerVertexIDs[0] == CornerVertexIDs[1] || CornerVertexIDs[0] == CornerVertexIDs[2] || CornerVertexIDs[1] == CornerVertexIDs[2])
+					{
+						continue;
 					}
-				}
 
-				// Triangulate
-				const FPolygonID NewPolygonID = MeshDescription.CreatePolygon(PolyGroupId, CornerVertexInstanceIDs);
+					// Set the UV
+					if (bHasUVData)
+					{
+						//for (int32 VertexIndex = 2; VertexIndex >= 0; --VertexIndex)
+						for (int32 VertexIndex = 0, TIndex = 2; VertexIndex < CornerCount; ++VertexIndex, --TIndex)
+						{
+							FVector2D UVValues(AlUVs[2 * Triangles[TIndex] + 0], AlUVs[2 * Triangles[TIndex] + 1]);
+							UVBBox += FVector(UVValues, 0.0f);
+							VertexInstanceUVs.Set(CornerVertexInstanceIDs[VertexIndex], 0, FVector2f(UVValues));
+						}
+					}
+
+					// Triangulate
+					const FPolygonID NewPolygonID = MeshDescription.CreatePolygon(PolyGroupId, CornerVertexInstanceIDs);
+				}
 			}
-		}
-		else
-		{
-			for (int32 FaceIndex = 0; FaceIndex < TriangleCount; ++FaceIndex, Triangles += 3)
+			else
 			{
-				// Create Vertex instances and set their attributes
-				for (int32 VertexIndex = 0; VertexIndex < CornerCount; ++VertexIndex)
+				for (int32 FaceIndex = 0; FaceIndex < TriangleCount; ++FaceIndex, Triangles += 3)
 				{
-					CornerVertexIDs[VertexIndex] = VertexPositionIDs[Triangles[VertexIndex] + VertexCount * Step];
-					CornerVertexInstanceIDs[VertexIndex] = MeshDescription.CreateVertexInstance(CornerVertexIDs[VertexIndex]);
-
-					// Set the normal
-					const float* CurNormal = &AlNormals[3 * Triangles[VertexIndex]];
-
-					// ConvertVector_ZUp_RightHanded
-					FVector3f UENormal(-CurNormal[0], CurNormal[1], CurNormal[2]);
-					UENormal = UENormal.GetSafeNormal();
-					if (Step > 0)
+					// Create Vertex instances and set their attributes
+					for (int32 VertexIndex = 0; VertexIndex < CornerCount; ++VertexIndex)
 					{
-						UENormal = SymmetricMatrix.TransformVector(UENormal) * -1;
+						CornerVertexIDs[VertexIndex] = VertexPositionIDs[Triangles[VertexIndex] + VertexCount * Step];
+						CornerVertexInstanceIDs[VertexIndex] = MeshDescription.CreateVertexInstance(CornerVertexIDs[VertexIndex]);
+
+						// Set the normal
+						const float* CurNormal = &AlNormals[3 * Triangles[VertexIndex]];
+
+						// ConvertVector_ZUp_RightHanded
+						FVector3f UENormal(-CurNormal[0], CurNormal[1], CurNormal[2]);
+						UENormal = UENormal.GetSafeNormal();
+						if (Step > 0)
+						{
+							UENormal = SymmetricMatrix.TransformVector(UENormal) * -1;
+						}
+						VertexInstanceNormals[CornerVertexInstanceIDs[VertexIndex]] = (FVector3f)UENormal;
+					}
+					if (CornerVertexIDs[0] == CornerVertexIDs[1] || CornerVertexIDs[0] == CornerVertexIDs[2] || CornerVertexIDs[1] == CornerVertexIDs[2])
+					{
+						continue;
 					}
-					VertexInstanceNormals[CornerVertexInstanceIDs[VertexIndex]] = (FVector3f)UENormal;
-				}
-				if (CornerVertexIDs[0] == CornerVertexIDs[1] || CornerVertexIDs[0] == CornerVertexIDs[2] || CornerVertexIDs[1] == CornerVertexIDs[2])
-				{
-					continue;
-				}
 
-				// Set the UV
-				if (bHasUVData)
-				{
-					for (int32 VertexIndex = 0; VertexIndex < CornerCount; ++VertexIndex)
+					// Set the UV
+					if (bHasUVData)
 					{
-						FVector2D UVValues(AlUVs[2 * Triangles[VertexIndex] + 0], AlUVs[2 * Triangles[VertexIndex] + 1]);
-						UVBBox += FVector(UVValues, 0.0f);
-						VertexInstanceUVs.Set(CornerVertexInstanceIDs[VertexIndex], 0, FVector2f(UVValues));
+						for (int32 VertexIndex = 0; VertexIndex < CornerCount; ++VertexIndex)
+						{
+							FVector2D UVValues(AlUVs[2 * Triangles[VertexIndex] + 0], AlUVs[2 * Triangles[VertexIndex] + 1]);
+							UVBBox += FVector(UVValues, 0.0f);
+							VertexInstanceUVs.Set(CornerVertexInstanceIDs[VertexIndex], 0, FVector2f(UVValues));
+						}
 					}
-				}
 
-				// Triangulate
-				const FPolygonID NewPolygonID = MeshDescription.CreatePolygon(PolyGroupId, CornerVertexInstanceIDs);
+					// Triangulate
+					const FPolygonID NewPolygonID = MeshDescription.CreatePolygon(PolyGroupId, CornerVertexInstanceIDs);
+				}
 			}
 		}
+
+		return true;
+	}
+
+	TAlDagNodePtr<AlMeshNode> OpenModelUtils::TesselateDagLeaf(const AlDagNode& DagLeaf, ETesselatorType TessType, double Tolerance)
+	{
+		AlDagNode* TesselatedNode = nullptr;
+		statusCode TessStatus;
+
+		switch (TessType)
+		{
+		case(ETesselatorType::Accurate):
+			TessStatus = AlTesselate::chordHeightDeviationAccurate(TesselatedNode, &DagLeaf, Tolerance);
+			break;
+		case(ETesselatorType::Fast):
+		default:
+			TessStatus = AlTesselate::chordHeightDeviationFast(TesselatedNode, &DagLeaf, Tolerance);
+			break;
+		}
+
+		TAlDagNodePtr<AlMeshNode> MeshNode(TesselatedNode ? TesselatedNode->asMeshNodePtr() : nullptr);
+		if ((TessStatus == sSuccess) && MeshNode)
+		{
+			return MeshNode;
+		}
+		else
+		{
+			return TAlDagNodePtr<AlMeshNode>();
+		}
 	}
 
-	return true;
-}
+	bool OpenModelUtils::IsGeometryValid(const TAlDagNodePtr<AlDagNode>& GeomNode)
+	{
+		TAlObjectPtr<AlLayer> Layer(GeomNode->layer());
+		if (Layer.IsValid() && (bool)Layer->invisible())
+		{
+			return false;
+		}
 
+		bool bValidGeometry = false;
 
-TAlObjectPtr<AlDagNode> OpenModelUtils::TesselateDagLeaf(const AlDagNode& DagLeaf, ETesselatorType TessType, double Tolerance)
-{
-	AlDagNode* TesselatedNode = nullptr;
-	statusCode TessStatus;
+		{
+			TAlDagNodePtr<AlMeshNode> MeshNode(GeomNode->asMeshNodePtr());
+			if (MeshNode)
+			{
+				TAlObjectPtr<AlMesh> Mesh(MeshNode->mesh());
+				bValidGeometry |= Mesh.IsValid();
+			}
+		}
+		{
+			TAlDagNodePtr<AlShellNode> ShellNode(GeomNode->asShellNodePtr());
+			if (ShellNode)
+			{
+				TAlObjectPtr<AlShell> Shell(ShellNode->shell());
+				bValidGeometry |= Shell.IsValid();
+			}
+		}
+		{
+			TAlDagNodePtr<AlSurfaceNode> SurfaceNode(GeomNode->asSurfaceNodePtr());
+			if (SurfaceNode)
+			{
+				TAlObjectPtr<AlSurface> Surface(SurfaceNode->surface());
+				bValidGeometry |= Surface.IsValid();
+			}
+		}
+		return bValidGeometry;
+	}
 
-	switch (TessType)
+	bool OpenModelUtils::IsDagNodeValid(const TAlDagNodePtr<AlDagNode>& DagNode)
 	{
-	case(ETesselatorType::Accurate):
-		TessStatus = AlTesselate::chordHeightDeviationAccurate(TesselatedNode, &DagLeaf, Tolerance);
-		break;
-	case(ETesselatorType::Fast):
-	default:
-		TessStatus = AlTesselate::chordHeightDeviationFast(TesselatedNode, &DagLeaf, Tolerance);
-		break;
+		TAlObjectPtr<AlLayer> Layer(DagNode->layer());
+		if (Layer.IsValid() && (bool)Layer->invisible())
+		{
+			return false;
+		}
+
+		TAlDagNodePtr<AlGroupNode> GroupNode(DagNode->asGroupNodePtr());
+		if (GroupNode)
+		{
+			return true;
+		}
+
+		return IsGeometryValid(DagNode);
 	}
 
-	if ((TessStatus == sSuccess) && (AlIsValid(TesselatedNode) == TRUE))
+	CADLibrary::FMeshParameters OpenModelUtils::GetMeshParameters(const TAlObjectPtr<AlLayer>& Layer)
 	{
-		return TAlObjectPtr<AlDagNode>(TesselatedNode);
+		CADLibrary::FMeshParameters MeshParameters;
+
+		if (Layer)
+		{
+			if (Layer->isSymmetric())
+			{
+				MeshParameters.bIsSymmetric = true;
+				double Normal[3], Origin[3];
+				Layer->symmetricNormal(Normal[0], Normal[1], Normal[2]);
+				Layer->symmetricOrigin(Origin[0], Origin[1], Origin[2]);
+
+				MeshParameters.SymmetricOrigin.X = (float)Origin[0];
+				MeshParameters.SymmetricOrigin.Y = (float)Origin[1];
+				MeshParameters.SymmetricOrigin.Z = (float)Origin[2];
+				MeshParameters.SymmetricNormal.X = (float)Normal[0];
+				MeshParameters.SymmetricNormal.Y = (float)Normal[1];
+				MeshParameters.SymmetricNormal.Z = (float)Normal[2];
+			}
+		}
+
+		return MeshParameters;
 	}
-	else
+
+	CADLibrary::FMeshParameters OpenModelUtils::GetMeshParameters(AlDagNode& DagNode)
 	{
-		return TAlObjectPtr<AlDagNode>();
+		CADLibrary::FMeshParameters MeshParameters = GetMeshParameters(DagNode.layer());
+		
+		boolean bAlOrientation;
+		DagNode.getSurfaceOrientation(bAlOrientation);
+
+		MeshParameters.bNeedSwapOrientation = (DagNode.type() == kMeshNodeType) ? (bool)bAlOrientation : false;
+
+		return MeshParameters;
 	}
-}
 
 }
 
-#if PLATFORM_WINDOWS
-#include "Windows/HideWindowsPlatformTypes.h"
-#endif
-
 #endif
 
 
diff --git a/Engine/Plugins/Enterprise/DatasmithCADImporter/Source/WireInterface/Private/OpenModelUtils.h b/Engine/Plugins/Enterprise/DatasmithCADImporter/Source/WireInterface/Private/OpenModelUtils.h
--- a/Engine/Plugins/Enterprise/DatasmithCADImporter/Source/WireInterface/Private/OpenModelUtils.h
+++ b/Engine/Plugins/Enterprise/DatasmithCADImporter/Source/WireInterface/Private/OpenModelUtils.h
@@ -16,6 +16,9 @@
 
 #include "AlShadingFields.h"
 #include "AlDagNode.h"
+#include "AlMeshNode.h"
+#include "AlShellNode.h"
+#include "AlSurfaceNode.h"
 #include "AlPersistentID.h"
 
 #if PLATFORM_WINDOWS
@@ -23,18 +26,25 @@
 #endif
 
 class IDatasmithActorElement;
-class AlMesh;
 
 struct FMeshDescription;
 
+#define LAYER_TYPE          TEXT("Layer")
+#define GROUPNODE_TYPE      TEXT("GroupNode")
+#define MESH_TYPE           TEXT("Mesh")
+#define MESHNODE_TYPE       TEXT("MeshNode")
+#define SHADER_TYPE         TEXT("Shell")
+#define SHELLNODE_TYPE      TEXT("ShellNode")
+#define SHELL_TYPE          TEXT("Shell")
+#define SURFACE_TYPE        TEXT("Surface")
+#define SURFACENODE_TYPE    TEXT("SurfaceNode")
+
 namespace UE_DATASMITHWIRETRANSLATOR_NAMESPACE
 {
-	//enum class EAliasObjectReference : uint8
-	//{
-	//	LocalReference,
-	//	ParentReference,
-	//	WorldReference,
-	//};
+	namespace OpenModelUtils
+	{
+
+	}
 
 	enum class ETesselatorType : uint8
 	{
@@ -73,7 +83,22 @@ namespace UE_DATASMITHWIRETRANSLATOR_NAMESPACE
 			return Ptr;
 		}
 
-		TAlObjectPtr<T> operator=(TAlObjectPtr<T> Src)
+		FString GetName() const
+		{
+			return Ptr ? StringCast<TCHAR>(Ptr->name()).Get() : FString();
+		}
+
+		uint32 GetHash() const
+		{
+			return GetTypeHash(FString(StringCast<TCHAR>(Ptr->name()).Get()));
+		}
+
+		FString GetUniqueID(const TCHAR* TypeName = TEXT("Object")) const
+		{
+			return FString(TypeName) + FString::FromInt(GetHash());
+		}
+
+		TAlObjectPtr<T> operator=(const TAlObjectPtr<T>& Src)
 		{
 			Ptr = Src.IsValid() ? Src.Get() : nullptr;
 			return *this;
@@ -84,13 +109,237 @@ namespace UE_DATASMITHWIRETRANSLATOR_NAMESPACE
 			return IsValid();
 		}
 
-	private:
+		friend uint32 GetTypeHash(const TAlObjectPtr<T>& Object)
+		{
+			return Object.GetHash()/*GetTypeHash(DateTime.Ticks)*/;
+		}
+
+	protected:
 		T* Ptr = nullptr;
 	};
 
+	template<typename T>
+	inline bool operator==(const TAlObjectPtr<T>& A, const TAlObjectPtr<T>& B)
+	{
+		return (bool)AlAreEqual((const AlObject*)A.Get(), (const AlObject*)B.Get());
+	}
+
+	template<typename T>
+	inline bool operator!=(const TAlObjectPtr<T>& A, const TAlObjectPtr<T>& B)
+	{
+		return !((bool)AlAreEqual((const AlObject*)A.Get(), (const AlObject*)B.Get()));
+	}
+
+	template<>
+	uint32 TAlObjectPtr<AlLayer>::GetHash() const;
+
+	template<typename T>
+	class TAlDagNodePtr : public TAlObjectPtr<T>
+	{
+	public:
+		TAlDagNodePtr(T* AlDagNode = nullptr) : TAlObjectPtr<T>(AlDagNode) {}
+
+		FString GetLayerName() const
+		{
+			static_assert(std::is_base_of<AlDagNode, T>::value);
+			T* LocalPtr = TAlObjectPtr<T>::Ptr;
+			return LocalPtr && LocalPtr->layer() ? StringCast<FString,const char*>(LocalPtr->layer()->name()) : FString();
+		}
+
+		TAlObjectPtr<AlLayer> GetLayer() const
+		{
+			static_assert(std::is_base_of<AlDagNode, T>::value);
+			T* LocalPtr = TAlObjectPtr<T>::Ptr;
+			return TAlObjectPtr<AlLayer>(LocalPtr && LocalPtr->layer() ? LocalPtr->layer() : nullptr);
+		}
+
+		TAlObjectPtr<AlMesh> GetMesh() const
+		{
+			static_assert(std::is_base_of<AlDagNode, T>::value);
+			AlMeshNode* MeshNode = ((AlDagNode*)TAlObjectPtr<T>::Ptr)->asMeshNodePtr();
+			return MeshNode ? TAlObjectPtr<AlMesh>(MeshNode->mesh()) : TAlObjectPtr<AlMesh>();
+		}
+
+		TAlObjectPtr<AlSurface> GetSurface() const
+		{
+			static_assert(std::is_base_of<AlDagNode, T>::value);
+			AlSurfaceNode* SurfaceNode = ((AlDagNode*)TAlObjectPtr<T>::Ptr)->asSurfaceNodePtr();
+			return SurfaceNode ? TAlObjectPtr<AlSurface>(SurfaceNode->surface()) : TAlObjectPtr<AlSurface>();
+		}
+
+		TAlObjectPtr<AlShell> GetShell() const
+		{
+			static_assert(std::is_base_of<AlDagNode, T>::value);
+			AlShellNode* ShellNode = ((AlDagNode*)TAlObjectPtr<T>::Ptr)->asShellNodePtr();
+			return ShellNode ? TAlObjectPtr<AlShell>(ShellNode->shell()) : TAlObjectPtr<AlShell>();
+		}
+	};
+
+	class FPatchMesh
+	{
+	public:
+		FPatchMesh(const FString& InName, const TAlObjectPtr<AlLayer>& InLayer, int32 Count)
+			: Name(InName)
+			, Layer(InLayer)
+		{
+			MeshNodes.Reserve(Count);
+		}
+
+		bool HasContent()
+		{
+			return Initialize() && !MeshNodes.IsEmpty();
+		}
+
+		bool HasSingleContent() const
+		{
+			return MeshNodes.Num() == 1;
+		}
+
+		TAlDagNodePtr<AlMeshNode> GetSingleContent() const
+		{
+			if (MeshNodes.Num() == 1)
+			{
+				return MeshNodes[0];
+			}
+
+			return TAlDagNodePtr<AlMeshNode>();
+		}
+
+		const FString& GetName() const { return Name; }
+
+		uint32 GetHash() const { return Hash; }
+
+		const FString& GetUniqueID() const { return UniqueID; }
+
+		const TAlObjectPtr<AlLayer> GetLayer() const { return Layer; }
+
+		void AddMeshNode(const TAlDagNodePtr<AlMeshNode>& MeshNode)
+		{
+			ensure(Layer == MeshNode.GetLayer());
+			ensure(MeshNode->mesh());
+			MeshNodes.Add(MeshNode);
+		}
+
+		void IterateOnMeshNodes(const TFunction<void(const TAlDagNodePtr<AlMeshNode>& MeshNode)>& Callback) const
+		{
+			for (const TAlDagNodePtr<AlMeshNode>& MeshNode : MeshNodes)
+			{
+				Callback(MeshNode);
+			}
+		}
+
+		bool Initialize();
+
+	private:
+		FString Name;
+		TArray<TAlDagNodePtr<AlMeshNode>> MeshNodes;
+		TAlObjectPtr<AlLayer> Layer;
+		uint32 Hash;
+		FString UniqueID;
+		bool bInitialized = false;
+	};
+
+	class FBodyNode
+	{
+	public:
+		FBodyNode(const FString& InName, const TAlObjectPtr<AlLayer>& InLayer, int32 Count)
+			: Name(InName)
+			, Layer(InLayer)
+		{
+			SurfaceNodes.Reserve(Count);
+			ShellNodes.Reserve(Count);
+		}
+
+		bool HasContent()
+		{
+			return Initialize() && (!SurfaceNodes.IsEmpty() || !ShellNodes.IsEmpty());
+		}
+		 
+		bool HasSingleContent() const
+		{
+			return (SurfaceNodes.Num() + ShellNodes.Num()) == 1;
+		}
+
+		TAlDagNodePtr<AlDagNode> GetSingleContent() const
+		{
+			if (!HasSingleContent())
+			{
+				return TAlDagNodePtr<AlDagNode>();
+			}
+
+			if (SurfaceNodes.Num() == 1)
+			{
+				return TAlDagNodePtr<AlDagNode>(SurfaceNodes[0].Get());
+			}
+
+			if (ShellNodes.Num() == 1)
+			{
+				return TAlDagNodePtr<AlDagNode>(ShellNodes[0].Get());
+			}
+
+			return TAlDagNodePtr<AlDagNode>();
+		}
+
+		const FString& GetName() const { return Name; }
+
+		uint32 GetHash() const { return Hash; }
+
+		const FString& GetUniqueID() const { return UniqueID; }
+
+		const TAlObjectPtr<AlLayer> GetLayer() const { return Layer; }
+
+		void AddSurfaceNode(const TAlDagNodePtr<AlSurfaceNode>& SurfaceNode);
+
+		void AddShellNode(const TAlDagNodePtr<AlShellNode>& ShellNode);
+
+		void IterateOnSurfaceNodes(const TFunction<void(const TAlDagNodePtr<AlSurfaceNode>& SurfaceNode)>& Callback) const
+		{
+			for (const TAlDagNodePtr<AlSurfaceNode>& SurfaceNode : SurfaceNodes)
+			{
+				Callback(SurfaceNode);
+			}
+		}
+
+		void IterateOnShellNodes(const TFunction<void(const TAlDagNodePtr<AlShellNode>& ShellNode)>& Callback) const
+		{
+			for (const TAlDagNodePtr<AlShellNode>& ShellNode : ShellNodes)
+			{
+				Callback(ShellNode);
+			}
+		}
+
+		void IterateOnSlotIndices(const TFunction<void(int SlotIndex, const TAlObjectPtr<AlShader>& Shader)>& Callback) const
+		{
+			for (const TPair<int, TAlObjectPtr<AlShader>>& Entry : SlotIndexToShader)
+			{
+				Callback(Entry.Key, Entry.Value);
+			}
+		}
+
+		bool Initialize();
+
+		int32 GetSlotIndex(const TAlDagNodePtr<AlDagNode>& DagNode);
+
+	private:
+		FString Name;
+		TArray<TAlDagNodePtr<AlSurfaceNode>> SurfaceNodes;
+		TArray<TAlDagNodePtr<AlShellNode>> ShellNodes;
+		TAlObjectPtr<AlLayer> Layer;
+		TMap<FString, int> ShaderNameToSlotIndex;
+		TMap<int, TAlObjectPtr<AlShader>> SlotIndexToShader;
+		uint32 Hash;
+		FString UniqueID;
+		bool bInitialized = false;
+	};
+
 	namespace OpenModelUtils
 	{
-		void SetActorTransform(TSharedPtr<IDatasmithActorElement>& OutActorElement, const AlDagNode& InDagNode);
+		/** Following layer hierarchy, get list of layers an actor would be in as a csv string*/
+		bool GetCsvLayerString(const TAlObjectPtr<AlLayer>& Layer, FString& CsvString);
+
+		bool ActorHasContent(const TSharedPtr<IDatasmithActorElement>& ActorElement);
+
+		void SetActorTransform(IDatasmithActorElement& OutActorElement, const AlDagNode& InDagNode);
 
 		bool IsValidActor(const TSharedPtr<IDatasmithActorElement>& ActorElement);
 
@@ -114,13 +363,22 @@ namespace UE_DATASMITHWIRETRANSLATOR_NAMESPACE
 				DagNode.persistentID(PersistentID);
 				return GetTypeHash(*PersistentID);
 			}
-			FString Label = UTF8_TO_TCHAR(DagNode.name());
+			FString Label(StringCast<TCHAR>(DagNode.name()).Get());
 			return GetTypeHash(Label);
 		}
 
 		bool TransferAlMeshToMeshDescription(const AlMesh& Mesh, const TCHAR* SlotMaterialName, FMeshDescription& MeshDescription, CADLibrary::FMeshParameters& SymmetricParameters, const bool bMerge = false);
 
-		TAlObjectPtr<AlDagNode> TesselateDagLeaf(const AlDagNode& DagLeaf, ETesselatorType TessType, double Tolerance);
+		TAlDagNodePtr<AlMeshNode> TesselateDagLeaf(const AlDagNode& DagLeaf, ETesselatorType TessType, double Tolerance);
+
+		/** Returns true if DagNode is supported and its content is valid */
+		bool IsDagNodeValid(const TAlDagNodePtr<AlDagNode>& DagNode);
+
+		/** Returns true if geometry node is supported and its content is valid */
+		bool IsGeometryValid(const TAlDagNodePtr<AlDagNode>& GeomNode);
+
+		CADLibrary::FMeshParameters GetMeshParameters(AlDagNode& DagNode);
+		CADLibrary::FMeshParameters GetMeshParameters(const TAlObjectPtr<AlLayer>& Layer);
 	}
 
 }
diff --git a/Engine/Plugins/Enterprise/DatasmithCADImporter/Source/WireInterface/Private/WireInterfaceImpl.cpp b/Engine/Plugins/Enterprise/DatasmithCADImporter/Source/WireInterface/Private/WireInterfaceImpl.cpp
--- a/Engine/Plugins/Enterprise/DatasmithCADImporter/Source/WireInterface/Private/WireInterfaceImpl.cpp
+++ b/Engine/Plugins/Enterprise/DatasmithCADImporter/Source/WireInterface/Private/WireInterfaceImpl.cpp
@@ -48,6 +48,8 @@
 #include <AlMeshNode.h>
 #include <AlPersistentID.h>
 #include <AlRetrieveOptions.h>
+#include <AlSet.h>
+#include <AlSetMember.h>
 #include <AlShader.h>
 #include <AlShadingFieldItem.h>
 #include <AlShell.h>
@@ -85,7 +87,6 @@
 {
 
 	static bool bGSewByMaterial = false;
-	static bool bGLayersAsActors = false;
 
 	static const FColor DefaultColor = FColor(200, 200, 200);
 
@@ -132,124 +133,6 @@
 	const FString AliasSdkVersion   = TEXT("2023.1");
 #endif
 
-	class FBodyData
-	{
-	private:
-		TArray<TPair<TAlObjectPtr<AlDagNode>, FColor>> Shells;
-		TSet<FString> ShaderNames;
-
-		FString Label;
-		FString LayerName;
-		bool bCadData = true;
-
-	public:
-
-		FBodyData(const FString& InLayerName, bool bInCadData)
-			: bCadData(bInCadData)
-		{
-			LayerName = InLayerName;
-		};
-
-		bool IsCadData() const
-		{
-			return bCadData;
-		}
-
-		void SetLabel(const FString& InLabel)
-		{
-			Label = InLabel;
-		}
-
-		FColor LastColor;
-		bool bWarningNotLogYet = true;
-		void Add(const TAlObjectPtr<AlDagNode>& Node, FColor Color, const FString& ShaderName)
-		{
-			using namespace CADLibrary;
-			if (ShaderNames.Num() == FImportParameters::GMaxMaterialCountPerMesh)
-			{
-				if (!ShaderNames.Contains(ShaderName))
-				{
-					UE_LOG(LogWireInterface, Warning, TEXT("The main UE5 rendering systems do not support more than 256 materials per mesh and the limit has been defined to %d materials. Only the first %d materials are kept. The others are replaced by the last one"), FImportParameters::GMaxMaterialCountPerMesh, FImportParameters::GMaxMaterialCountPerMesh);
-					bWarningNotLogYet = false;
-					Color = LastColor;
-				}
-			}
-			else
-			{
-				ShaderNames.Add(ShaderName);
-				LastColor = Color;
-			}
-			Shells.Emplace(Node, Color);
-		}
-
-		void Reserve(int32 MaxElement)
-		{
-			Shells.Reserve(MaxElement);
-			ShaderNames.Reserve(MaxElement);
-		}
-
-		const TSet<FString>& GetShaderNames() const
-		{
-			return ShaderNames;
-		}
-
-		const FString& GetLayerName() const
-		{
-			return LayerName;
-		}
-
-		const TArray<TPair<TAlObjectPtr<AlDagNode>, FColor>>& GetShells() const
-		{
-			return Shells;
-		}
-
-		// Generates FBodyData's unique id from AlDagNode objects
-		uint32 GetUuid(const uint32& ParentUuid)
-		{
-			if (Shells.Num() == 0)
-			{
-				return ParentUuid;
-			}
-
-			if (Shells.Num() > 1)
-			{
-				Shells.Sort([&](const TPair<TAlObjectPtr<AlDagNode>, FColor>& NodeA, const TPair<TAlObjectPtr<AlDagNode>, FColor>& NodeB) { return OpenModelUtils::GetAlDagNodeUuid(*NodeA.Key) < OpenModelUtils::GetAlDagNodeUuid(*NodeB.Key); });
-			}
-
-			uint32 BodyUUID = 0;
-			for (const TPair<TAlObjectPtr<AlDagNode>, FColor>& DagNode : Shells)
-			{
-				BodyUUID = HashCombine(BodyUUID, OpenModelUtils::GetAlDagNodeUuid(*DagNode.Key));
-			}
-
-			return HashCombine(ParentUuid, BodyUUID);
-		}
-	};
-
-	struct FDagNodeInfo
-	{
-		uint32 Uuid;  // Use for actor name
-		FString Label;
-		TSharedPtr<IDatasmithActorElement> ActorElement;
-		FString LayerName;
-		bool bLayerIsVisible = true;
-
-		TMap<FString, TSharedPtr<IDatasmithActorElement>> LayerToActorElement;
-		FDagNodeInfo* ParentInfo = nullptr;
-
-		FDagNodeInfo(FDagNodeInfo& InParentInfo)
-			: ParentInfo(&InParentInfo)
-		{}
-
-		FDagNodeInfo()
-			: ParentInfo(nullptr)
-		{}
-
-		void ExtractDagNodeInfo(AlDagNode& CurrentNode);
-		void ExtractDagNodeInfo(FBodyData& CurrentNode);
-		void ExtractDagNodeLayer(const AlDagNode& InDagNode);
-	};
-
 	// Alias material management (to allow sew of BReps of different materials):
 	// To be compatible with "Retessellate" function, Alias material management as to be the same than CAD (TechSoft) import. 
 	// As a reminder: the name and slot of UE Material from CAD is based on CAD material/color data i.e. RGBA Color components => "UE Material slot" (int32) and "UE Material name" (FString = FString::FromInt("UE Material slot"))
@@ -304,19 +187,15 @@
 
 	bool FWireTranslatorImpl::Initialize(const TCHAR* InSceneFullName)
 	{
-		bool bCanProcess = true;
-
 		// TODO: Check file can be loaded with current SDK
 
 		if (InSceneFullName)
 		{
 			SceneFullPath = InSceneFullName;
-			SceneName = FPaths::GetBaseFilename(SceneFullPath);
-			CurrentPath = FPaths::GetPath(SceneFullPath);
 		}
 
 		// TODO: Check whether this implementation can load it
-		return bCanProcess;
+		return true;
 	}
 
 	FWireTranslatorImpl::~FWireTranslatorImpl()
@@ -328,55 +207,16 @@
 		bSceneLoaded = false;
 	}
 
-	bool FWireTranslatorImpl::GetCommonParameters(int32 Field, double Value, FColor& Color, FColor& TransparencyColor, FColor& IncandescenceColor, double GlowIntensity)
+	void FWireTranslatorImpl::SetImportSettings(const FWireSettings& Settings)
 	{
-		switch (AlShadingFields(Field))
-		{
-		case AlShadingFields::kFLD_SHADING_COMMON_COLOR_R:
-			Color.R = (uint8)Value;
-			return true;
-		case AlShadingFields::kFLD_SHADING_COMMON_COLOR_G:
-			Color.G = (uint8)Value;
-			return true;
-		case  AlShadingFields::kFLD_SHADING_COMMON_COLOR_B:
-			Color.B = (uint8)Value;
-			return true;
-		case AlShadingFields::kFLD_SHADING_COMMON_INCANDESCENCE_R:
-			IncandescenceColor.R = (uint8)Value;
-			return true;
-		case AlShadingFields::kFLD_SHADING_COMMON_INCANDESCENCE_G:
-			IncandescenceColor.G = (uint8)Value;
-			return true;
-		case AlShadingFields::kFLD_SHADING_COMMON_INCANDESCENCE_B:
-			IncandescenceColor.B = (uint8)Value;
-			return true;
-		case  AlShadingFields::kFLD_SHADING_COMMON_TRANSPARENCY_R:
-			TransparencyColor.R = (uint8)Value;
-			return true;
-		case  AlShadingFields::kFLD_SHADING_COMMON_TRANSPARENCY_G:
-			TransparencyColor.G = (uint8)Value;
-			return true;
-		case  AlShadingFields::kFLD_SHADING_COMMON_TRANSPARENCY_B:
-			TransparencyColor.B = (uint8)Value;
-			return true;
-		case AlShadingFields::kFLD_SHADING_COMMON_GLOW_INTENSITY:
-			GlowIntensity = Value;
-			return true;
-		default:
-			return false;
-		}
-	}
-
-	void FWireTranslatorImpl::SetTessellationOptions(const FDatasmithTessellationOptions& Options)
-	{
-		TessellationOptions = Options;
+		WireSettings = Settings;
 		if (CADModelConverter)
 		{
-			CADModelConverter->SetImportParameters(Options.ChordTolerance, Options.MaxEdgeLength, Options.NormalTolerance, (CADLibrary::EStitchingTechnique)Options.StitchingTechnique);
+			CADModelConverter->SetImportParameters(Settings.ChordTolerance, Settings.MaxEdgeLength, Settings.NormalTolerance, (CADLibrary::EStitchingTechnique)Settings.StitchingTechnique);
 		}
-		SceneFileHash = HashCombine(Options.GetHash(), GetSceneFileHash(SceneFullPath, SceneName));
 	}
 
+	// Wire file parsing
 	bool FWireTranslatorImpl::Load(TSharedPtr<IDatasmithScene> InScene)
 	{
 		DatasmithScene = InScene;
@@ -388,8 +228,8 @@
 		DatasmithScene->SetExporterSDKVersion(*AliasSdkVersion);
 		DatasmithScene->SetProductVersion(*AliasProductVersion);
 
-		bAliasUseNative = GetConsoleBoolValue(TEXT("ds.Wiretranslator.UseNative"), false);
-		if (!bAliasUseNative)
+		WireSettings.bAliasUseNative = GetConsoleBoolValue(TEXT("ds.Wiretranslator.UseNative"), false);
+		if (!WireSettings.bAliasUseNative)
 		{
 			CADLibrary::FImportParameters ImportParameters;
 			if (CADLibrary::FImportParameters::bGDisableCADKernelTessellation)
@@ -407,7 +247,6 @@
 		}
 
 		bGSewByMaterial = GetConsoleBoolValue(TEXT("ds.CADTranslator.Alias.SewByMaterial"), false);
-		bGLayersAsActors = GetConsoleBoolValue(TEXT("ds.CADTranslator.Alias.LayersAsActors"), false);
 
 		// Initialize Alias.
 		AlUniverse::initialize();
@@ -422,589 +261,1039 @@
 
 		bSceneLoaded = true;
 
-		// Make material
-		if (!GetShader())
+		return TraverseModel();
+	}
+
+	bool FWireTranslatorImpl::TraverseModel()
+	{
+		TAlDagNodePtr<AlDagNode> DagNode(AlUniverse::firstDagNode());
+		while (DagNode)
+		{
+			TSharedPtr<IDatasmithActorElement> ActorElement = TraverseDag(DagNode);
+			if (ActorElement)
+			{
+				DatasmithScene->AddActor(ActorElement);
+			}
+
+			DagNode = TAlDagNodePtr<AlDagNode>(DagNode->nextNode());
+		}
+
+		TAlObjectPtr<AlSet> Set(AlUniverse::firstSet());
+		while (Set)
+		{
+			// TODO: Add an actor to represent the set.
+			TAlObjectPtr<AlSetMember> SetMember(Set->firstMember());
+			while (SetMember)
+			{
+				TAlDagNodePtr<AlDagNode> DagNodeInSet(SetMember->object() ? SetMember->object()->asDagNodePtr() : nullptr);
+				if (DagNodeInSet)
+				{
+					TSharedPtr<IDatasmithActorElement> ActorElement = TraverseDag(DagNodeInSet);
+					if (ActorElement)
+					{
+						DatasmithScene->AddActor(ActorElement);
+					}
+				}
+
+				SetMember = TAlObjectPtr<AlSetMember>(SetMember->nextSetMember());
+			}
+
+			Set = TAlObjectPtr<AlSet>(Set->nextSet());
+		}
+
+		return true;
+	}
+
+	TSharedPtr<IDatasmithActorElement> FWireTranslatorImpl::TraverseDag(const TAlDagNodePtr<AlDagNode>& RootNode)
+	{
+		TAlDagNodePtr<AlGroupNode> GroupNode(RootNode->asGroupNodePtr());
+		if (GroupNode)
+		{
+			return WireSettings.bMergeGeomtryByGroup ? ProcessGroupNode(GroupNode) : TraverseGroupNode(GroupNode);
+		}
+		else if (OpenModelUtils::IsGeometryValid(RootNode))
+		{
+			return ProcessGeometryNode(RootNode);
+		}
+
+		return TSharedPtr<IDatasmithActorElement>();
+	}
+
+	TSharedPtr<IDatasmithActorElement> FWireTranslatorImpl::TraverseGroupNode(const TAlDagNodePtr<AlGroupNode>& GroupNode, const TAlObjectPtr<AlLayer>& ParentLayer)
+	{
+		TAlDagNodePtr<AlDagNode> ChildNode(GroupNode->childNode());
+		int32 ChildrenCount = 0;
+		while (ChildNode)
+		{
+			ChildrenCount++;
+			ChildNode = TAlDagNodePtr<AlDagNode>(ChildNode->nextNode());
+		}
+
+		if (ChildrenCount == 0)
+		{
+			return TSharedPtr<IDatasmithActorElement>();
+		}
+
+		TArray<TSharedPtr<IDatasmithActorElement>> ChildActors;
+		ChildActors.Reserve(ChildrenCount);
+
+		ChildNode = TAlDagNodePtr<AlDagNode>(GroupNode->childNode());
+		while (ChildNode)
+		{
+			TAlDagNodePtr<AlGroupNode> ChildGroupNode(ChildNode->asGroupNodePtr());
+			if (ChildGroupNode)
+			{
+				if (TSharedPtr<IDatasmithActorElement> ActorElement = TraverseGroupNode(ChildGroupNode, GroupNode->layer()))
+				{
+					ChildActors.Emplace(MoveTemp(ActorElement));
+				}
+
+			}
+			else if (OpenModelUtils::IsGeometryValid(ChildNode))
+			{
+				if (TSharedPtr<IDatasmithActorElement> ActorElement = ProcessGeometryNode(ChildNode, ParentLayer))
+				{
+					ChildActors.Emplace(MoveTemp(ActorElement));
+				}
+			}
+
+			ChildNode = TAlDagNodePtr<AlDagNode>(ChildNode->nextNode());
+		}
+
+		if (ChildActors.Num() == 0)
+		{
+			return TSharedPtr<IDatasmithActorElement>();
+		}
+
+		TSharedPtr<IDatasmithActorElement> ActorElement = FDatasmithSceneFactory::CreateActor(*GroupNode.GetUniqueID(GROUPNODE_TYPE));
+		if (!ActorElement.IsValid())
+		{
+			return ActorElement;
+		}
+
+		FString Label = GroupNode.GetName();
+		ActorElement->SetLabel(Label.Len() > 0 ? *Label : TEXT("UnnamedGroup"));
+
+		TAlObjectPtr<AlLayer> Layer(GroupNode->layer());
+		FString CsvLayerString;
+		if (OpenModelUtils::GetCsvLayerString(Layer, CsvLayerString))
+		{
+			ActorElement->SetLayer(*CsvLayerString);
+		}
+
+		OpenModelUtils::SetActorTransform(*ActorElement, *GroupNode);
+
+		for (TSharedPtr<IDatasmithActorElement>& ChildActor : ChildActors)
+		{
+			if (OpenModelUtils::ActorHasContent(ChildActor))
+			{
+				ActorElement->AddChild(ChildActor);
+			}
+		}
+
+		if (WireSettings.bUseLayerAsActor && Layer != ParentLayer)
+		{
+			if (OpenModelUtils::ActorHasContent(ActorElement))
+			{
+				if (TSharedPtr<IDatasmithActorElement> LayerActor = FindOrAddLayerActor(Layer))
+				{
+					LayerActor->AddChild(ActorElement);
+				}
+			}
+
+			return TSharedPtr<IDatasmithActorElement>();
+		}
+
+		return ActorElement;
+	}
+
+	TSharedPtr<IDatasmithActorElement> FWireTranslatorImpl::ProcessGroupNode(const TAlDagNodePtr<AlGroupNode>& GroupNode, const TAlObjectPtr<AlLayer>& ParentLayer)
+	{
+		TAlDagNodePtr<AlDagNode> ChildNode(GroupNode->childNode());
+		int32 ChildrenCount = 0;
+		while (ChildNode)
+		{
+			ChildrenCount++;
+			ChildNode = TAlDagNodePtr<AlDagNode>(ChildNode->nextNode());
+		}
+
+		if (ChildrenCount == 0)
+		{
+			return TSharedPtr<IDatasmithActorElement>();
+		}
+
+		TArray<TSharedPtr<IDatasmithActorElement>> ChildActors;
+		ChildActors.Reserve(ChildrenCount);
+
+		TSharedPtr<FBodyNode> BodyNode = MakeShared<FBodyNode>(GroupNode.GetName() + TEXT("_surf"), GroupNode->layer(), ChildrenCount);
+		TSharedPtr<FPatchMesh> PatchMesh = MakeShared<FPatchMesh>(GroupNode.GetName() + TEXT("_mesh"), GroupNode->layer(), ChildrenCount);
+
+		ChildNode = TAlDagNodePtr<AlDagNode>(GroupNode->childNode());
+		while (ChildNode)
+		{
+			TAlDagNodePtr<AlGroupNode> ChildGroupNode(ChildNode->asGroupNodePtr());
+			if (ChildGroupNode)
+			{
+				if (TSharedPtr<IDatasmithActorElement> ActorElement = ProcessGroupNode(ChildGroupNode, GroupNode->layer()))
+				{
+					ChildActors.Emplace(MoveTemp(ActorElement));
+				}
+
+			}
+			else if (OpenModelUtils::IsGeometryValid(ChildNode))
+			{
+				TAlDagNodePtr<AlMeshNode> MeshNode(ChildNode->asMeshNodePtr());
+				if (MeshNode && MeshNode->mesh())
+				{
+					PatchMesh->AddMeshNode(MeshNode);
+				}
+				TAlDagNodePtr<AlShellNode> ShellNode(ChildNode->asShellNodePtr());
+				if (ShellNode && ShellNode->shell())
+				{
+					BodyNode->AddShellNode(ShellNode);
+				}
+				else
+				{
+					TAlDagNodePtr<AlSurfaceNode> SurfaceNode(ChildNode->asSurfaceNodePtr());
+					if (SurfaceNode && SurfaceNode->surface())
+					{
+						BodyNode->AddSurfaceNode(SurfaceNode);
+					}
+				}
+			}
+
+			ChildNode = TAlDagNodePtr<AlDagNode>(ChildNode->nextNode());
+		}
+
+		if (ChildActors.IsEmpty() && !BodyNode->Initialize() && !PatchMesh->Initialize())
+		{
+			return TSharedPtr<IDatasmithActorElement>();
+		}
+
+		if (TSharedPtr<IDatasmithActorElement> ActorElement = ProcessBodyNode(BodyNode, GroupNode, GroupNode->layer()))
+		{
+			if (OpenModelUtils::ActorHasContent(ActorElement))
+			{
+				ChildActors.Emplace(MoveTemp(ActorElement));
+			}
+		}
+
+		if (TSharedPtr<IDatasmithActorElement> ActorElement = ProcessPatchMesh(PatchMesh, GroupNode, GroupNode->layer()))
+		{
+			if (OpenModelUtils::ActorHasContent(ActorElement))
+			{
+				ChildActors.Emplace(MoveTemp(ActorElement));
+			}
+		}
+
+		if (ChildActors.Num() == 0)
+		{
+			return TSharedPtr<IDatasmithActorElement>();
+		}
+
+		TSharedPtr<IDatasmithActorElement> ActorElement;
+		if (ChildActors.Num() == 1)
+		{
+			ActorElement = ChildActors[0];
+		}
+		else
+		{
+			ActorElement = FDatasmithSceneFactory::CreateActor(*GroupNode.GetUniqueID(GROUPNODE_TYPE));
+		}
+
+		if (!ActorElement.IsValid())
+		{
+			return ActorElement;
+		}
+
+		FString Label = GroupNode.GetName();
+		ActorElement->SetLabel(Label.Len() > 0 ? *Label : TEXT("UnnamedGroup"));
+
+		TAlObjectPtr<AlLayer> Layer(GroupNode->layer());
+		FString CsvLayerString;
+		if (OpenModelUtils::GetCsvLayerString(Layer, CsvLayerString))
+		{
+			ActorElement->SetLayer(*CsvLayerString);
+		}
+
+		if (ChildActors.Num() > 1)
+		{
+			OpenModelUtils::SetActorTransform(*ActorElement, *GroupNode);
+
+			for (TSharedPtr<IDatasmithActorElement>& ChildActor : ChildActors)
+			{
+				if (OpenModelUtils::ActorHasContent(ChildActor))
+				{
+					ActorElement->AddChild(ChildActor);
+				}
+			}
+
+		}
+
+		if (WireSettings.bUseLayerAsActor && Layer != ParentLayer)
+		{
+			if (OpenModelUtils::ActorHasContent(ActorElement))
+			{
+				if (TSharedPtr<IDatasmithActorElement> LayerActor = FindOrAddLayerActor(Layer))
+				{
+					LayerActor->AddChild(ActorElement);
+				}
+			}
+
+			return TSharedPtr<IDatasmithActorElement>();
+		}
+
+		return ActorElement;
+	}
+
+	TSharedPtr<IDatasmithActorElement> FWireTranslatorImpl::ProcessGeometryNode(const TAlDagNodePtr<AlDagNode>& GeomNode, const TAlObjectPtr<AlLayer>& ParentLayer)
+	{
+		TSharedPtr<IDatasmithMeshElement> MeshElement = FindOrAddMeshElement(GeomNode);
+		if (!MeshElement.IsValid())
+		{
+			return TSharedPtr<IDatasmithActorElement>();
+		}
+
+		TSharedPtr<IDatasmithMeshActorElement> ActorElement = FDatasmithSceneFactory::CreateMeshActor(*GeomNode.GetUniqueID(MESHNODE_TYPE));
+		if (!ActorElement.IsValid())
+		{
+			return TSharedPtr<IDatasmithActorElement>();
+		}
+
+		FString Label = GeomNode.GetName();
+		ActorElement->SetLabel(Label.Len() > 0 ? *Label : TEXT("NoName"));
+		ActorElement->SetStaticMeshPathName(MeshElement->GetName());
+
+		TAlObjectPtr<AlLayer> Layer(GeomNode->layer());
+		FString CsvLayerString;
+		if (OpenModelUtils::GetCsvLayerString(Layer, CsvLayerString))
+		{
+			ActorElement->SetLayer(*CsvLayerString);
+		}
+
+		OpenModelUtils::SetActorTransform(*ActorElement, *GeomNode);
+
+		TAlDagNodePtr<AlShellNode> ShellNode(GeomNode->asShellNodePtr());
+		if (ShellNode && ShellNode->shell())
+		{
+			TAlObjectPtr<AlShader> Shader(ShellNode->shell()->firstShader());
+			int32 SlotIndex = 0;
+			while (Shader.IsValid())
+			{
+				TSharedPtr<IDatasmithMaterialIDElement> MaterialIDElement = FindOrAddMaterial(Shader);
+				MaterialIDElement->SetId(SlotIndex++);
+				ActorElement->AddMaterialOverride(MaterialIDElement);
+			}
+		}
+		else
+		{
+			TSharedPtr<IDatasmithMaterialIDElement> MaterialIDElement;
+
+			TAlDagNodePtr<AlSurfaceNode> SurfaceNode(GeomNode->asSurfaceNodePtr());
+			if (SurfaceNode && SurfaceNode->surface())
+			{
+				MaterialIDElement = FindOrAddMaterial(SurfaceNode->surface()->firstShader());
+			}
+			else
+			{
+				TAlDagNodePtr<AlMeshNode> MeshNode(GeomNode->asMeshNodePtr());
+				if (MeshNode && MeshNode->mesh())
+				{
+					MaterialIDElement = FindOrAddMaterial(MeshNode->mesh()->firstShader());
+				}
+			}
+
+			if (MaterialIDElement)
+			{
+				MaterialIDElement->SetId(0);
+				ActorElement->AddMaterialOverride(MaterialIDElement);
+			}
+		}
+
+		if (WireSettings.bUseLayerAsActor && Layer != ParentLayer)
+		{
+			if (TSharedPtr<IDatasmithActorElement> LayerActor = FindOrAddLayerActor(Layer))
+			{
+				LayerActor->AddChild(ActorElement);
+				return TSharedPtr<IDatasmithActorElement>();
+			}
+
+			ensure(false);
+		}
+
+		return ActorElement;
+	}
+
+	TSharedPtr<IDatasmithActorElement> FWireTranslatorImpl::ProcessBodyNode(TSharedPtr<FBodyNode>& BodyNode, const TAlDagNodePtr<AlGroupNode>& GroupNode, const TAlObjectPtr<AlLayer>& ParentLayer)
+	{
+		if (!BodyNode->HasContent())
+		{
+			return TSharedPtr<IDatasmithActorElement>();
+		}
+
+		if (BodyNode->HasSingleContent())
+		{
+			return ProcessGeometryNode(BodyNode->GetSingleContent(), ParentLayer);
+		}
+
+		TSharedPtr<IDatasmithMeshElement> MeshElement = FindOrAddMeshElement(BodyNode);
+		if (!MeshElement.IsValid())
+		{
+			return TSharedPtr<IDatasmithActorElement>();
+		}
+
+		TSharedPtr<IDatasmithMeshActorElement> ActorElement = FDatasmithSceneFactory::CreateMeshActor(*BodyNode->GetUniqueID());
+		if (!ActorElement.IsValid())
+		{
+			return TSharedPtr<IDatasmithActorElement>();
+		}
+
+		FString Label = BodyNode->GetName();
+		ActorElement->SetLabel(Label.Len() > 0 ? *Label : TEXT("NoName"));
+		ActorElement->SetStaticMeshPathName(MeshElement->GetName());
+
+		FString CsvLayerString;
+		if (OpenModelUtils::GetCsvLayerString(BodyNode->GetLayer(), CsvLayerString))
+		{
+			ActorElement->SetLayer(*CsvLayerString);
+		}
+
+		if (!BodyNode->GetLayer()->isSymmetric())
+		{
+			OpenModelUtils::SetActorTransform(*ActorElement, *GroupNode);
+		}
+
+		if (WireSettings.bUseLayerAsActor && BodyNode->GetLayer() != ParentLayer)
+		{
+			if (TSharedPtr<IDatasmithActorElement> LayerActor = FindOrAddLayerActor(BodyNode->GetLayer()))
+			{
+				LayerActor->AddChild(ActorElement);
+				return TSharedPtr<IDatasmithActorElement>();
+			}
+
+			ensure(false);
+		}
+
+		return ActorElement;
+	}
+
+	TSharedPtr<IDatasmithActorElement> FWireTranslatorImpl::ProcessPatchMesh(TSharedPtr<FPatchMesh>& PatchMesh, const TAlDagNodePtr<AlGroupNode>& GroupNode, const TAlObjectPtr<AlLayer>& ParentLayer)
+	{
+		if (!PatchMesh->HasContent())
+		{
+			return TSharedPtr<IDatasmithActorElement>();
+		}
+
+		if (PatchMesh->HasSingleContent())
+		{
+			return ProcessGeometryNode(PatchMesh->GetSingleContent().Get(), ParentLayer);
+		}
+
+		TSharedPtr<IDatasmithMeshElement> MeshElement = FindOrAddMeshElement(PatchMesh);
+		if (!MeshElement.IsValid())
+		{
+			return TSharedPtr<IDatasmithActorElement>();
+		}
+
+		TSharedPtr<IDatasmithMeshActorElement> ActorElement = FDatasmithSceneFactory::CreateMeshActor(*PatchMesh->GetUniqueID());
+		if (!ActorElement.IsValid())
+		{
+			return TSharedPtr<IDatasmithActorElement>();
+		}
+
+		FString Label = PatchMesh->GetName();
+		ActorElement->SetLabel(Label.Len() > 0 ? *Label : TEXT("NoName"));
+		ActorElement->SetStaticMeshPathName(MeshElement->GetName());
+
+		FString CsvLayerString;
+		if (OpenModelUtils::GetCsvLayerString(PatchMesh->GetLayer(), CsvLayerString))
+		{
+			ActorElement->SetLayer(*CsvLayerString);
+		}
+
+		OpenModelUtils::SetActorTransform(*ActorElement, *GroupNode);
+
+		if (WireSettings.bUseLayerAsActor && PatchMesh->GetLayer() != ParentLayer)
+		{
+			if (TSharedPtr<IDatasmithActorElement> LayerActor = FindOrAddLayerActor(PatchMesh->GetLayer()))
+			{
+				LayerActor->AddChild(ActorElement);
+				return TSharedPtr<IDatasmithActorElement>();
+			}
+
+			ensure(false);
+		}
+
+		return ActorElement;
+	}
+
+	TSharedPtr<IDatasmithMeshElement> FWireTranslatorImpl::FindOrAddMeshElement(const TAlDagNodePtr<AlDagNode>& GeomNode)
+	{
+		// Look if geometry has not been already processed, return it if found
+		if (TSharedPtr<IDatasmithMeshElement>* MeshElementPtr = GeomNodeToMeshElement.Find(GeomNode.GetHash()))
+		{
+			return *MeshElementPtr;
+		}
+
+		if (!OpenModelUtils::IsGeometryValid(GeomNode))
+		{
+			// TODO: Log an error
+			return TSharedPtr<IDatasmithMeshElement>();
+		}
+
+		TSharedPtr<IDatasmithMeshElement> MeshElement = FDatasmithSceneFactory::CreateMesh(*GeomNode.GetUniqueID(MESH_TYPE));
+
+		MeshElement->SetLabel(*GeomNode.GetName());
+		MeshElement->SetLightmapSourceUV(-1);
+
+		auto ApplyMaterial = [this, &MeshElement](const TAlObjectPtr<AlShader>& Shader, int32 SlotIndex)
+			{
+				if (TSharedPtr<IDatasmithMaterialIDElement> MaterialIDElement = this->FindOrAddMaterial(Shader))
+				{
+					MaterialIDElement->SetId(SlotIndex);
+					MeshElement->SetMaterial(MaterialIDElement->GetName(), SlotIndex);
+				}
+			};
+
+
+		TAlDagNodePtr<AlShellNode> ShellNode(GeomNode->asShellNodePtr());
+		if (ShellNode)
+		{
+			TAlObjectPtr<AlShell> Shell(ShellNode->shell());
+			ensure(Shell);
+
+			TAlObjectPtr<AlShader> Shader(Shell->firstShader());
+			int32 SlotIndex = 0;
+			while (Shader)
+			{
+				ApplyMaterial(Shader, SlotIndex++);
+				Shader = TAlObjectPtr<AlShader>(Shell->nextShader(Shader.Get()));
+			}
+			// TODO: Check There are as many shaders as trim regions
+		}
+		else
+		{
+			TAlDagNodePtr<AlSurfaceNode> SurfaceNode(GeomNode->asSurfaceNodePtr());
+			if (SurfaceNode && SurfaceNode->surface())
+			{
+				ApplyMaterial(SurfaceNode->surface()->firstShader(), 0);
+			}
+			else
+			{
+				TAlDagNodePtr<AlMeshNode> MeshNode(GeomNode->asMeshNodePtr());
+				if (MeshNode && MeshNode->mesh())
+				{
+					ApplyMaterial(MeshNode->mesh()->firstShader(), 0);
+				}
+			}
+		}
+
+		DatasmithScene->AddMesh(MeshElement);
+		GeomNodeToMeshElement.Add(GeomNode.GetHash(), MeshElement);
+		MeshElementToGeomNode.Add(MeshElement, GeomNode.Get());
+
+		return MeshElement;
+	}
+
+	TSharedPtr<IDatasmithMeshElement> FWireTranslatorImpl::FindOrAddMeshElement(TSharedPtr<FBodyNode>& BodyNode)
+	{
+		// Look if geometry has not been already processed, return it if found
+		if (TSharedPtr<IDatasmithMeshElement>* MeshElementPtr = BodyNodeToMeshElement.Find(BodyNode->GetHash()))
+		{
+			return *MeshElementPtr;
+		}
+
+		if (!BodyNode->HasContent())
+		{
+			// TODO: Log an error
+			return TSharedPtr<IDatasmithMeshElement>();
+		}
+
+		TSharedPtr<IDatasmithMeshElement> MeshElement = FDatasmithSceneFactory::CreateMesh(*BodyNode->GetUniqueID());
+
+		MeshElement->SetLabel(*BodyNode->GetName());
+		MeshElement->SetLightmapSourceUV(-1);
+
+		auto ApplyMaterial = [this, &MeshElement](int32 SlotIndex, const TAlObjectPtr<AlShader>& Shader)
+			{
+				if (TSharedPtr<IDatasmithMaterialIDElement> MaterialIDElement = this->FindOrAddMaterial(Shader))
+				{
+					MaterialIDElement->SetId(SlotIndex);
+					MeshElement->SetMaterial(MaterialIDElement->GetName(), SlotIndex);
+				}
+			};
+
+
+		BodyNode->IterateOnSlotIndices(ApplyMaterial);
+
+		DatasmithScene->AddMesh(MeshElement);
+		BodyNodeToMeshElement.Add(BodyNode->GetHash(), MeshElement);
+		MeshElementToBodyNode.Add(MeshElement, BodyNode);
+
+		return MeshElement;
+	}
+
+	TSharedPtr<IDatasmithMeshElement> FWireTranslatorImpl::FindOrAddMeshElement(TSharedPtr<FPatchMesh>& PatchMesh)
+	{
+		// Look if geometry has not been already processed, return it if found
+		if (TSharedPtr<IDatasmithMeshElement>* MeshElementPtr = PatchMeshToMeshElement.Find(PatchMesh->GetHash()))
+		{
+			return *MeshElementPtr;
+		}
+
+		if (!PatchMesh->HasContent())
+		{
+			// TODO: Log an error
+			return TSharedPtr<IDatasmithMeshElement>();
+		}
+
+		TSharedPtr<IDatasmithMeshElement> MeshElement = FDatasmithSceneFactory::CreateMesh(*PatchMesh->GetUniqueID());
+
+		MeshElement->SetLabel(*PatchMesh->GetName());
+		MeshElement->SetLightmapSourceUV(-1);
+
+		auto ApplyMaterial = [this, &MeshElement](const TAlObjectPtr<AlShader>& Shader, int32 SlotIndex)
+			{
+				if (TSharedPtr<IDatasmithMaterialIDElement> MaterialIDElement = this->FindOrAddMaterial(Shader))
+				{
+					MaterialIDElement->SetId(SlotIndex);
+					MeshElement->SetMaterial(MaterialIDElement->GetName(), SlotIndex);
+				}
+			};
+
+
+		int32 SlotIndex = 0;
+		PatchMesh->IterateOnMeshNodes([&ApplyMaterial, &SlotIndex](const TAlDagNodePtr<AlMeshNode>& MeshNode)
+			{
+				if (MeshNode->mesh())
+				{
+					ApplyMaterial(MeshNode->mesh()->firstShader(), SlotIndex++);
+				}
+			});
+
+		DatasmithScene->AddMesh(MeshElement);
+		BodyNodeToMeshElement.Add(PatchMesh->GetHash(), MeshElement);
+		MeshElementToPatchMesh.Add(MeshElement, PatchMesh);
+
+		return MeshElement;
+	}
+
+	TSharedPtr<IDatasmithActorElement> FWireTranslatorImpl::FindOrAddLayerActor(const TAlObjectPtr<AlLayer>& Layer)
+	{
+		if (!WireSettings.bUseLayerAsActor || !Layer || Layer.GetName().IsEmpty())
+		{
+			return TSharedPtr<IDatasmithActorElement>();
+		}
+
+		TSharedPtr<IDatasmithActorElement>* LayerActorPtr = LayerToActor.Find(Layer.GetHash());
+		if (LayerActorPtr)
+		{
+			return *LayerActorPtr;
+		}
+
+		TSharedPtr<IDatasmithActorElement> ParentLayerActor;
+		TAlObjectPtr<AlLayer> ParentLayer(Layer->parentLayer());
+		if (ParentLayer)
+		{
+			ParentLayerActor = FindOrAddLayerActor(ParentLayer);
+		}
+
+		TSharedPtr<IDatasmithActorElement> LayerActor = FDatasmithSceneFactory::CreateActor(*Layer.GetUniqueID(LAYER_TYPE));
+
+		FString LayerName = Layer.GetName();
+		LayerActor->SetLabel(*LayerName);
+
+		FString CsvLayerString;
+		if (OpenModelUtils::GetCsvLayerString(Layer, CsvLayerString))
+		{
+			LayerActor->SetLayer(*CsvLayerString);
+		}
+
+		if (ParentLayerActor)
+		{
+			ParentLayerActor->AddChild(LayerActor);
+		}
+		else
+		{
+			DatasmithScene->AddActor(LayerActor);
+		}
+
+		LayerToActor.Add(Layer.GetHash(), LayerActor);
+
+		return LayerActor;
+	}
+
+	// Geometry retrieval
+
+#define TRACK_MESH 0
+
+	bool FWireTranslatorImpl::LoadStaticMesh(const TSharedPtr<IDatasmithMeshElement> MeshElement, FDatasmithMeshElementPayload& OutMeshPayload, const FDatasmithTessellationOptions& InTessellationOptions)
+	{
+#if TRACK_MESH
+		static const FString MeshName(TEXT("shell_30181"));
+		bool bSkip = true;
+		if (MeshName.Equals(MeshElement->GetLabel()))
+		{
+			bSkip = false;
+		}
+		if (bSkip)
 		{
 			return false;
 		}
 
-		// Parse and extract the DAG leaf nodes.
-		// Note that Alias file unit is cm like UE
-		return GetDagLeaves();
+		uint64 StartTime = FPlatformTime::Cycles64();
+#endif
+		CADLibrary::FMeshParameters MeshParameters;
+
+		if (TOptional<FMeshDescription> Mesh = GetMeshDescription(MeshElement, MeshParameters))
+		{
+			OutMeshPayload.LodMeshes.Add(MoveTemp(Mesh.GetValue()));
+			const TCHAR* MeshFilename = MeshElement->GetFile();
+			if (!WireSettings.bAliasUseNative && FPaths::FileExists(MeshFilename))
+			{
+				CADModelConverter->AddSurfaceDataForMesh(MeshFilename, MeshParameters, InTessellationOptions, OutMeshPayload);
+
+				// Remove the file because it is temporary since caching is disabled.
+				if (!CADLibrary::FImportParameters::bGEnableCADCache)
+				{
+					IFileManager::Get().Delete(MeshFilename);
+				}
+			}
+		}
+
+#if TRACK_MESH
+		double ElapsedSeconds = FPlatformTime::ToSeconds64(FPlatformTime::Cycles64() - StartTime);
+
+		UE_LOG(LogWireInterface, Display, TEXT("Mesh description of %s retrieved in %.3f s"), MeshElement->GetLabel(), ElapsedSeconds);
+#endif
+
+		return OutMeshPayload.LodMeshes.Num() > 0;
 	}
 
+	TOptional<FMeshDescription> FWireTranslatorImpl::GetMeshDescription(TSharedPtr<IDatasmithMeshElement> MeshElement, CADLibrary::FMeshParameters& MeshParameters)
+	{
+		if (TSharedPtr<FBodyNode>* BodyNodePtr = MeshElementToBodyNode.Find(MeshElement))
+		{
+			return GetMeshDescriptionFromBodyNode(*BodyNodePtr, MeshElement, MeshParameters);
+		}
+
+		if (TSharedPtr<FPatchMesh>* PatchMeshPtr = MeshElementToPatchMesh.Find(MeshElement))
+		{
+			return GetMeshDescriptionFromPatchMesh(*PatchMeshPtr, MeshElement, MeshParameters);
+		}
+
+		if (AlDagNode** GeomNodePtr = MeshElementToGeomNode.Find(MeshElement))
+		{
+			TAlDagNodePtr<AlDagNode> GeomNode(*GeomNodePtr);
+			if (GeomNode)
+			{
+				MeshParameters = OpenModelUtils::GetMeshParameters(*GeomNode);
+				switch (GeomNode->type())
+				{
+				case kShellNodeType:
+				case kSurfaceNodeType:
+				{
+					return GetMeshDescriptionFromParametricNode(*GeomNode, MeshElement, MeshParameters);
+					break;
+				}
+
+				case kMeshNodeType:
+				{
+					return GetMeshDescriptionFromMeshNode(GeomNode->asMeshNodePtr(), MeshElement, MeshParameters);
+					break;
+				}
+
+				default:
+					break;
+				}
+			}
+		}
+
+		return TOptional<FMeshDescription>();
+	}
+
+	TOptional<FMeshDescription> FWireTranslatorImpl::GetMeshDescriptionFromBodyNode(TSharedPtr<FBodyNode>& BodyNode, TSharedPtr<IDatasmithMeshElement> MeshElement, CADLibrary::FMeshParameters& MeshParameters)
+	{
+		MeshParameters = OpenModelUtils::GetMeshParameters(BodyNode->GetLayer());
+
+		IAliasBRepConverter* BRepConverter = nullptr;
+		TSharedPtr<CADLibrary::ICADModelConverter> ModelConverter = GetModelConverter(BRepConverter);
+
+		ModelConverter->InitializeProcess();
+
+		EAliasObjectReference ObjectReference = EAliasObjectReference::LocalReference;
+		if (MeshParameters.bIsSymmetric)
+		{
+			// All actors of a Alias symmetric layer are defined in the world Reference i.e. they have identity transform. So Mesh actor has to be defined in the world reference.
+			ObjectReference = EAliasObjectReference::WorldReference;
+		}
+		else if (WireSettings.StitchingTechnique != EDatasmithCADStitchingTechnique::StitchingNone)
+		{
+			// In the case of StitchingSew, AlDagNode children of a GroupNode are merged together. To be merged, they have to be defined in the reference of parent GroupNode.
+			ObjectReference = EAliasObjectReference::ParentReference;
+		}
+
+		BodyNode->IterateOnSurfaceNodes([&](const TAlDagNodePtr<AlSurfaceNode>& SurfaceNode)
+			{
+				BRepConverter->AddBRep(*SurfaceNode, BodyNode->GetSlotIndex(SurfaceNode.Get()), ObjectReference);
+			});
+
+		BodyNode->IterateOnShellNodes([&](const TAlDagNodePtr<AlShellNode>& ShellNode)
+			{
+				BRepConverter->AddBRep(*ShellNode, BodyNode->GetSlotIndex(ShellNode.Get()), ObjectReference);
+			});
+
+		ModelConverter->RepairTopology();
+
+		ModelConverter->SaveModel(*OutputPath, MeshElement);
+
+		FMeshDescription MeshDescription;
+		DatasmithMeshHelper::PrepareAttributeForStaticMesh(MeshDescription);
+
+		ModelConverter->Tessellate(MeshParameters, MeshDescription);
+
+		return MoveTemp(MeshDescription);
+	}
+
+	TOptional<FMeshDescription> FWireTranslatorImpl::GetMeshDescriptionFromPatchMesh(TSharedPtr<FPatchMesh>& PatchMesh, TSharedPtr<IDatasmithMeshElement> MeshElement, CADLibrary::FMeshParameters& MeshParameters)
+	{
+		MeshParameters = OpenModelUtils::GetMeshParameters(PatchMesh->GetLayer());
+
+		FMeshDescription MeshDescription;
+		DatasmithMeshHelper::PrepareAttributeForStaticMesh(MeshDescription);
+		MeshDescription.Empty();
+
+		constexpr bool bMerge = true;
+		PatchMesh->IterateOnMeshNodes([&](const TAlDagNodePtr<AlMeshNode>& MeshNode)
+			{
+				TAlObjectPtr<AlMesh> Mesh(MeshNode.GetMesh());
+				if (Mesh)
+				{
+					AlMatrix4x4 AlMatrix;
+					MeshNode->localTransformationMatrix(AlMatrix);
+
+					Mesh->transform(AlMatrix);
+
+					TSharedPtr<IDatasmithMaterialIDElement> MaterialIDElement = MeshElement->GetMaterialSlotAt(0);
+					const FString SlotMaterialName = MaterialIDElement ? MaterialIDElement->GetName() : FString();
+
+					OpenModelUtils::TransferAlMeshToMeshDescription(*Mesh, *SlotMaterialName, MeshDescription, MeshParameters, bMerge);
+				}
+			});
+
+		// Build edge meta data
+		FStaticMeshOperations::DetermineEdgeHardnessesFromVertexInstanceNormals(MeshDescription);
+
+		return MoveTemp(MeshDescription);
+	}
+
+	TOptional<FMeshDescription> FWireTranslatorImpl::GetMeshDescriptionFromParametricNode(AlDagNode& DagNode, TSharedPtr<IDatasmithMeshElement> MeshElement, CADLibrary::FMeshParameters& MeshParameters)
+	{
+		if (!WireSettings.bAliasUseNative)
+		{
+			return TessellateParametricNode(DagNode, MeshElement, MeshParameters);
+		}
+		else
+		{
+			// TODO: the best way, should be to don't have to apply inverse global transform to the generated mesh
+			TAlDagNodePtr<AlMeshNode> MeshNode = OpenModelUtils::TesselateDagLeaf(DagNode, ETesselatorType::Fast, WireSettings.ChordTolerance);
+
+			if (MeshNode && MeshNode->mesh())
+			{
+				AlMatrix4x4 AlMatrix;
+				DagNode.inverseGlobalTransformationMatrix(AlMatrix);
+				MeshNode->mesh()->transform(AlMatrix);
+
+				// Get the meshes from the dag nodes. Note that removing the mesh's DAG.
+				// will also removes the meshes, so we have to do it later.
+				return GetMeshDescriptionFromMeshNode(MeshNode, MeshElement, MeshParameters);
+			}
+		}
+
+		return TOptional<FMeshDescription>();
+	}
+
+	TOptional<FMeshDescription> FWireTranslatorImpl::GetMeshDescriptionFromMeshNode(const TAlDagNodePtr<AlMeshNode>& MeshNode, TSharedPtr<IDatasmithMeshElement> MeshElement, CADLibrary::FMeshParameters& MeshParameters, AlMatrix4x4* AlMeshInvGlobalMatrix)
+	{
+		if (!MeshNode)
+		{
+			return TOptional<FMeshDescription>();
+		}
+
+		TAlObjectPtr<AlMesh> Mesh(MeshNode->mesh());
+		if (!Mesh)
+		{
+			return TOptional<FMeshDescription>();
+		}
+
+		if (AlMeshInvGlobalMatrix != nullptr)
+		{
+			Mesh->transform(*AlMeshInvGlobalMatrix);
+		}
+
+		FMeshDescription MeshDescription;
+		DatasmithMeshHelper::PrepareAttributeForStaticMesh(MeshDescription);
+
+		TSharedPtr<IDatasmithMaterialIDElement> MaterialIDElement = MeshElement->GetMaterialSlotAt(0);
+		const FString SlotMaterialName = MaterialIDElement ? MaterialIDElement->GetName() : FString();
+
+		const bool bMerge = false;
+		OpenModelUtils::TransferAlMeshToMeshDescription(*Mesh, *SlotMaterialName, MeshDescription, MeshParameters, bMerge);
+
+		// Build edge meta data
+		FStaticMeshOperations::DetermineEdgeHardnessesFromVertexInstanceNormals(MeshDescription);
+
+		return MoveTemp(MeshDescription);
+	}
+
+	TSharedPtr<CADLibrary::ICADModelConverter> FWireTranslatorImpl::GetModelConverter(IAliasBRepConverter*& BRepConverter) const
+	{
+		TSharedPtr<CADLibrary::ICADModelConverter> ModelConverter;
+
+		CADLibrary::FImportParameters ImportParameters;
+		if (CADLibrary::FImportParameters::bGDisableCADKernelTessellation)
+		{
+			TSharedPtr<FAliasModelToTechSoftConverter> AliasToTechSoftConverter = MakeShared<FAliasModelToTechSoftConverter>(ImportParameters);
+			ModelConverter = AliasToTechSoftConverter;
+			BRepConverter = AliasToTechSoftConverter.Get();
+		}
+		else
+		{
+			TSharedPtr<FAliasModelToCADKernelConverter> AliasToCADKernelConverter = MakeShared<FAliasModelToCADKernelConverter>(ImportParameters);
+			ModelConverter = AliasToCADKernelConverter;
+			BRepConverter = AliasToCADKernelConverter.Get();
+		}
+
+		ModelConverter->SetImportParameters(WireSettings.ChordTolerance, WireSettings.MaxEdgeLength, WireSettings.NormalTolerance, (CADLibrary::EStitchingTechnique)WireSettings.StitchingTechnique);
+
+		return ModelConverter;
+	}
+
+	TOptional<FMeshDescription> FWireTranslatorImpl::TessellateParametricNode(AlDagNode& DagNode, TSharedPtr<IDatasmithMeshElement> MeshElement, CADLibrary::FMeshParameters& MeshParameters)
+	{
+		// Wire unit is cm
+		IAliasBRepConverter* BRepConverter = nullptr;
+		TSharedPtr<CADLibrary::ICADModelConverter> ModelConverter = GetModelConverter(BRepConverter);
+
+		ModelConverter->InitializeProcess();
+
+		EAliasObjectReference ObjectReference = EAliasObjectReference::LocalReference;
+
+		if (MeshParameters.bIsSymmetric)
+		{
+			// All actors of a Alias symmetric layer are defined in the world Reference i.e. they have identity transform. So Mesh actor has to be defined in the world reference.
+			ObjectReference = EAliasObjectReference::WorldReference;
+		}
+
+		ensure(MeshElement->GetMaterialSlotCount() == 1);
+
+		if (!BRepConverter->AddBRep(DagNode, 0, ObjectReference))
+		{
+			return TOptional<FMeshDescription>();
+		}
+
+		ModelConverter->RepairTopology();
+
+		ModelConverter->SaveModel(*OutputPath, MeshElement);
+
+		FMeshDescription MeshDescription;
+		DatasmithMeshHelper::PrepareAttributeForStaticMesh(MeshDescription);
+
+		bool bRet = ModelConverter->Tessellate(MeshParameters, MeshDescription);
+		if (!bRet)
+		{
+			const TCHAR* StaticMeshLable = MeshElement->GetLabel();
+			const TCHAR* StaticMeshName = MeshElement->GetName();
+			UE_LOG(LogWireInterface, Warning, TEXT("Failed to generate the mesh of \"%s\" (%s) StaticMesh."), StaticMeshLable, StaticMeshName);
+		}
+
+		return MoveTemp(MeshDescription);
+	}
+
+	// Material creation
+
+	TSharedPtr<IDatasmithMaterialIDElement> FWireTranslatorImpl::FindOrAddMaterial(const TAlObjectPtr<AlShader>& Shader)
+	{
+		const FString ShaderName = Shader.GetName();
+
+		if (TSharedPtr<IDatasmithBaseMaterialElement>* MaterialElementPtr = ShaderNameToMaterial.Find(ShaderName))
+		{
+			return FDatasmithSceneFactory::CreateMaterialId((*MaterialElementPtr)->GetName());
+		}
+
+		const FString ShaderModelName = Shader->shadingModel();
+
+		//const FColor Color = CreateShaderColorFromShaderName(ShaderName);
+		//ShaderNameToColor.Add(ShaderName, Color);
+
+		TSharedPtr<IDatasmithUEPbrMaterialElement> MaterialElement = FDatasmithSceneFactory::CreateUEPbrMaterial(*Shader.GetUniqueID(SHADER_TYPE));
+		MaterialElement->SetLabel(*ShaderName);
+
+		if (ShaderModelName.Equals(TEXT("BLINN")))
+		{
+			AddAlBlinnParameters(Shader, MaterialElement);
+		}
+		else if (ShaderModelName.Equals(TEXT("LAMBERT")))
+		{
+			AddAlLambertParameters(Shader, MaterialElement);
+		}
+		else if (ShaderModelName.Equals(TEXT("LIGHTSOURCE")))
+		{
+			AddAlLightSourceParameters(Shader, MaterialElement);
+		}
+		else if (ShaderModelName.Equals(TEXT("PHONG")))
+		{
+			AddAlPhongParameters(Shader, MaterialElement);
+		}
+
+		DatasmithScene->AddMaterial(MaterialElement);
+		ShaderNameToMaterial.Add(*ShaderName, MaterialElement);
+
+		return FDatasmithSceneFactory::CreateMaterialId(MaterialElement->GetName());
+	}
+
+	bool FWireTranslatorImpl::GetCommonParameters(int32 Field, double Value, FColor& Color, FColor& TransparencyColor, FColor& IncandescenceColor, double GlowIntensity)
+	{
+		switch (AlShadingFields(Field))
+		{
+		case AlShadingFields::kFLD_SHADING_COMMON_COLOR_R:
+			Color.R = (uint8)Value;
+			return true;
+		case AlShadingFields::kFLD_SHADING_COMMON_COLOR_G:
+			Color.G = (uint8)Value;
+			return true;
+		case  AlShadingFields::kFLD_SHADING_COMMON_COLOR_B:
+			Color.B = (uint8)Value;
+			return true;
+		case AlShadingFields::kFLD_SHADING_COMMON_INCANDESCENCE_R:
+			IncandescenceColor.R = (uint8)Value;
+			return true;
+		case AlShadingFields::kFLD_SHADING_COMMON_INCANDESCENCE_G:
+			IncandescenceColor.G = (uint8)Value;
+			return true;
+		case AlShadingFields::kFLD_SHADING_COMMON_INCANDESCENCE_B:
+			IncandescenceColor.B = (uint8)Value;
+			return true;
+		case  AlShadingFields::kFLD_SHADING_COMMON_TRANSPARENCY_R:
+			TransparencyColor.R = (uint8)Value;
+			return true;
+		case  AlShadingFields::kFLD_SHADING_COMMON_TRANSPARENCY_G:
+			TransparencyColor.G = (uint8)Value;
+			return true;
+		case  AlShadingFields::kFLD_SHADING_COMMON_TRANSPARENCY_B:
+			TransparencyColor.B = (uint8)Value;
+			return true;
+		case AlShadingFields::kFLD_SHADING_COMMON_GLOW_INTENSITY:
+			GlowIntensity = Value;
+			return true;
+		default:
+			return false;
+		}
+	}
+
 	void FWireTranslatorImpl::AddAlBlinnParameters(const TAlObjectPtr<AlShader>& Shader, TSharedPtr<IDatasmithUEPbrMaterialElement> MaterialElement)
 	{
 		// Default values for a Blinn material
 		FColor Color(145, 148, 153);
 		FColor TransparencyColor(0, 0, 0);
 		FColor IncandescenceColor(0, 0, 0);
-		FColor SpecularColor(38, 38, 38);
-		double Diffuse = 1.0;
-		double GlowIntensity = 0.0;
-		double Gloss = 0.8;
-		double Eccentricity = 0.35;
-		double Specularity = 1.0;
-		double Reflectivity = 0.5;
-		double SpecularRolloff = 0.5;
-
-		AlList* List = Shader->fields();
-		for (AlShadingFieldItem* Item = static_cast<AlShadingFieldItem*>(List->first()); Item; Item = Item->nextField())
-		{
-			double Value = 0.0f;
-			statusCode ErrorCode = Shader->parameter(Item->field(), Value);
-			if (ErrorCode != 0)
-			{
-				continue;
-			}
-
-			if (GetCommonParameters(Item->field(), Value, Color, TransparencyColor, IncandescenceColor, GlowIntensity))
-			{
-				continue;
-			}
-
-			switch (Item->field())
-			{
-			case AlShadingFields::kFLD_SHADING_BLINN_DIFFUSE:
-				Diffuse = Value;
-				break;
-			case AlShadingFields::kFLD_SHADING_BLINN_GLOSS_:
-				Gloss = Value;
-				break;
-			case AlShadingFields::kFLD_SHADING_BLINN_SPECULAR_R:
-				SpecularColor.R = (uint8)(255.f * Value);
-				break;
-			case AlShadingFields::kFLD_SHADING_BLINN_SPECULAR_G:
-				SpecularColor.G = (uint8)(255.f * Value);;
-				break;
-			case AlShadingFields::kFLD_SHADING_BLINN_SPECULAR_B:
-				SpecularColor.B = (uint8)(255.f * Value);;
-				break;
-			case AlShadingFields::kFLD_SHADING_BLINN_SPECULARITY_:
-				Specularity = Value;
-				break;
-			case AlShadingFields::kFLD_SHADING_BLINN_SPECULAR_ROLLOFF:
-				SpecularRolloff = Value;
-				break;
-			case AlShadingFields::kFLD_SHADING_BLINN_ECCENTRICITY:
-				Eccentricity = Value;
-				break;
-			case AlShadingFields::kFLD_SHADING_BLINN_REFLECTIVITY:
-				Reflectivity = Value;
-				break;
-			}
-		}
-
-		bool bIsTransparent = IsTransparent(TransparencyColor);
-
-		// Construct parameter expressions
-		IDatasmithMaterialExpressionScalar* DiffuseExpression = MaterialElement->AddMaterialExpression<IDatasmithMaterialExpressionScalar>();
-		DiffuseExpression->GetScalar() = Diffuse;
-		DiffuseExpression->SetName(TEXT("Diffuse"));
-
-		IDatasmithMaterialExpressionScalar* GlossExpression = MaterialElement->AddMaterialExpression<IDatasmithMaterialExpressionScalar>();
-		GlossExpression->GetScalar() = Gloss;
-		GlossExpression->SetName(TEXT("Gloss"));
-
-		IDatasmithMaterialExpressionColor* SpecularColorExpression = MaterialElement->AddMaterialExpression<IDatasmithMaterialExpressionColor>();
-		SpecularColorExpression->SetName(TEXT("SpecularColor"));
-		SpecularColorExpression->GetColor() = FLinearColor::FromSRGBColor(SpecularColor);
-
-		IDatasmithMaterialExpressionScalar* SpecularityExpression = MaterialElement->AddMaterialExpression<IDatasmithMaterialExpressionScalar>();
-		SpecularityExpression->GetScalar() = Specularity * 0.3;
-		SpecularityExpression->SetName(TEXT("Specularity"));
-
-		IDatasmithMaterialExpressionScalar* SpecularRolloffExpression = MaterialElement->AddMaterialExpression<IDatasmithMaterialExpressionScalar>();
-		SpecularRolloffExpression->GetScalar() = SpecularRolloff;
-		SpecularRolloffExpression->SetName(TEXT("SpecularRolloff"));
-
-		IDatasmithMaterialExpressionScalar* EccentricityExpression = MaterialElement->AddMaterialExpression<IDatasmithMaterialExpressionScalar>();
-		EccentricityExpression->GetScalar() = Eccentricity;
-		EccentricityExpression->SetName(TEXT("Eccentricity"));
-
-		IDatasmithMaterialExpressionScalar* ReflectivityExpression = MaterialElement->AddMaterialExpression<IDatasmithMaterialExpressionScalar>();
-		ReflectivityExpression->GetScalar() = Reflectivity;
-		ReflectivityExpression->SetName(TEXT("Reflectivity"));
-
-		IDatasmithMaterialExpressionColor* ColorExpression = MaterialElement->AddMaterialExpression<IDatasmithMaterialExpressionColor>();
-		ColorExpression->SetName(TEXT("Color"));
-		ColorExpression->GetColor() = FLinearColor::FromSRGBColor(Color);
-
-		IDatasmithMaterialExpressionColor* IncandescenceColorExpression = MaterialElement->AddMaterialExpression<IDatasmithMaterialExpressionColor>();
-		IncandescenceColorExpression->SetName(TEXT("IncandescenceColor"));
-		IncandescenceColorExpression->GetColor() = FLinearColor::FromSRGBColor(IncandescenceColor);
-
-		IDatasmithMaterialExpressionColor* TransparencyColorExpression = MaterialElement->AddMaterialExpression<IDatasmithMaterialExpressionColor>();
-		TransparencyColorExpression->SetName(TEXT("TransparencyColor"));
-		TransparencyColorExpression->GetColor() = FLinearColor::FromSRGBColor(TransparencyColor);
-
-		IDatasmithMaterialExpressionScalar* GlowIntensityExpression = MaterialElement->AddMaterialExpression<IDatasmithMaterialExpressionScalar>();
-		GlowIntensityExpression->GetScalar() = GlowIntensity;
-		GlowIntensityExpression->SetName(TEXT("GlowIntensity"));
-
-		// Create aux expressions
-		IDatasmithMaterialExpressionGeneric* ColorSpecLerp = MaterialElement->AddMaterialExpression<IDatasmithMaterialExpressionGeneric>();
-		ColorSpecLerp->SetExpressionName(TEXT("LinearInterpolate"));
-
-		IDatasmithMaterialExpressionScalar* ColorSpecLerpValue = MaterialElement->AddMaterialExpression<IDatasmithMaterialExpressionScalar>();
-		ColorSpecLerpValue->GetScalar() = 0.96f;
-
-		IDatasmithMaterialExpressionGeneric* ColorMetallicLerp = MaterialElement->AddMaterialExpression<IDatasmithMaterialExpressionGeneric>();
-		ColorMetallicLerp->SetExpressionName(TEXT("LinearInterpolate"));
-
-		IDatasmithMaterialExpressionGeneric* DiffuseLerp = MaterialElement->AddMaterialExpression<IDatasmithMaterialExpressionGeneric>();
-		DiffuseLerp->SetExpressionName(TEXT("LinearInterpolate"));
-
-		IDatasmithMaterialExpressionScalar* DiffuseLerpA = MaterialElement->AddMaterialExpression<IDatasmithMaterialExpressionScalar>();
-		DiffuseLerpA->GetScalar() = 0.04f;
-
-		IDatasmithMaterialExpressionScalar* DiffuseLerpB = MaterialElement->AddMaterialExpression<IDatasmithMaterialExpressionScalar>();
-		DiffuseLerpB->GetScalar() = 1.0f;
-
-		IDatasmithMaterialExpressionGeneric* BaseColorMultiply = MaterialElement->AddMaterialExpression<IDatasmithMaterialExpressionGeneric>();
-		BaseColorMultiply->SetExpressionName(TEXT("Multiply"));
-
-		IDatasmithMaterialExpressionGeneric* BaseColorAdd = MaterialElement->AddMaterialExpression<IDatasmithMaterialExpressionGeneric>();
-		BaseColorAdd->SetExpressionName(TEXT("Add"));
-
-		IDatasmithMaterialExpressionGeneric* BaseColorTransparencyMultiply = MaterialElement->AddMaterialExpression<IDatasmithMaterialExpressionGeneric>();
-		BaseColorTransparencyMultiply->SetExpressionName(TEXT("Multiply"));
-
-		IDatasmithMaterialExpressionGeneric* IncandescenceMultiply = MaterialElement->AddMaterialExpression<IDatasmithMaterialExpressionGeneric>();
-		IncandescenceMultiply->SetExpressionName(TEXT("Multiply"));
-
-		IDatasmithMaterialExpressionGeneric* IncandescenceScaleMultiply = MaterialElement->AddMaterialExpression<IDatasmithMaterialExpressionGeneric>();
-		IncandescenceScaleMultiply->SetExpressionName(TEXT("Multiply"));
-
-		IDatasmithMaterialExpressionScalar* IncandescenceScale = MaterialElement->AddMaterialExpression<IDatasmithMaterialExpressionScalar>();
-		IncandescenceScale->GetScalar() = 100.0f;
-
-		IDatasmithMaterialExpressionGeneric* EccentricityMultiply = MaterialElement->AddMaterialExpression<IDatasmithMaterialExpressionGeneric>();
-		EccentricityMultiply->SetExpressionName(TEXT("Multiply"));
-
-		IDatasmithMaterialExpressionGeneric* EccentricityOneMinus = MaterialElement->AddMaterialExpression<IDatasmithMaterialExpressionGeneric>();
-		EccentricityOneMinus->SetExpressionName(TEXT("OneMinus"));
-
-		IDatasmithMaterialExpressionGeneric* RoughnessOneMinus = MaterialElement->AddMaterialExpression<IDatasmithMaterialExpressionGeneric>();
-		RoughnessOneMinus->SetExpressionName(TEXT("OneMinus"));
-
-		IDatasmithMaterialExpressionScalar* FresnelExponent = MaterialElement->AddMaterialExpression<IDatasmithMaterialExpressionScalar>();
-		FresnelExponent->GetScalar() = 4.0f;
-
-		IDatasmithMaterialExpressionFunctionCall* FresnelFunc = MaterialElement->AddMaterialExpression<IDatasmithMaterialExpressionFunctionCall>();
-		FresnelFunc->SetFunctionPathName(TEXT("/Engine/Functions/Engine_MaterialFunctions02/Fresnel_Function.Fresnel_Function"));
-
-		IDatasmithMaterialExpressionGeneric* FresnelLerp = MaterialElement->AddMaterialExpression<IDatasmithMaterialExpressionGeneric>();
-		FresnelLerp->SetExpressionName(TEXT("LinearInterpolate"));
-
-		IDatasmithMaterialExpressionScalar* FresnelLerpA = MaterialElement->AddMaterialExpression<IDatasmithMaterialExpressionScalar>();
-		FresnelLerpA->GetScalar() = 1.0f;
-
-		IDatasmithMaterialExpressionScalar* SpecularPowerExp = MaterialElement->AddMaterialExpression<IDatasmithMaterialExpressionScalar>();
-		SpecularPowerExp->GetScalar() = 0.5f;
-
-		IDatasmithMaterialExpressionGeneric* Power = MaterialElement->AddMaterialExpression<IDatasmithMaterialExpressionGeneric>();
-		Power->SetExpressionName(TEXT("Power"));
-
-		IDatasmithMaterialExpressionGeneric* FresnelMultiply = MaterialElement->AddMaterialExpression<IDatasmithMaterialExpressionGeneric>();
-		FresnelMultiply->SetExpressionName(TEXT("Multiply"));
-
-		IDatasmithMaterialExpressionGeneric* TransparencyOneMinus = MaterialElement->AddMaterialExpression<IDatasmithMaterialExpressionGeneric>();
-		TransparencyOneMinus->SetExpressionName(TEXT("OneMinus"));
-
-		IDatasmithMaterialExpressionFunctionCall* BreakFloat3 = nullptr;
-		IDatasmithMaterialExpressionGeneric* AddRG = nullptr;
-		IDatasmithMaterialExpressionGeneric* AddRGB = nullptr;
-		IDatasmithMaterialExpressionGeneric* Divide = nullptr;
-		IDatasmithMaterialExpressionScalar* DivideConstant = nullptr;
-		if (bIsTransparent)
-		{
-			BreakFloat3 = MaterialElement->AddMaterialExpression<IDatasmithMaterialExpressionFunctionCall>();
-			BreakFloat3->SetFunctionPathName(TEXT("/Engine/Functions/Engine_MaterialFunctions02/Utility/BreakFloat3Components.BreakFloat3Components"));
-
-			AddRG = MaterialElement->AddMaterialExpression<IDatasmithMaterialExpressionGeneric>();
-			AddRG->SetExpressionName(TEXT("Add"));
-
-			AddRGB = MaterialElement->AddMaterialExpression<IDatasmithMaterialExpressionGeneric>();
-			AddRGB->SetExpressionName(TEXT("Add"));
-
-			Divide = MaterialElement->AddMaterialExpression<IDatasmithMaterialExpressionGeneric>();
-			Divide->SetExpressionName(TEXT("Divide"));
-
-			DivideConstant = MaterialElement->AddMaterialExpression<IDatasmithMaterialExpressionScalar>();
-			DivideConstant->GetScalar() = 3.0f;
-		}
-
-		// Connect expressions
-		SpecularColorExpression->ConnectExpression(*ColorSpecLerp->GetInput(0));
-		ColorExpression->ConnectExpression(*ColorSpecLerp->GetInput(1));
-		ColorSpecLerpValue->ConnectExpression(*ColorSpecLerp->GetInput(2));
-
-		ColorExpression->ConnectExpression(*ColorMetallicLerp->GetInput(0));
-		ColorSpecLerp->ConnectExpression(*ColorMetallicLerp->GetInput(1));
-		GlossExpression->ConnectExpression(*ColorMetallicLerp->GetInput(2));
-
-		DiffuseLerpA->ConnectExpression(*DiffuseLerp->GetInput(0));
-		DiffuseLerpB->ConnectExpression(*DiffuseLerp->GetInput(1));
-		DiffuseExpression->ConnectExpression(*DiffuseLerp->GetInput(2));
-
-		ColorMetallicLerp->ConnectExpression(*BaseColorMultiply->GetInput(0));
-		DiffuseLerp->ConnectExpression(*BaseColorMultiply->GetInput(1));
-
-		BaseColorMultiply->ConnectExpression(*BaseColorAdd->GetInput(0));
-		IncandescenceColorExpression->ConnectExpression(*BaseColorAdd->GetInput(1));
-
-		BaseColorAdd->ConnectExpression(*BaseColorTransparencyMultiply->GetInput(0));
-		TransparencyOneMinus->ConnectExpression(*BaseColorTransparencyMultiply->GetInput(1));
-
-		GlowIntensityExpression->ConnectExpression(*IncandescenceScaleMultiply->GetInput(0));
-		IncandescenceScale->ConnectExpression(*IncandescenceScaleMultiply->GetInput(1));
-
-		BaseColorTransparencyMultiply->ConnectExpression(*IncandescenceMultiply->GetInput(0));
-		IncandescenceScaleMultiply->ConnectExpression(*IncandescenceMultiply->GetInput(1));
-
-		EccentricityExpression->ConnectExpression(*EccentricityOneMinus->GetInput(0));
-
-		EccentricityOneMinus->ConnectExpression(*EccentricityMultiply->GetInput(0));
-		SpecularityExpression->ConnectExpression(*EccentricityMultiply->GetInput(1));
-
-		EccentricityMultiply->ConnectExpression(*RoughnessOneMinus->GetInput(0));
-
-		FresnelExponent->ConnectExpression(*FresnelFunc->GetInput(3));
-
-		SpecularRolloffExpression->ConnectExpression(*Power->GetInput(0));
-		SpecularPowerExp->ConnectExpression(*Power->GetInput(1));
-
-		FresnelLerpA->ConnectExpression(*FresnelLerp->GetInput(0));
-		FresnelFunc->ConnectExpression(*FresnelLerp->GetInput(1));
-		Power->ConnectExpression(*FresnelLerp->GetInput(2));
-
-		FresnelLerp->ConnectExpression(*FresnelMultiply->GetInput(0));
-		ReflectivityExpression->ConnectExpression(*FresnelMultiply->GetInput(1));
-
-		TransparencyColorExpression->ConnectExpression(*TransparencyOneMinus->GetInput(0));
-
-		if (bIsTransparent)
-		{
-			TransparencyOneMinus->ConnectExpression(*BreakFloat3->GetInput(0));
-
-			BreakFloat3->ConnectExpression(*AddRG->GetInput(0), 0);
-			BreakFloat3->ConnectExpression(*AddRG->GetInput(1), 1);
-
-			AddRG->ConnectExpression(*AddRGB->GetInput(0));
-			BreakFloat3->ConnectExpression(*AddRGB->GetInput(1), 2);
-
-			AddRGB->ConnectExpression(*Divide->GetInput(0));
-			DivideConstant->ConnectExpression(*Divide->GetInput(1));
-		}
-
-		// Connect material outputs
-		MaterialElement->GetBaseColor().SetExpression(BaseColorTransparencyMultiply);
-		MaterialElement->GetMetallic().SetExpression(GlossExpression);
-		MaterialElement->GetSpecular().SetExpression(FresnelMultiply);
-		MaterialElement->GetRoughness().SetExpression(RoughnessOneMinus);
-		MaterialElement->GetEmissiveColor().SetExpression(IncandescenceMultiply);
-
-		if (bIsTransparent)
-		{
-			MaterialElement->GetOpacity().SetExpression(Divide);
-			MaterialElement->SetParentLabel(TEXT("M_DatasmithAliasBlinnTransparent"));
-		}
-		else
-		{
-			MaterialElement->SetParentLabel(TEXT("M_DatasmithAliasBlinn"));
-		}
-
-	}
-
-	void FWireTranslatorImpl::AddAlLambertParameters(const TAlObjectPtr<AlShader>& Shader, TSharedPtr<IDatasmithUEPbrMaterialElement> MaterialElement)
-	{
-		// Default values for a Lambert material
-		FColor Color(145, 148, 153);
-		FColor TransparencyColor(0, 0, 0);
-		FColor IncandescenceColor(0, 0, 0);
-		double Diffuse = 1.0;
-		double GlowIntensity = 0.0;
-
-		AlList* List = Shader->fields();
-		for (AlShadingFieldItem* Item = static_cast<AlShadingFieldItem*>(List->first()); Item; Item = Item->nextField())
-		{
-			double Value = 0.0f;
-			statusCode ErrorCode = Shader->parameter(Item->field(), Value);
-			if (ErrorCode != 0)
-			{
-				continue;
-			}
-
-			if (GetCommonParameters(Item->field(), Value, Color, TransparencyColor, IncandescenceColor, GlowIntensity))
-			{
-				continue;
-			}
-
-			switch (Item->field())
-			{
-			case AlShadingFields::kFLD_SHADING_LAMBERT_DIFFUSE:
-				Diffuse = Value;
-				break;
-			}
-		}
-
-		bool bIsTransparent = IsTransparent(TransparencyColor);
-
-		// Construct parameter expressions
-		IDatasmithMaterialExpressionScalar* DiffuseExpression = MaterialElement->AddMaterialExpression<IDatasmithMaterialExpressionScalar>();
-		DiffuseExpression->GetScalar() = Diffuse;
-		DiffuseExpression->SetName(TEXT("Diffuse"));
-
-		IDatasmithMaterialExpressionColor* ColorExpression = MaterialElement->AddMaterialExpression<IDatasmithMaterialExpressionColor>();
-		ColorExpression->SetName(TEXT("Color"));
-		ColorExpression->GetColor() = FLinearColor::FromSRGBColor(Color);
-
-		IDatasmithMaterialExpressionColor* IncandescenceColorExpression = MaterialElement->AddMaterialExpression<IDatasmithMaterialExpressionColor>();
-		IncandescenceColorExpression->SetName(TEXT("IncandescenceColor"));
-		IncandescenceColorExpression->GetColor() = FLinearColor::FromSRGBColor(IncandescenceColor);
-
-		IDatasmithMaterialExpressionColor* TransparencyColorExpression = MaterialElement->AddMaterialExpression<IDatasmithMaterialExpressionColor>();
-		TransparencyColorExpression->SetName(TEXT("TransparencyColor"));
-		TransparencyColorExpression->GetColor() = FLinearColor::FromSRGBColor(TransparencyColor);
-
-		IDatasmithMaterialExpressionScalar* GlowIntensityExpression = MaterialElement->AddMaterialExpression<IDatasmithMaterialExpressionScalar>();
-		GlowIntensityExpression->GetScalar() = GlowIntensity;
-		GlowIntensityExpression->SetName(TEXT("GlowIntensity"));
-
-		// Create aux expressions
-		IDatasmithMaterialExpressionGeneric* DiffuseLerp = MaterialElement->AddMaterialExpression<IDatasmithMaterialExpressionGeneric>();
-		DiffuseLerp->SetExpressionName(TEXT("LinearInterpolate"));
-
-		IDatasmithMaterialExpressionScalar* DiffuseLerpA = MaterialElement->AddMaterialExpression<IDatasmithMaterialExpressionScalar>();
-		DiffuseLerpA->GetScalar() = 0.04f;
-
-		IDatasmithMaterialExpressionScalar* DiffuseLerpB = MaterialElement->AddMaterialExpression<IDatasmithMaterialExpressionScalar>();
-		DiffuseLerpB->GetScalar() = 1.0f;
-
-		IDatasmithMaterialExpressionGeneric* BaseColorMultiply = MaterialElement->AddMaterialExpression<IDatasmithMaterialExpressionGeneric>();
-		BaseColorMultiply->SetExpressionName(TEXT("Multiply"));
-
-		IDatasmithMaterialExpressionGeneric* BaseColorAdd = MaterialElement->AddMaterialExpression<IDatasmithMaterialExpressionGeneric>();
-		BaseColorAdd->SetExpressionName(TEXT("Add"));
-
-		IDatasmithMaterialExpressionGeneric* BaseColorTransparencyMultiply = MaterialElement->AddMaterialExpression<IDatasmithMaterialExpressionGeneric>();
-		BaseColorTransparencyMultiply->SetExpressionName(TEXT("Multiply"));
-
-		IDatasmithMaterialExpressionGeneric* IncandescenceMultiply = MaterialElement->AddMaterialExpression<IDatasmithMaterialExpressionGeneric>();
-		IncandescenceMultiply->SetExpressionName(TEXT("Multiply"));
-
-		IDatasmithMaterialExpressionGeneric* IncandescenceScaleMultiply = MaterialElement->AddMaterialExpression<IDatasmithMaterialExpressionGeneric>();
-		IncandescenceScaleMultiply->SetExpressionName(TEXT("Multiply"));
-
-		IDatasmithMaterialExpressionScalar* IncandescenceScale = MaterialElement->AddMaterialExpression<IDatasmithMaterialExpressionScalar>();
-		IncandescenceScale->GetScalar() = 100.0f;
-
-		IDatasmithMaterialExpressionGeneric* TransparencyOneMinus = MaterialElement->AddMaterialExpression<IDatasmithMaterialExpressionGeneric>();
-		TransparencyOneMinus->SetExpressionName(TEXT("OneMinus"));
-
-		IDatasmithMaterialExpressionFunctionCall* BreakFloat3 = nullptr;
-		IDatasmithMaterialExpressionGeneric* AddRG = nullptr;
-		IDatasmithMaterialExpressionGeneric* AddRGB = nullptr;
-		IDatasmithMaterialExpressionGeneric* Divide = nullptr;
-		IDatasmithMaterialExpressionScalar* DivideConstant = nullptr;
-		if (bIsTransparent)
-		{
-			BreakFloat3 = MaterialElement->AddMaterialExpression<IDatasmithMaterialExpressionFunctionCall>();
-			BreakFloat3->SetFunctionPathName(TEXT("/Engine/Functions/Engine_MaterialFunctions02/Utility/BreakFloat3Components.BreakFloat3Components"));
-
-			AddRG = MaterialElement->AddMaterialExpression<IDatasmithMaterialExpressionGeneric>();
-			AddRG->SetExpressionName(TEXT("Add"));
-
-			AddRGB = MaterialElement->AddMaterialExpression<IDatasmithMaterialExpressionGeneric>();
-			AddRGB->SetExpressionName(TEXT("Add"));
-
-			Divide = MaterialElement->AddMaterialExpression<IDatasmithMaterialExpressionGeneric>();
-			Divide->SetExpressionName(TEXT("Divide"));
-
-			DivideConstant = MaterialElement->AddMaterialExpression<IDatasmithMaterialExpressionScalar>();
-			DivideConstant->GetScalar() = 3.0f;
-		}
-
-		// Connect expressions
-		DiffuseLerpA->ConnectExpression(*DiffuseLerp->GetInput(0));
-		DiffuseLerpB->ConnectExpression(*DiffuseLerp->GetInput(1));
-		DiffuseExpression->ConnectExpression(*DiffuseLerp->GetInput(2));
-
-		ColorExpression->ConnectExpression(*BaseColorMultiply->GetInput(0));
-		DiffuseLerp->ConnectExpression(*BaseColorMultiply->GetInput(1));
-
-		BaseColorMultiply->ConnectExpression(*BaseColorAdd->GetInput(0));
-		IncandescenceColorExpression->ConnectExpression(*BaseColorAdd->GetInput(1));
-
-		BaseColorAdd->ConnectExpression(*BaseColorTransparencyMultiply->GetInput(0));
-		TransparencyOneMinus->ConnectExpression(*BaseColorTransparencyMultiply->GetInput(1));
-
-		GlowIntensityExpression->ConnectExpression(*IncandescenceScaleMultiply->GetInput(0));
-		IncandescenceScale->ConnectExpression(*IncandescenceScaleMultiply->GetInput(1));
-
-		BaseColorTransparencyMultiply->ConnectExpression(*IncandescenceMultiply->GetInput(0));
-		IncandescenceScaleMultiply->ConnectExpression(*IncandescenceMultiply->GetInput(1));
-
-		TransparencyColorExpression->ConnectExpression(*TransparencyOneMinus->GetInput(0));
-
-		if (bIsTransparent)
-		{
-			TransparencyOneMinus->ConnectExpression(*BreakFloat3->GetInput(0));
-
-			BreakFloat3->ConnectExpression(*AddRG->GetInput(0), 0);
-			BreakFloat3->ConnectExpression(*AddRG->GetInput(1), 1);
-
-			AddRG->ConnectExpression(*AddRGB->GetInput(0));
-			BreakFloat3->ConnectExpression(*AddRGB->GetInput(1), 2);
-
-			AddRGB->ConnectExpression(*Divide->GetInput(0));
-			DivideConstant->ConnectExpression(*Divide->GetInput(1));
-		}
-
-		// Connect material outputs
-		MaterialElement->GetBaseColor().SetExpression(BaseColorTransparencyMultiply);
-		MaterialElement->GetEmissiveColor().SetExpression(IncandescenceMultiply);
-		if (bIsTransparent)
-		{
-			MaterialElement->GetOpacity().SetExpression(Divide);
-			MaterialElement->SetParentLabel(TEXT("M_DatasmithAliasLambertTransparent"));
-		}
-		else {
-			MaterialElement->SetParentLabel(TEXT("M_DatasmithAliasLambert"));
-		}
-
-	}
-
-	void FWireTranslatorImpl::AddAlLightSourceParameters(const TAlObjectPtr<AlShader>& Shader, TSharedPtr<IDatasmithUEPbrMaterialElement> MaterialElement)
-	{
-		// Default values for a LightSource material
-		FColor Color(145, 148, 153);
-		FColor TransparencyColor(0, 0, 0);
-		FColor IncandescenceColor(0, 0, 0);
-		double GlowIntensity = 0.0;
-
-		AlList* List = Shader->fields();
-		for (AlShadingFieldItem* Item = static_cast<AlShadingFieldItem*>(List->first()); Item; Item = Item->nextField())
-		{
-			double Value = 0.0f;
-			statusCode ErrorCode = Shader->parameter(Item->field(), Value);
-			if (ErrorCode != 0)
-			{
-				continue;
-			}
-
-			GetCommonParameters(Item->field(), Value, Color, TransparencyColor, IncandescenceColor, GlowIntensity);
-		}
-
-		bool bIsTransparent = IsTransparent(TransparencyColor);
-
-		// Construct parameter expressions
-		IDatasmithMaterialExpressionColor* ColorExpression = MaterialElement->AddMaterialExpression<IDatasmithMaterialExpressionColor>();
-		ColorExpression->SetName(TEXT("Color"));
-		ColorExpression->GetColor() = FLinearColor::FromSRGBColor(Color);
-
-		IDatasmithMaterialExpressionColor* IncandescenceColorExpression = MaterialElement->AddMaterialExpression<IDatasmithMaterialExpressionColor>();
-		IncandescenceColorExpression->SetName(TEXT("IncandescenceColor"));
-		IncandescenceColorExpression->GetColor() = FLinearColor::FromSRGBColor(IncandescenceColor);
-
-		IDatasmithMaterialExpressionColor* TransparencyColorExpression = MaterialElement->AddMaterialExpression<IDatasmithMaterialExpressionColor>();
-		TransparencyColorExpression->SetName(TEXT("TransparencyColor"));
-		TransparencyColorExpression->GetColor() = FLinearColor::FromSRGBColor(TransparencyColor);
-
-		IDatasmithMaterialExpressionScalar* GlowIntensityExpression = MaterialElement->AddMaterialExpression<IDatasmithMaterialExpressionScalar>();
-		GlowIntensityExpression->GetScalar() = GlowIntensity;
-		GlowIntensityExpression->SetName(TEXT("GlowIntensity"));
-
-		// Create aux expressions
-		IDatasmithMaterialExpressionGeneric* BaseColorAdd = MaterialElement->AddMaterialExpression<IDatasmithMaterialExpressionGeneric>();
-		BaseColorAdd->SetExpressionName(TEXT("Add"));
-
-		IDatasmithMaterialExpressionGeneric* BaseColorTransparencyMultiply = MaterialElement->AddMaterialExpression<IDatasmithMaterialExpressionGeneric>();
-		BaseColorTransparencyMultiply->SetExpressionName(TEXT("Multiply"));
-
-		IDatasmithMaterialExpressionGeneric* IncandescenceMultiply = MaterialElement->AddMaterialExpression<IDatasmithMaterialExpressionGeneric>();
-		IncandescenceMultiply->SetExpressionName(TEXT("Multiply"));
-
-		IDatasmithMaterialExpressionGeneric* IncandescenceScaleMultiply = MaterialElement->AddMaterialExpression<IDatasmithMaterialExpressionGeneric>();
-		IncandescenceScaleMultiply->SetExpressionName(TEXT("Multiply"));
-
-		IDatasmithMaterialExpressionScalar* IncandescenceScale = MaterialElement->AddMaterialExpression<IDatasmithMaterialExpressionScalar>();
-		IncandescenceScale->GetScalar() = 100.0f;
-
-		IDatasmithMaterialExpressionGeneric* TransparencyOneMinus = MaterialElement->AddMaterialExpression<IDatasmithMaterialExpressionGeneric>();
-		TransparencyOneMinus->SetExpressionName(TEXT("OneMinus"));
-
-		IDatasmithMaterialExpressionFunctionCall* BreakFloat3 = nullptr;
-		IDatasmithMaterialExpressionGeneric* AddRG = nullptr;
-		IDatasmithMaterialExpressionGeneric* AddRGB = nullptr;
-		IDatasmithMaterialExpressionGeneric* Divide = nullptr;
-		IDatasmithMaterialExpressionScalar* DivideConstant = nullptr;
-		if (bIsTransparent)
-		{
-			BreakFloat3 = MaterialElement->AddMaterialExpression<IDatasmithMaterialExpressionFunctionCall>();
-			BreakFloat3->SetFunctionPathName(TEXT("/Engine/Functions/Engine_MaterialFunctions02/Utility/BreakFloat3Components.BreakFloat3Components"));
-
-			AddRG = MaterialElement->AddMaterialExpression<IDatasmithMaterialExpressionGeneric>();
-			AddRG->SetExpressionName(TEXT("Add"));
-
-			AddRGB = MaterialElement->AddMaterialExpression<IDatasmithMaterialExpressionGeneric>();
-			AddRGB->SetExpressionName(TEXT("Add"));
-
-			Divide = MaterialElement->AddMaterialExpression<IDatasmithMaterialExpressionGeneric>();
-			Divide->SetExpressionName(TEXT("Divide"));
-
-			DivideConstant = MaterialElement->AddMaterialExpression<IDatasmithMaterialExpressionScalar>();
-			DivideConstant->GetScalar() = 3.0f;
-		}
-
-		// Connect expressions
-		ColorExpression->ConnectExpression(*BaseColorAdd->GetInput(0));
-		IncandescenceColorExpression->ConnectExpression(*BaseColorAdd->GetInput(1));
-
-		BaseColorAdd->ConnectExpression(*BaseColorTransparencyMultiply->GetInput(0));
-		TransparencyOneMinus->ConnectExpression(*BaseColorTransparencyMultiply->GetInput(1));
-
-		GlowIntensityExpression->ConnectExpression(*IncandescenceScaleMultiply->GetInput(0));
-		IncandescenceScale->ConnectExpression(*IncandescenceScaleMultiply->GetInput(1));
-
-		BaseColorTransparencyMultiply->ConnectExpression(*IncandescenceMultiply->GetInput(0));
-		IncandescenceScaleMultiply->ConnectExpression(*IncandescenceMultiply->GetInput(1));
-
-		TransparencyColorExpression->ConnectExpression(*TransparencyOneMinus->GetInput(0));
-
-		if (bIsTransparent)
-		{
-			TransparencyOneMinus->ConnectExpression(*BreakFloat3->GetInput(0));
-
-			BreakFloat3->ConnectExpression(*AddRG->GetInput(0), 0);
-			BreakFloat3->ConnectExpression(*AddRG->GetInput(1), 1);
-
-			AddRG->ConnectExpression(*AddRGB->GetInput(0));
-			BreakFloat3->ConnectExpression(*AddRGB->GetInput(1), 2);
-
-			AddRGB->ConnectExpression(*Divide->GetInput(0));
-			DivideConstant->ConnectExpression(*Divide->GetInput(1));
-		}
-
-		// Connect material outputs
-		MaterialElement->GetBaseColor().SetExpression(BaseColorTransparencyMultiply);
-		MaterialElement->GetEmissiveColor().SetExpression(IncandescenceMultiply);
-
-		if (bIsTransparent)
-		{
-			MaterialElement->GetOpacity().SetExpression(Divide);
-			MaterialElement->SetParentLabel(TEXT("M_DatasmithAliasLightSourceTransparent"));
-		}
-		else {
-			MaterialElement->SetParentLabel(TEXT("M_DatasmithAliasLightSource"));
-		}
-
-	}
-
-	void FWireTranslatorImpl::AddAlPhongParameters(const TAlObjectPtr<AlShader>& Shader, TSharedPtr<IDatasmithUEPbrMaterialElement> MaterialElement)
-	{
-		// Default values for a Phong material
-		FColor Color(145, 148, 153);
-		FColor TransparencyColor(0, 0, 0);
-		FColor IncandescenceColor(0, 0, 0);
@@ -1012,9 +1301,10 @@
 		double Diffuse = 1.0;
 		double GlowIntensity = 0.0;
 		double Gloss = 0.8;
-		double Shinyness = 20.0;
+		double Eccentricity = 0.35;
 		double Specularity = 1.0;
 		double Reflectivity = 0.5;
+		double SpecularRolloff = 0.5;
 
 		AlList* List = Shader->fields();
 		for (AlShadingFieldItem* Item = static_cast<AlShadingFieldItem*>(List->first()); Item; Item = Item->nextField())
@@ -1033,28 +1323,31 @@
 
 			switch (Item->field())
 			{
-			case AlShadingFields::kFLD_SHADING_PHONG_DIFFUSE:
+			case AlShadingFields::kFLD_SHADING_BLINN_DIFFUSE:
 				Diffuse = Value;
 				break;
-			case AlShadingFields::kFLD_SHADING_PHONG_GLOSS_:
+			case AlShadingFields::kFLD_SHADING_BLINN_GLOSS_:
 				Gloss = Value;
 				break;
-			case AlShadingFields::kFLD_SHADING_PHONG_SPECULAR_R:
-				SpecularColor.R = (uint8)(255.f * Value);;
+			case AlShadingFields::kFLD_SHADING_BLINN_SPECULAR_R:
+				SpecularColor.R = (uint8)(255.f * Value);
 				break;
-			case AlShadingFields::kFLD_SHADING_PHONG_SPECULAR_G:
+			case AlShadingFields::kFLD_SHADING_BLINN_SPECULAR_G:
 				SpecularColor.G = (uint8)(255.f * Value);;
 				break;
-			case AlShadingFields::kFLD_SHADING_PHONG_SPECULAR_B:
+			case AlShadingFields::kFLD_SHADING_BLINN_SPECULAR_B:
 				SpecularColor.B = (uint8)(255.f * Value);;
 				break;
-			case AlShadingFields::kFLD_SHADING_PHONG_SPECULARITY_:
+			case AlShadingFields::kFLD_SHADING_BLINN_SPECULARITY_:
 				Specularity = Value;
 				break;
-			case AlShadingFields::kFLD_SHADING_PHONG_SHINYNESS:
-				Shinyness = Value;
+			case AlShadingFields::kFLD_SHADING_BLINN_SPECULAR_ROLLOFF:
+				SpecularRolloff = Value;
 				break;
-			case AlShadingFields::kFLD_SHADING_PHONG_REFLECTIVITY:
+			case AlShadingFields::kFLD_SHADING_BLINN_ECCENTRICITY:
+				Eccentricity = Value;
+				break;
+			case AlShadingFields::kFLD_SHADING_BLINN_REFLECTIVITY:
 				Reflectivity = Value;
 				break;
 			}
@@ -1079,9 +1372,13 @@
 		SpecularityExpression->GetScalar() = Specularity * 0.3;
 		SpecularityExpression->SetName(TEXT("Specularity"));
 
-		IDatasmithMaterialExpressionScalar* ShinynessExpression = MaterialElement->AddMaterialExpression<IDatasmithMaterialExpressionScalar>();
-		ShinynessExpression->GetScalar() = Shinyness;
-		ShinynessExpression->SetName(TEXT("Shinyness"));
+		IDatasmithMaterialExpressionScalar* SpecularRolloffExpression = MaterialElement->AddMaterialExpression<IDatasmithMaterialExpressionScalar>();
+		SpecularRolloffExpression->GetScalar() = SpecularRolloff;
+		SpecularRolloffExpression->SetName(TEXT("SpecularRolloff"));
 
+		IDatasmithMaterialExpressionScalar* EccentricityExpression = MaterialElement->AddMaterialExpression<IDatasmithMaterialExpressionScalar>();
+		EccentricityExpression->GetScalar() = Eccentricity;
+		EccentricityExpression->SetName(TEXT("Eccentricity"));
+
 		IDatasmithMaterialExpressionScalar* ReflectivityExpression = MaterialElement->AddMaterialExpression<IDatasmithMaterialExpressionScalar>();
 		ReflectivityExpression->GetScalar() = Reflectivity;
@@ -1140,23 +1437,35 @@
 		IDatasmithMaterialExpressionScalar* IncandescenceScale = MaterialElement->AddMaterialExpression<IDatasmithMaterialExpressionScalar>();
 		IncandescenceScale->GetScalar() = 100.0f;
 
-		IDatasmithMaterialExpressionGeneric* ShinynessSubtract = MaterialElement->AddMaterialExpression<IDatasmithMaterialExpressionGeneric>();
-		ShinynessSubtract->SetExpressionName(TEXT("Subtract"));
+		IDatasmithMaterialExpressionGeneric* EccentricityMultiply = MaterialElement->AddMaterialExpression<IDatasmithMaterialExpressionGeneric>();
+		EccentricityMultiply->SetExpressionName(TEXT("Multiply"));
 
-		IDatasmithMaterialExpressionScalar* ShinynessSubtract2 = MaterialElement->AddMaterialExpression<IDatasmithMaterialExpressionScalar>();
-		ShinynessSubtract2->GetScalar() = 2.0f;
+		IDatasmithMaterialExpressionGeneric* EccentricityOneMinus = MaterialElement->AddMaterialExpression<IDatasmithMaterialExpressionGeneric>();
+		EccentricityOneMinus->SetExpressionName(TEXT("OneMinus"));
 
-		IDatasmithMaterialExpressionGeneric* ShinynessDivide = MaterialElement->AddMaterialExpression<IDatasmithMaterialExpressionGeneric>();
-		ShinynessDivide->SetExpressionName(TEXT("Divide"));
-
-		IDatasmithMaterialExpressionScalar* ShinynessDivide98 = MaterialElement->AddMaterialExpression<IDatasmithMaterialExpressionScalar>();
-		ShinynessDivide98->GetScalar() = 98.0f;
-
-		IDatasmithMaterialExpressionGeneric* SpecularityMultiply = MaterialElement->AddMaterialExpression<IDatasmithMaterialExpressionGeneric>();
-		SpecularityMultiply->SetExpressionName(TEXT("Multiply"));
-
 		IDatasmithMaterialExpressionGeneric* RoughnessOneMinus = MaterialElement->AddMaterialExpression<IDatasmithMaterialExpressionGeneric>();
 		RoughnessOneMinus->SetExpressionName(TEXT("OneMinus"));
 
+		IDatasmithMaterialExpressionScalar* FresnelExponent = MaterialElement->AddMaterialExpression<IDatasmithMaterialExpressionScalar>();
+		FresnelExponent->GetScalar() = 4.0f;
+
+		IDatasmithMaterialExpressionFunctionCall* FresnelFunc = MaterialElement->AddMaterialExpression<IDatasmithMaterialExpressionFunctionCall>();
+		FresnelFunc->SetFunctionPathName(TEXT("/Engine/Functions/Engine_MaterialFunctions02/Fresnel_Function.Fresnel_Function"));
+
+		IDatasmithMaterialExpressionGeneric* FresnelLerp = MaterialElement->AddMaterialExpression<IDatasmithMaterialExpressionGeneric>();
+		FresnelLerp->SetExpressionName(TEXT("LinearInterpolate"));
+
+		IDatasmithMaterialExpressionScalar* FresnelLerpA = MaterialElement->AddMaterialExpression<IDatasmithMaterialExpressionScalar>();
+		FresnelLerpA->GetScalar() = 1.0f;
+
+		IDatasmithMaterialExpressionScalar* SpecularPowerExp = MaterialElement->AddMaterialExpression<IDatasmithMaterialExpressionScalar>();
+		SpecularPowerExp->GetScalar() = 0.5f;
+
+		IDatasmithMaterialExpressionGeneric* Power = MaterialElement->AddMaterialExpression<IDatasmithMaterialExpressionGeneric>();
+		Power->SetExpressionName(TEXT("Power"));
+
+		IDatasmithMaterialExpressionGeneric* FresnelMultiply = MaterialElement->AddMaterialExpression<IDatasmithMaterialExpressionGeneric>();
+		FresnelMultiply->SetExpressionName(TEXT("Multiply"));
+
 		IDatasmithMaterialExpressionGeneric* TransparencyOneMinus = MaterialElement->AddMaterialExpression<IDatasmithMaterialExpressionGeneric>();
 		TransparencyOneMinus->SetExpressionName(TEXT("OneMinus"));
@@ -1212,16 +1521,24 @@
 		BaseColorTransparencyMultiply->ConnectExpression(*IncandescenceMultiply->GetInput(0));
 		IncandescenceScaleMultiply->ConnectExpression(*IncandescenceMultiply->GetInput(1));
 
-		ShinynessExpression->ConnectExpression(*ShinynessSubtract->GetInput(0));
-		ShinynessSubtract2->ConnectExpression(*ShinynessSubtract->GetInput(1));
+		EccentricityExpression->ConnectExpression(*EccentricityOneMinus->GetInput(0));
 
-		ShinynessSubtract->ConnectExpression(*ShinynessDivide->GetInput(0));
-		ShinynessDivide98->ConnectExpression(*ShinynessDivide->GetInput(1));
+		EccentricityOneMinus->ConnectExpression(*EccentricityMultiply->GetInput(0));
+		SpecularityExpression->ConnectExpression(*EccentricityMultiply->GetInput(1));
 
-		ShinynessDivide->ConnectExpression(*SpecularityMultiply->GetInput(0));
-		SpecularityExpression->ConnectExpression(*SpecularityMultiply->GetInput(1));
+		EccentricityMultiply->ConnectExpression(*RoughnessOneMinus->GetInput(0));
 
-		SpecularityMultiply->ConnectExpression(*RoughnessOneMinus->GetInput(0));
+		FresnelExponent->ConnectExpression(*FresnelFunc->GetInput(3));
 
+		SpecularRolloffExpression->ConnectExpression(*Power->GetInput(0));
+		SpecularPowerExp->ConnectExpression(*Power->GetInput(1));
+
+		FresnelLerpA->ConnectExpression(*FresnelLerp->GetInput(0));
+		FresnelFunc->ConnectExpression(*FresnelLerp->GetInput(1));
+		Power->ConnectExpression(*FresnelLerp->GetInput(2));
+
+		FresnelLerp->ConnectExpression(*FresnelMultiply->GetInput(0));
+		ReflectivityExpression->ConnectExpression(*FresnelMultiply->GetInput(1));
+
 		TransparencyColorExpression->ConnectExpression(*TransparencyOneMinus->GetInput(0));
 
@@ -1242,6 +1559,547 @@
 		// Connect material outputs
 		MaterialElement->GetBaseColor().SetExpression(BaseColorTransparencyMultiply);
 		MaterialElement->GetMetallic().SetExpression(GlossExpression);
+		MaterialElement->GetSpecular().SetExpression(FresnelMultiply);
+		MaterialElement->GetRoughness().SetExpression(RoughnessOneMinus);
+		MaterialElement->GetEmissiveColor().SetExpression(IncandescenceMultiply);
+
+		if (bIsTransparent)
+		{
+			MaterialElement->GetOpacity().SetExpression(Divide);
+			MaterialElement->SetParentLabel(TEXT("M_DatasmithAliasBlinnTransparent"));
+		}
+		else
+		{
+			MaterialElement->SetParentLabel(TEXT("M_DatasmithAliasBlinn"));
+		}
+
+	}
+
+	void FWireTranslatorImpl::AddAlLambertParameters(const TAlObjectPtr<AlShader>& Shader, TSharedPtr<IDatasmithUEPbrMaterialElement> MaterialElement)
+	{
+		// Default values for a Lambert material
+		FColor Color(145, 148, 153);
+		FColor TransparencyColor(0, 0, 0);
+		FColor IncandescenceColor(0, 0, 0);
+		double Diffuse = 1.0;
+		double GlowIntensity = 0.0;
+
+		AlList* List = Shader->fields();
+		for (AlShadingFieldItem* Item = static_cast<AlShadingFieldItem*>(List->first()); Item; Item = Item->nextField())
+		{
+			double Value = 0.0f;
+			statusCode ErrorCode = Shader->parameter(Item->field(), Value);
+			if (ErrorCode != 0)
+			{
+				continue;
+			}
+
+			if (GetCommonParameters(Item->field(), Value, Color, TransparencyColor, IncandescenceColor, GlowIntensity))
+			{
+				continue;
+			}
+
+			switch (Item->field())
+			{
+			case AlShadingFields::kFLD_SHADING_LAMBERT_DIFFUSE:
+				Diffuse = Value;
+				break;
+			}
+		}
+
+		bool bIsTransparent = IsTransparent(TransparencyColor);
+
+		// Construct parameter expressions
+		IDatasmithMaterialExpressionScalar* DiffuseExpression = MaterialElement->AddMaterialExpression<IDatasmithMaterialExpressionScalar>();
+		DiffuseExpression->GetScalar() = Diffuse;
+		DiffuseExpression->SetName(TEXT("Diffuse"));
+
+		IDatasmithMaterialExpressionColor* ColorExpression = MaterialElement->AddMaterialExpression<IDatasmithMaterialExpressionColor>();
+		ColorExpression->SetName(TEXT("Color"));
+		ColorExpression->GetColor() = FLinearColor::FromSRGBColor(Color);
+
+		IDatasmithMaterialExpressionColor* IncandescenceColorExpression = MaterialElement->AddMaterialExpression<IDatasmithMaterialExpressionColor>();
+		IncandescenceColorExpression->SetName(TEXT("IncandescenceColor"));
+		IncandescenceColorExpression->GetColor() = FLinearColor::FromSRGBColor(IncandescenceColor);
+
+		IDatasmithMaterialExpressionColor* TransparencyColorExpression = MaterialElement->AddMaterialExpression<IDatasmithMaterialExpressionColor>();
+		TransparencyColorExpression->SetName(TEXT("TransparencyColor"));
+		TransparencyColorExpression->GetColor() = FLinearColor::FromSRGBColor(TransparencyColor);
+
+		IDatasmithMaterialExpressionScalar* GlowIntensityExpression = MaterialElement->AddMaterialExpression<IDatasmithMaterialExpressionScalar>();
+		GlowIntensityExpression->GetScalar() = GlowIntensity;
+		GlowIntensityExpression->SetName(TEXT("GlowIntensity"));
+
+		// Create aux expressions
+		IDatasmithMaterialExpressionGeneric* DiffuseLerp = MaterialElement->AddMaterialExpression<IDatasmithMaterialExpressionGeneric>();
+		DiffuseLerp->SetExpressionName(TEXT("LinearInterpolate"));
+
+		IDatasmithMaterialExpressionScalar* DiffuseLerpA = MaterialElement->AddMaterialExpression<IDatasmithMaterialExpressionScalar>();
+		DiffuseLerpA->GetScalar() = 0.04f;
+
+		IDatasmithMaterialExpressionScalar* DiffuseLerpB = MaterialElement->AddMaterialExpression<IDatasmithMaterialExpressionScalar>();
+		DiffuseLerpB->GetScalar() = 1.0f;
+
+		IDatasmithMaterialExpressionGeneric* BaseColorMultiply = MaterialElement->AddMaterialExpression<IDatasmithMaterialExpressionGeneric>();
+		BaseColorMultiply->SetExpressionName(TEXT("Multiply"));
+
+		IDatasmithMaterialExpressionGeneric* BaseColorAdd = MaterialElement->AddMaterialExpression<IDatasmithMaterialExpressionGeneric>();
+		BaseColorAdd->SetExpressionName(TEXT("Add"));
+
+		IDatasmithMaterialExpressionGeneric* BaseColorTransparencyMultiply = MaterialElement->AddMaterialExpression<IDatasmithMaterialExpressionGeneric>();
+		BaseColorTransparencyMultiply->SetExpressionName(TEXT("Multiply"));
+
+		IDatasmithMaterialExpressionGeneric* IncandescenceMultiply = MaterialElement->AddMaterialExpression<IDatasmithMaterialExpressionGeneric>();
+		IncandescenceMultiply->SetExpressionName(TEXT("Multiply"));
+
+		IDatasmithMaterialExpressionGeneric* IncandescenceScaleMultiply = MaterialElement->AddMaterialExpression<IDatasmithMaterialExpressionGeneric>();
+		IncandescenceScaleMultiply->SetExpressionName(TEXT("Multiply"));
+
+		IDatasmithMaterialExpressionScalar* IncandescenceScale = MaterialElement->AddMaterialExpression<IDatasmithMaterialExpressionScalar>();
+		IncandescenceScale->GetScalar() = 100.0f;
+
+		IDatasmithMaterialExpressionGeneric* TransparencyOneMinus = MaterialElement->AddMaterialExpression<IDatasmithMaterialExpressionGeneric>();
+		TransparencyOneMinus->SetExpressionName(TEXT("OneMinus"));
+
+		IDatasmithMaterialExpressionFunctionCall* BreakFloat3 = nullptr;
+		IDatasmithMaterialExpressionGeneric* AddRG = nullptr;
+		IDatasmithMaterialExpressionGeneric* AddRGB = nullptr;
+		IDatasmithMaterialExpressionGeneric* Divide = nullptr;
+		IDatasmithMaterialExpressionScalar* DivideConstant = nullptr;
+		if (bIsTransparent)
+		{
+			BreakFloat3 = MaterialElement->AddMaterialExpression<IDatasmithMaterialExpressionFunctionCall>();
+			BreakFloat3->SetFunctionPathName(TEXT("/Engine/Functions/Engine_MaterialFunctions02/Utility/BreakFloat3Components.BreakFloat3Components"));
+
+			AddRG = MaterialElement->AddMaterialExpression<IDatasmithMaterialExpressionGeneric>();
+			AddRG->SetExpressionName(TEXT("Add"));
+
+			AddRGB = MaterialElement->AddMaterialExpression<IDatasmithMaterialExpressionGeneric>();
+			AddRGB->SetExpressionName(TEXT("Add"));
+
+			Divide = MaterialElement->AddMaterialExpression<IDatasmithMaterialExpressionGeneric>();
+			Divide->SetExpressionName(TEXT("Divide"));
+
+			DivideConstant = MaterialElement->AddMaterialExpression<IDatasmithMaterialExpressionScalar>();
+			DivideConstant->GetScalar() = 3.0f;
+		}
+
+		// Connect expressions
+		DiffuseLerpA->ConnectExpression(*DiffuseLerp->GetInput(0));
+		DiffuseLerpB->ConnectExpression(*DiffuseLerp->GetInput(1));
+		DiffuseExpression->ConnectExpression(*DiffuseLerp->GetInput(2));
+
+		ColorExpression->ConnectExpression(*BaseColorMultiply->GetInput(0));
+		DiffuseLerp->ConnectExpression(*BaseColorMultiply->GetInput(1));
+
+		BaseColorMultiply->ConnectExpression(*BaseColorAdd->GetInput(0));
+		IncandescenceColorExpression->ConnectExpression(*BaseColorAdd->GetInput(1));
+
+		BaseColorAdd->ConnectExpression(*BaseColorTransparencyMultiply->GetInput(0));
+		TransparencyOneMinus->ConnectExpression(*BaseColorTransparencyMultiply->GetInput(1));
+
+		GlowIntensityExpression->ConnectExpression(*IncandescenceScaleMultiply->GetInput(0));
+		IncandescenceScale->ConnectExpression(*IncandescenceScaleMultiply->GetInput(1));
+
+		BaseColorTransparencyMultiply->ConnectExpression(*IncandescenceMultiply->GetInput(0));
+		IncandescenceScaleMultiply->ConnectExpression(*IncandescenceMultiply->GetInput(1));
+
+		TransparencyColorExpression->ConnectExpression(*TransparencyOneMinus->GetInput(0));
+
+		if (bIsTransparent)
+		{
+			TransparencyOneMinus->ConnectExpression(*BreakFloat3->GetInput(0));
+
+			BreakFloat3->ConnectExpression(*AddRG->GetInput(0), 0);
+			BreakFloat3->ConnectExpression(*AddRG->GetInput(1), 1);
+
+			AddRG->ConnectExpression(*AddRGB->GetInput(0));
+			BreakFloat3->ConnectExpression(*AddRGB->GetInput(1), 2);
+
+			AddRGB->ConnectExpression(*Divide->GetInput(0));
+			DivideConstant->ConnectExpression(*Divide->GetInput(1));
+		}
+
+		// Connect material outputs
+		MaterialElement->GetBaseColor().SetExpression(BaseColorTransparencyMultiply);
+		MaterialElement->GetEmissiveColor().SetExpression(IncandescenceMultiply);
+		if (bIsTransparent)
+		{
+			MaterialElement->GetOpacity().SetExpression(Divide);
+			MaterialElement->SetParentLabel(TEXT("M_DatasmithAliasLambertTransparent"));
+		}
+		else {
+			MaterialElement->SetParentLabel(TEXT("M_DatasmithAliasLambert"));
+		}
+
+	}
+
+	void FWireTranslatorImpl::AddAlLightSourceParameters(const TAlObjectPtr<AlShader>& Shader, TSharedPtr<IDatasmithUEPbrMaterialElement> MaterialElement)
+	{
+		// Default values for a LightSource material
+		FColor Color(145, 148, 153);
+		FColor TransparencyColor(0, 0, 0);
+		FColor IncandescenceColor(0, 0, 0);
+		double GlowIntensity = 0.0;
+
+		AlList* List = Shader->fields();
+		for (AlShadingFieldItem* Item = static_cast<AlShadingFieldItem*>(List->first()); Item; Item = Item->nextField())
+		{
+			double Value = 0.0f;
+			statusCode ErrorCode = Shader->parameter(Item->field(), Value);
+			if (ErrorCode != 0)
+			{
+				continue;
+			}
+
+			GetCommonParameters(Item->field(), Value, Color, TransparencyColor, IncandescenceColor, GlowIntensity);
+		}
+
+		bool bIsTransparent = IsTransparent(TransparencyColor);
+
+		// Construct parameter expressions
+		IDatasmithMaterialExpressionColor* ColorExpression = MaterialElement->AddMaterialExpression<IDatasmithMaterialExpressionColor>();
+		ColorExpression->SetName(TEXT("Color"));
+		ColorExpression->GetColor() = FLinearColor::FromSRGBColor(Color);
+
+		IDatasmithMaterialExpressionColor* IncandescenceColorExpression = MaterialElement->AddMaterialExpression<IDatasmithMaterialExpressionColor>();
+		IncandescenceColorExpression->SetName(TEXT("IncandescenceColor"));
+		IncandescenceColorExpression->GetColor() = FLinearColor::FromSRGBColor(IncandescenceColor);
+
+		IDatasmithMaterialExpressionColor* TransparencyColorExpression = MaterialElement->AddMaterialExpression<IDatasmithMaterialExpressionColor>();
+		TransparencyColorExpression->SetName(TEXT("TransparencyColor"));
+		TransparencyColorExpression->GetColor() = FLinearColor::FromSRGBColor(TransparencyColor);
+
+		IDatasmithMaterialExpressionScalar* GlowIntensityExpression = MaterialElement->AddMaterialExpression<IDatasmithMaterialExpressionScalar>();
+		GlowIntensityExpression->GetScalar() = GlowIntensity;
+		GlowIntensityExpression->SetName(TEXT("GlowIntensity"));
+
+		// Create aux expressions
+		IDatasmithMaterialExpressionGeneric* BaseColorAdd = MaterialElement->AddMaterialExpression<IDatasmithMaterialExpressionGeneric>();
+		BaseColorAdd->SetExpressionName(TEXT("Add"));
+
+		IDatasmithMaterialExpressionGeneric* BaseColorTransparencyMultiply = MaterialElement->AddMaterialExpression<IDatasmithMaterialExpressionGeneric>();
+		BaseColorTransparencyMultiply->SetExpressionName(TEXT("Multiply"));
+
+		IDatasmithMaterialExpressionGeneric* IncandescenceMultiply = MaterialElement->AddMaterialExpression<IDatasmithMaterialExpressionGeneric>();
+		IncandescenceMultiply->SetExpressionName(TEXT("Multiply"));
+
+		IDatasmithMaterialExpressionGeneric* IncandescenceScaleMultiply = MaterialElement->AddMaterialExpression<IDatasmithMaterialExpressionGeneric>();
+		IncandescenceScaleMultiply->SetExpressionName(TEXT("Multiply"));
+
+		IDatasmithMaterialExpressionScalar* IncandescenceScale = MaterialElement->AddMaterialExpression<IDatasmithMaterialExpressionScalar>();
+		IncandescenceScale->GetScalar() = 100.0f;
+
+		IDatasmithMaterialExpressionGeneric* TransparencyOneMinus = MaterialElement->AddMaterialExpression<IDatasmithMaterialExpressionGeneric>();
+		TransparencyOneMinus->SetExpressionName(TEXT("OneMinus"));
+
+		IDatasmithMaterialExpressionFunctionCall* BreakFloat3 = nullptr;
+		IDatasmithMaterialExpressionGeneric* AddRG = nullptr;
+		IDatasmithMaterialExpressionGeneric* AddRGB = nullptr;
+		IDatasmithMaterialExpressionGeneric* Divide = nullptr;
+		IDatasmithMaterialExpressionScalar* DivideConstant = nullptr;
+		if (bIsTransparent)
+		{
+			BreakFloat3 = MaterialElement->AddMaterialExpression<IDatasmithMaterialExpressionFunctionCall>();
+			BreakFloat3->SetFunctionPathName(TEXT("/Engine/Functions/Engine_MaterialFunctions02/Utility/BreakFloat3Components.BreakFloat3Components"));
+
+			AddRG = MaterialElement->AddMaterialExpression<IDatasmithMaterialExpressionGeneric>();
+			AddRG->SetExpressionName(TEXT("Add"));
+
+			AddRGB = MaterialElement->AddMaterialExpression<IDatasmithMaterialExpressionGeneric>();
+			AddRGB->SetExpressionName(TEXT("Add"));
+
+			Divide = MaterialElement->AddMaterialExpression<IDatasmithMaterialExpressionGeneric>();
+			Divide->SetExpressionName(TEXT("Divide"));
+
+			DivideConstant = MaterialElement->AddMaterialExpression<IDatasmithMaterialExpressionScalar>();
+			DivideConstant->GetScalar() = 3.0f;
+		}
+
+		// Connect expressions
+		ColorExpression->ConnectExpression(*BaseColorAdd->GetInput(0));
+		IncandescenceColorExpression->ConnectExpression(*BaseColorAdd->GetInput(1));
+
+		BaseColorAdd->ConnectExpression(*BaseColorTransparencyMultiply->GetInput(0));
+		TransparencyOneMinus->ConnectExpression(*BaseColorTransparencyMultiply->GetInput(1));
+
+		GlowIntensityExpression->ConnectExpression(*IncandescenceScaleMultiply->GetInput(0));
+		IncandescenceScale->ConnectExpression(*IncandescenceScaleMultiply->GetInput(1));
+
+		BaseColorTransparencyMultiply->ConnectExpression(*IncandescenceMultiply->GetInput(0));
+		IncandescenceScaleMultiply->ConnectExpression(*IncandescenceMultiply->GetInput(1));
+
+		TransparencyColorExpression->ConnectExpression(*TransparencyOneMinus->GetInput(0));
+
+		if (bIsTransparent)
+		{
+			TransparencyOneMinus->ConnectExpression(*BreakFloat3->GetInput(0));
+
+			BreakFloat3->ConnectExpression(*AddRG->GetInput(0), 0);
+			BreakFloat3->ConnectExpression(*AddRG->GetInput(1), 1);
+
+			AddRG->ConnectExpression(*AddRGB->GetInput(0));
+			BreakFloat3->ConnectExpression(*AddRGB->GetInput(1), 2);
+
+			AddRGB->ConnectExpression(*Divide->GetInput(0));
+			DivideConstant->ConnectExpression(*Divide->GetInput(1));
+		}
+
+		// Connect material outputs
+		MaterialElement->GetBaseColor().SetExpression(BaseColorTransparencyMultiply);
+		MaterialElement->GetEmissiveColor().SetExpression(IncandescenceMultiply);
+
+		if (bIsTransparent)
+		{
+			MaterialElement->GetOpacity().SetExpression(Divide);
+			MaterialElement->SetParentLabel(TEXT("M_DatasmithAliasLightSourceTransparent"));
+		}
+		else {
+			MaterialElement->SetParentLabel(TEXT("M_DatasmithAliasLightSource"));
+		}
+
+	}
+
+	void FWireTranslatorImpl::AddAlPhongParameters(const TAlObjectPtr<AlShader>& Shader, TSharedPtr<IDatasmithUEPbrMaterialElement> MaterialElement)
+	{
+		// Default values for a Phong material
+		FColor Color(145, 148, 153);
+		FColor TransparencyColor(0, 0, 0);
+		FColor IncandescenceColor(0, 0, 0);
+		FColor SpecularColor(38, 38, 38);
+		double Diffuse = 1.0;
+		double GlowIntensity = 0.0;
+		double Gloss = 0.8;
+		double Shinyness = 20.0;
+		double Specularity = 1.0;
+		double Reflectivity = 0.5;
+
+		AlList* List = Shader->fields();
+		for (AlShadingFieldItem* Item = static_cast<AlShadingFieldItem*>(List->first()); Item; Item = Item->nextField())
+		{
+			double Value = 0.0f;
+			statusCode ErrorCode = Shader->parameter(Item->field(), Value);
+			if (ErrorCode != 0)
+			{
+				continue;
+			}
+
+			if (GetCommonParameters(Item->field(), Value, Color, TransparencyColor, IncandescenceColor, GlowIntensity))
+			{
+				continue;
+			}
+
+			switch (Item->field())
+			{
+			case AlShadingFields::kFLD_SHADING_PHONG_DIFFUSE:
+				Diffuse = Value;
+				break;
+			case AlShadingFields::kFLD_SHADING_PHONG_GLOSS_:
+				Gloss = Value;
+				break;
+			case AlShadingFields::kFLD_SHADING_PHONG_SPECULAR_R:
+				SpecularColor.R = (uint8)(255.f * Value);;
+				break;
+			case AlShadingFields::kFLD_SHADING_PHONG_SPECULAR_G:
+				SpecularColor.G = (uint8)(255.f * Value);;
+				break;
+			case AlShadingFields::kFLD_SHADING_PHONG_SPECULAR_B:
+				SpecularColor.B = (uint8)(255.f * Value);;
+				break;
+			case AlShadingFields::kFLD_SHADING_PHONG_SPECULARITY_:
+				Specularity = Value;
+				break;
+			case AlShadingFields::kFLD_SHADING_PHONG_SHINYNESS:
+				Shinyness = Value;
+				break;
+			case AlShadingFields::kFLD_SHADING_PHONG_REFLECTIVITY:
+				Reflectivity = Value;
+				break;
+			}
+		}
+
+		bool bIsTransparent = IsTransparent(TransparencyColor);
+
+		// Construct parameter expressions
+		IDatasmithMaterialExpressionScalar* DiffuseExpression = MaterialElement->AddMaterialExpression<IDatasmithMaterialExpressionScalar>();
+		DiffuseExpression->GetScalar() = Diffuse;
+		DiffuseExpression->SetName(TEXT("Diffuse"));
+
+		IDatasmithMaterialExpressionScalar* GlossExpression = MaterialElement->AddMaterialExpression<IDatasmithMaterialExpressionScalar>();
+		GlossExpression->GetScalar() = Gloss;
+		GlossExpression->SetName(TEXT("Gloss"));
+
+		IDatasmithMaterialExpressionColor* SpecularColorExpression = MaterialElement->AddMaterialExpression<IDatasmithMaterialExpressionColor>();
+		SpecularColorExpression->SetName(TEXT("SpecularColor"));
+		SpecularColorExpression->GetColor() = FLinearColor::FromSRGBColor(SpecularColor);
+
+		IDatasmithMaterialExpressionScalar* SpecularityExpression = MaterialElement->AddMaterialExpression<IDatasmithMaterialExpressionScalar>();
+		SpecularityExpression->GetScalar() = Specularity * 0.3;
+		SpecularityExpression->SetName(TEXT("Specularity"));
+
+		IDatasmithMaterialExpressionScalar* ShinynessExpression = MaterialElement->AddMaterialExpression<IDatasmithMaterialExpressionScalar>();
+		ShinynessExpression->GetScalar() = Shinyness;
+		ShinynessExpression->SetName(TEXT("Shinyness"));
+
+		IDatasmithMaterialExpressionScalar* ReflectivityExpression = MaterialElement->AddMaterialExpression<IDatasmithMaterialExpressionScalar>();
+		ReflectivityExpression->GetScalar() = Reflectivity;
+		ReflectivityExpression->SetName(TEXT("Reflectivity"));
+
+		IDatasmithMaterialExpressionColor* ColorExpression = MaterialElement->AddMaterialExpression<IDatasmithMaterialExpressionColor>();
+		ColorExpression->SetName(TEXT("Color"));
+		ColorExpression->GetColor() = FLinearColor::FromSRGBColor(Color);
+
+		IDatasmithMaterialExpressionColor* IncandescenceColorExpression = MaterialElement->AddMaterialExpression<IDatasmithMaterialExpressionColor>();
+		IncandescenceColorExpression->SetName(TEXT("IncandescenceColor"));
+		IncandescenceColorExpression->GetColor() = FLinearColor::FromSRGBColor(IncandescenceColor);
+
+		IDatasmithMaterialExpressionColor* TransparencyColorExpression = MaterialElement->AddMaterialExpression<IDatasmithMaterialExpressionColor>();
+		TransparencyColorExpression->SetName(TEXT("TransparencyColor"));
+		TransparencyColorExpression->GetColor() = FLinearColor::FromSRGBColor(TransparencyColor);
+
+		IDatasmithMaterialExpressionScalar* GlowIntensityExpression = MaterialElement->AddMaterialExpression<IDatasmithMaterialExpressionScalar>();
+		GlowIntensityExpression->GetScalar() = GlowIntensity;
+		GlowIntensityExpression->SetName(TEXT("GlowIntensity"));
+
+		// Create aux expressions
+		IDatasmithMaterialExpressionGeneric* ColorSpecLerp = MaterialElement->AddMaterialExpression<IDatasmithMaterialExpressionGeneric>();
+		ColorSpecLerp->SetExpressionName(TEXT("LinearInterpolate"));
+
+		IDatasmithMaterialExpressionScalar* ColorSpecLerpValue = MaterialElement->AddMaterialExpression<IDatasmithMaterialExpressionScalar>();
+		ColorSpecLerpValue->GetScalar() = 0.96f;
+
+		IDatasmithMaterialExpressionGeneric* ColorMetallicLerp = MaterialElement->AddMaterialExpression<IDatasmithMaterialExpressionGeneric>();
+		ColorMetallicLerp->SetExpressionName(TEXT("LinearInterpolate"));
+
+		IDatasmithMaterialExpressionGeneric* DiffuseLerp = MaterialElement->AddMaterialExpression<IDatasmithMaterialExpressionGeneric>();
+		DiffuseLerp->SetExpressionName(TEXT("LinearInterpolate"));
+
+		IDatasmithMaterialExpressionScalar* DiffuseLerpA = MaterialElement->AddMaterialExpression<IDatasmithMaterialExpressionScalar>();
+		DiffuseLerpA->GetScalar() = 0.04f;
+
+		IDatasmithMaterialExpressionScalar* DiffuseLerpB = MaterialElement->AddMaterialExpression<IDatasmithMaterialExpressionScalar>();
+		DiffuseLerpB->GetScalar() = 1.0f;
+
+		IDatasmithMaterialExpressionGeneric* BaseColorMultiply = MaterialElement->AddMaterialExpression<IDatasmithMaterialExpressionGeneric>();
+		BaseColorMultiply->SetExpressionName(TEXT("Multiply"));
+
+		IDatasmithMaterialExpressionGeneric* BaseColorAdd = MaterialElement->AddMaterialExpression<IDatasmithMaterialExpressionGeneric>();
+		BaseColorAdd->SetExpressionName(TEXT("Add"));
+
+		IDatasmithMaterialExpressionGeneric* BaseColorTransparencyMultiply = MaterialElement->AddMaterialExpression<IDatasmithMaterialExpressionGeneric>();
+		BaseColorTransparencyMultiply->SetExpressionName(TEXT("Multiply"));
+
+		IDatasmithMaterialExpressionGeneric* IncandescenceMultiply = MaterialElement->AddMaterialExpression<IDatasmithMaterialExpressionGeneric>();
+		IncandescenceMultiply->SetExpressionName(TEXT("Multiply"));
+
+		IDatasmithMaterialExpressionGeneric* IncandescenceScaleMultiply = MaterialElement->AddMaterialExpression<IDatasmithMaterialExpressionGeneric>();
+		IncandescenceScaleMultiply->SetExpressionName(TEXT("Multiply"));
+
+		IDatasmithMaterialExpressionScalar* IncandescenceScale = MaterialElement->AddMaterialExpression<IDatasmithMaterialExpressionScalar>();
+		IncandescenceScale->GetScalar() = 100.0f;
+
+		IDatasmithMaterialExpressionGeneric* ShinynessSubtract = MaterialElement->AddMaterialExpression<IDatasmithMaterialExpressionGeneric>();
+		ShinynessSubtract->SetExpressionName(TEXT("Subtract"));
+
+		IDatasmithMaterialExpressionScalar* ShinynessSubtract2 = MaterialElement->AddMaterialExpression<IDatasmithMaterialExpressionScalar>();
+		ShinynessSubtract2->GetScalar() = 2.0f;
+
+		IDatasmithMaterialExpressionGeneric* ShinynessDivide = MaterialElement->AddMaterialExpression<IDatasmithMaterialExpressionGeneric>();
+		ShinynessDivide->SetExpressionName(TEXT("Divide"));
+
+		IDatasmithMaterialExpressionScalar* ShinynessDivide98 = MaterialElement->AddMaterialExpression<IDatasmithMaterialExpressionScalar>();
+		ShinynessDivide98->GetScalar() = 98.0f;
+
+		IDatasmithMaterialExpressionGeneric* SpecularityMultiply = MaterialElement->AddMaterialExpression<IDatasmithMaterialExpressionGeneric>();
+		SpecularityMultiply->SetExpressionName(TEXT("Multiply"));
+
+		IDatasmithMaterialExpressionGeneric* RoughnessOneMinus = MaterialElement->AddMaterialExpression<IDatasmithMaterialExpressionGeneric>();
+		RoughnessOneMinus->SetExpressionName(TEXT("OneMinus"));
+
+		IDatasmithMaterialExpressionGeneric* TransparencyOneMinus = MaterialElement->AddMaterialExpression<IDatasmithMaterialExpressionGeneric>();
+		TransparencyOneMinus->SetExpressionName(TEXT("OneMinus"));
+
+		IDatasmithMaterialExpressionFunctionCall* BreakFloat3 = nullptr;
+		IDatasmithMaterialExpressionGeneric* AddRG = nullptr;
+		IDatasmithMaterialExpressionGeneric* AddRGB = nullptr;
+		IDatasmithMaterialExpressionGeneric* Divide = nullptr;
+		IDatasmithMaterialExpressionScalar* DivideConstant = nullptr;
+		if (bIsTransparent)
+		{
+			BreakFloat3 = MaterialElement->AddMaterialExpression<IDatasmithMaterialExpressionFunctionCall>();
+			BreakFloat3->SetFunctionPathName(TEXT("/Engine/Functions/Engine_MaterialFunctions02/Utility/BreakFloat3Components.BreakFloat3Components"));
+
+			AddRG = MaterialElement->AddMaterialExpression<IDatasmithMaterialExpressionGeneric>();
+			AddRG->SetExpressionName(TEXT("Add"));
+
+			AddRGB = MaterialElement->AddMaterialExpression<IDatasmithMaterialExpressionGeneric>();
+			AddRGB->SetExpressionName(TEXT("Add"));
+
+			Divide = MaterialElement->AddMaterialExpression<IDatasmithMaterialExpressionGeneric>();
+			Divide->SetExpressionName(TEXT("Divide"));
+
+			DivideConstant = MaterialElement->AddMaterialExpression<IDatasmithMaterialExpressionScalar>();
+			DivideConstant->GetScalar() = 3.0f;
+		}
+
+		// Connect expressions
+		SpecularColorExpression->ConnectExpression(*ColorSpecLerp->GetInput(0));
+		ColorExpression->ConnectExpression(*ColorSpecLerp->GetInput(1));
+		ColorSpecLerpValue->ConnectExpression(*ColorSpecLerp->GetInput(2));
+
+		ColorExpression->ConnectExpression(*ColorMetallicLerp->GetInput(0));
+		ColorSpecLerp->ConnectExpression(*ColorMetallicLerp->GetInput(1));
+		GlossExpression->ConnectExpression(*ColorMetallicLerp->GetInput(2));
+
+		DiffuseLerpA->ConnectExpression(*DiffuseLerp->GetInput(0));
+		DiffuseLerpB->ConnectExpression(*DiffuseLerp->GetInput(1));
+		DiffuseExpression->ConnectExpression(*DiffuseLerp->GetInput(2));
+
+		ColorMetallicLerp->ConnectExpression(*BaseColorMultiply->GetInput(0));
+		DiffuseLerp->ConnectExpression(*BaseColorMultiply->GetInput(1));
+
+		BaseColorMultiply->ConnectExpression(*BaseColorAdd->GetInput(0));
+		IncandescenceColorExpression->ConnectExpression(*BaseColorAdd->GetInput(1));
+
+		BaseColorAdd->ConnectExpression(*BaseColorTransparencyMultiply->GetInput(0));
+		TransparencyOneMinus->ConnectExpression(*BaseColorTransparencyMultiply->GetInput(1));
+
+		GlowIntensityExpression->ConnectExpression(*IncandescenceScaleMultiply->GetInput(0));
+		IncandescenceScale->ConnectExpression(*IncandescenceScaleMultiply->GetInput(1));
+
+		BaseColorTransparencyMultiply->ConnectExpression(*IncandescenceMultiply->GetInput(0));
+		IncandescenceScaleMultiply->ConnectExpression(*IncandescenceMultiply->GetInput(1));
+
+		ShinynessExpression->ConnectExpression(*ShinynessSubtract->GetInput(0));
+		ShinynessSubtract2->ConnectExpression(*ShinynessSubtract->GetInput(1));
+
+		ShinynessSubtract->ConnectExpression(*ShinynessDivide->GetInput(0));
+		ShinynessDivide98->ConnectExpression(*ShinynessDivide->GetInput(1));
+
+		ShinynessDivide->ConnectExpression(*SpecularityMultiply->GetInput(0));
+		SpecularityExpression->ConnectExpression(*SpecularityMultiply->GetInput(1));
+
+		SpecularityMultiply->ConnectExpression(*RoughnessOneMinus->GetInput(0));
+
+		TransparencyColorExpression->ConnectExpression(*TransparencyOneMinus->GetInput(0));
+
+		if (bIsTransparent)
+		{
+			TransparencyOneMinus->ConnectExpression(*BreakFloat3->GetInput(0));
+
+			BreakFloat3->ConnectExpression(*AddRG->GetInput(0), 0);
+			BreakFloat3->ConnectExpression(*AddRG->GetInput(1), 1);
+
+			AddRG->ConnectExpression(*AddRGB->GetInput(0));
+			BreakFloat3->ConnectExpression(*AddRGB->GetInput(1), 2);
+
+			AddRGB->ConnectExpression(*Divide->GetInput(0));
+			DivideConstant->ConnectExpression(*Divide->GetInput(1));
+		}
+
+		// Connect material outputs
+		MaterialElement->GetBaseColor().SetExpression(BaseColorTransparencyMultiply);
+		MaterialElement->GetMetallic().SetExpression(GlossExpression);
 		MaterialElement->GetSpecular().SetExpression(ReflectivityExpression);
 		MaterialElement->GetRoughness().SetExpression(RoughnessOneMinus);
 		MaterialElement->GetEmissiveColor().SetExpression(IncandescenceMultiply);
@@ -1255,1089 +2113,4 @@
 		}
 	}
 
-	// Make material
-	bool FWireTranslatorImpl::GetShader()
-	{
-		for (TAlObjectPtr<AlShader> Shader(AlUniverse::firstShader()); Shader.IsValid(); Shader = TAlObjectPtr<AlShader>(AlUniverse::nextShader(Shader.Get())))
-		{
-			const FString ShaderName = UTF8_TO_TCHAR(Shader->name());
-			const FString ShaderModelName = Shader->shadingModel();
-
-			const FColor Color = CreateShaderColorFromShaderName(ShaderName);
-			ShaderNameToColor.Add(ShaderName, Color);
-
-			const int32 ShaderId = CreateShaderId(Color);
-			const FString MaterialElementName = FString::FromInt(ShaderId);
-
-			TSharedPtr<IDatasmithUEPbrMaterialElement> MaterialElement = FDatasmithSceneFactory::CreateUEPbrMaterial(*MaterialElementName);
-			MaterialElement->SetLabel(*ShaderName);
-
-			if (ShaderModelName.Equals(TEXT("BLINN")))
-			{
-				AddAlBlinnParameters(Shader, MaterialElement);
-			}
-			else if (ShaderModelName.Equals(TEXT("LAMBERT")))
-			{
-				AddAlLambertParameters(Shader, MaterialElement);
-			}
-			else if (ShaderModelName.Equals(TEXT("LIGHTSOURCE")))
-			{
-				AddAlLightSourceParameters(Shader, MaterialElement);
-			}
-			else if (ShaderModelName.Equals(TEXT("PHONG")))
-			{
-				AddAlPhongParameters(Shader, MaterialElement);
-			}
-
-			DatasmithScene->AddMaterial(MaterialElement);
-
-			TSharedPtr<IDatasmithMaterialIDElement> MaterialIDElement = FDatasmithSceneFactory::CreateMaterialId(MaterialElement->GetName());
-			ShaderNameToUEMaterialId.Add(*ShaderName, MaterialIDElement);
-		}
-		return true;
-	}
-
-	bool FWireTranslatorImpl::GetDagLeaves()
-	{
-		FDagNodeInfo RootContainer;
-		AlRootNode = TAlObjectPtr<AlDagNode>(AlUniverse::firstDagNode());
-		if (!AlRootNode.IsValid())
-		{
-			return false;
-		}
-
-		if (TessellationOptions.StitchingTechnique != EDatasmithCADStitchingTechnique::StitchingSew)
-		{
-			RecurseDagForLeavesNoMerge(AlRootNode, RootContainer);
-		}
-		else
-		{
-			RecurseDagForLeaves(AlRootNode, RootContainer);
-		}
-
-		// Delete empty branches
-		TArray<TSharedPtr<IDatasmithActorElement>> ActorsToRemove;
-		for (int32 Index = DatasmithScene->GetActorsCount() - 1; Index >= 0; --Index)
-		{
-			TSharedPtr<IDatasmithActorElement> Actor = DatasmithScene->GetActor(Index);
-			if (Actor.IsValid() && !Actor->IsA(EDatasmithElementType::StaticMeshActor))
-			{
-				RecurseDeleteEmptyActor(Actor);
-				if (Actor->GetChildrenCount() == 0)
-				{
-					ActorsToRemove.Add(Actor);
-				}
-			}
-		}
-
-		for (const TSharedPtr<IDatasmithActorElement>& Actor : ActorsToRemove)
-		{
-			DatasmithScene->RemoveActor(Actor, EDatasmithActorRemovalRule::RemoveChildren);
-		}
-
-		return true;
-	}
-
-	void FWireTranslatorImpl::RecurseDeleteEmptyActor(TSharedPtr<IDatasmithActorElement> Actor)
-	{
-		for (int32 Index = Actor->GetChildrenCount() - 1; Index >= 0; --Index)
-		{
-			TSharedPtr<IDatasmithActorElement> ActorChild = Actor->GetChild(Index);
-			if (ActorChild.IsValid() && !ActorChild->IsA(EDatasmithElementType::StaticMeshActor))
-			{
-				RecurseDeleteEmptyActor(ActorChild);
-				if (ActorChild->GetChildrenCount() == 0)
-				{
-					Actor->RemoveChild(ActorChild);
-				}
-			}
-		}
-	}
-
-
-	void FDagNodeInfo::ExtractDagNodeLayer(const AlDagNode& InDagNode)
-	{
-		AlLayer* LayerPtr = InDagNode.layer();
-		if (LayerPtr)
-		{
-			TAlObjectPtr<AlLayer> Layer(LayerPtr);
-			LayerName = UTF8_TO_TCHAR(Layer->name());
-			bLayerIsVisible = !Layer->invisible();
-			if (ActorElement.IsValid())
-			{
-				ActorElement->SetLayer(*LayerName);
-				LayerToActorElement.Add(LayerName, ActorElement);
-			}
-		}
-	}
-
-	void FDagNodeInfo::ExtractDagNodeInfo(AlDagNode& CurrentNode)
-	{
-		Label = UTF8_TO_TCHAR(CurrentNode.name());
-		uint32 ThisGroupNodeUuid = OpenModelUtils::GetAlDagNodeUuid(CurrentNode);
-		Uuid = HashCombine(ParentInfo ? ParentInfo->Uuid : 1, ThisGroupNodeUuid);
-	}
-
-	void FDagNodeInfo::ExtractDagNodeInfo(FBodyData& CurrentNode)
-	{
-		if (!ParentInfo)
-		{
-			return;
-		}
-		Label = ParentInfo->Label;
-		CurrentNode.SetLabel(ParentInfo->Label);
-		Uuid = CurrentNode.GetUuid(ParentInfo->Uuid);
-	}
-
-	TSharedPtr<IDatasmithActorElement> FWireTranslatorImpl::FindOrAddLayerActor(const FString& LayerName)
-	{
-		if (!bGLayersAsActors || LayerName.IsEmpty())
-		{
-			return TSharedPtr<IDatasmithActorElement>();
-		}
-
-		TSharedPtr<IDatasmithActorElement>* LayerActor = LayerNameToActor.Find(LayerName);
-		if (LayerActor)
-		{
-			return *LayerActor;
-		}
-
-		uint32 LayerUuid = GetTypeHash(LayerName);
-		TSharedPtr<IDatasmithActorElement> NewLayerActor = FDatasmithSceneFactory::CreateActor(*OpenModelUtils::UuidToString(LayerUuid));
-		NewLayerActor->SetLabel(*LayerName);
-		NewLayerActor->SetLayer(*LayerName);
-		DatasmithScene->AddActor(NewLayerActor);
-		LayerNameToActor.Add(LayerName, NewLayerActor);
-		return NewLayerActor;
-	}
-
-	TSharedPtr<IDatasmithActorElement> FWireTranslatorImpl::FindOrAddParentActor(FDagNodeInfo& ParentInfo, const FString& LayerName)
-	{
-		if (!ParentInfo.ActorElement.IsValid())
-		{
-			if (!LayerName.IsEmpty())
-			{
-				return FindOrAddLayerActor(LayerName);
-			}
-			return TSharedPtr<IDatasmithActorElement>();
-		}
-
-		if (!bGLayersAsActors)
-		{
-			return ParentInfo.ActorElement;
-		}
-
-		TSharedPtr<IDatasmithActorElement>* ActorParent = ParentInfo.LayerToActorElement.Find(LayerName);
-		if (ActorParent)
-		{
-			return *ActorParent;
-		}
-
-		uint32 LayerActorUuid = HashCombine(ParentInfo.Uuid, GetTypeHash(LayerName));
-
-		TSharedPtr<IDatasmithActorElement> LayerActorParent = FDatasmithSceneFactory::CreateActor(*OpenModelUtils::UuidToString(LayerActorUuid));
-		LayerActorParent->SetLabel(*ParentInfo.Label);
-		LayerActorParent->SetLayer(*LayerName);
-
-		ParentInfo.LayerToActorElement.Add(LayerName, LayerActorParent);
-
-		TSharedPtr<IDatasmithActorElement> GrandParentActor = FindOrAddParentActor(*ParentInfo.ParentInfo, LayerName);
-		if (GrandParentActor.IsValid())
-		{
-			GrandParentActor->AddChild(LayerActorParent);
-		}
-		else
-		{
-			TSharedPtr<IDatasmithActorElement> LayerActor = FindOrAddLayerActor(LayerName);
-			LayerActor->AddChild(LayerActorParent);
-		}
-
-		return LayerActorParent;
-	}
-
-	bool FWireTranslatorImpl::ProcessAlGroupNode(AlDagNode& GroupNode, FDagNodeInfo& ParentInfo)
-	{
-		AlGroupNode* AlGroup = GroupNode.asGroupNodePtr();
-		if (!AlIsValid(AlGroup))
-		{
-			return false;
-		}
-
-		AlDagNode* AlChildPtr = AlGroup->childNode();
-		if (!AlIsValid(AlChildPtr))
-		{
-			return false;
-		}
-
-		TAlObjectPtr<AlDagNode> ChildNode(AlChildPtr);
-
-		FDagNodeInfo ThisGroupNodeInfo(ParentInfo);
-		ThisGroupNodeInfo.ExtractDagNodeInfo(GroupNode);
-
-		ThisGroupNodeInfo.ActorElement = FDatasmithSceneFactory::CreateActor(*OpenModelUtils::UuidToString(ThisGroupNodeInfo.Uuid));
-		ThisGroupNodeInfo.ActorElement->SetLabel(*ThisGroupNodeInfo.Label);
-		ThisGroupNodeInfo.ExtractDagNodeLayer(GroupNode);
-
-		TSharedPtr<IDatasmithActorElement> ParentActorElement = FindOrAddParentActor(ParentInfo, ThisGroupNodeInfo.LayerName);
-		if (ParentActorElement.IsValid())
-		{
-			ParentActorElement->AddChild(ThisGroupNodeInfo.ActorElement);
-		}
-		else
-		{
-			DatasmithScene->AddActor(ThisGroupNodeInfo.ActorElement);
-		}
-
-		RecurseDagForLeaves(ChildNode, ThisGroupNodeInfo);
-
-		// Apply local transform to actor element
-		OpenModelUtils::SetActorTransform(ThisGroupNodeInfo.ActorElement, GroupNode);
-
-		return true;
-	}
-
-	TSharedPtr<IDatasmithMeshElement> FWireTranslatorImpl::FindOrAddMeshElement(const TSharedPtr<FBodyData>& Body, const FDagNodeInfo& NodeInfo)
-	{
-		TSharedPtr<IDatasmithMeshElement>* MeshElementPtr = BodyUuidToMeshElementMap.Find(NodeInfo.Uuid);
-		if (MeshElementPtr != nullptr)
-		{
-			return *MeshElementPtr;
-		}
-
-		TSharedPtr<IDatasmithMeshElement> MeshElement = FDatasmithSceneFactory::CreateMesh(*OpenModelUtils::UuidToString(NodeInfo.Uuid));
-		MeshElement->SetLabel(*NodeInfo.Label);
-		MeshElement->SetLightmapSourceUV(-1);
-
-		const TSet<FString>& ShaderNames = Body->GetShaderNames();
-		for (const FString& ShaderName : ShaderNames)
-		{
-			TSharedPtr<IDatasmithMaterialIDElement> MaterialElement = ShaderNameToUEMaterialId[ShaderName];
-			if (MaterialElement.IsValid())
-			{
-				const FColor* ShaderColor = ShaderNameToColor.Find(ShaderName);
-				if (ShaderColor)
-				{
-					int32 ShaderId = CreateShaderId(*ShaderColor);
-					MeshElement->SetMaterial(MaterialElement->GetName(), ShaderId);
-				}
-			}
-		}
-
-		DatasmithScene->AddMesh(MeshElement);
-
-		ShellUuidToMeshElementMap.Add(NodeInfo.Uuid, MeshElement);
-		MeshElementToBodyMap.Add(MeshElement.Get(), Body);
-
-		BodyUuidToMeshElementMap.Add(NodeInfo.Uuid, MeshElement);
-
-		return MeshElement;
-	}
-
-	TSharedPtr<IDatasmithMeshElement> FWireTranslatorImpl::FindOrAddMeshElement(const TAlObjectPtr<AlDagNode>& ShellNode, const FDagNodeInfo& ShellNodeInfo, const FString& ShaderName)
-	{
-		uint32 ShellUuid = OpenModelUtils::GetAlDagNodeUuid(*ShellNode);
-
-		// Look if geometry has not been already processed, return it if found
-		TSharedPtr<IDatasmithMeshElement>* MeshElementPtr = ShellUuidToMeshElementMap.Find(ShellUuid);
-		if (MeshElementPtr != nullptr)
-		{
-			return *MeshElementPtr;
-		}
-
-		TSharedPtr<IDatasmithMeshElement> MeshElement = FDatasmithSceneFactory::CreateMesh(*OpenModelUtils::UuidToString(ShellNodeInfo.Uuid));
-		MeshElement->SetLabel(*ShellNodeInfo.Label);
-		MeshElement->SetLightmapSourceUV(-1);
-
-		// Set MeshElement FileHash used for re-import task
-		FMD5 MD5; // unique Value that define the mesh
-		MD5.Update(reinterpret_cast<const uint8*>(&SceneFileHash), sizeof SceneFileHash);
-		// MeshActor Name
-		MD5.Update(reinterpret_cast<const uint8*>(&ShellUuid), sizeof ShellUuid);
-		FMD5Hash Hash;
-		Hash.Set(MD5);
-		MeshElement->SetFileHash(Hash);
-
-		if (ShaderName.Len())
-		{
-			TSharedPtr<IDatasmithMaterialIDElement> MaterialElement = ShaderNameToUEMaterialId[ShaderName];
-			if (MaterialElement.IsValid())
-			{
-				const FColor* ShaderColor = ShaderNameToColor.Find(ShaderName);
-				if (ShaderColor)
-				{
-					int32 ShaderId = CreateShaderId(*ShaderColor);
-					MeshElement->SetMaterial(MaterialElement->GetName(), ShaderId);
-				}
-			}
-
-		}
-
-		DatasmithScene->AddMesh(MeshElement);
-
-		ShellUuidToMeshElementMap.Add(ShellUuid, MeshElement);
-		MeshElementToAlDagNodeMap.Add(MeshElement.Get(), ShellNode);
-		MeshElementToShaderName.Add(MeshElement.Get(), ShaderName);
-
-		return MeshElement;
-	}
-
-	void FWireTranslatorImpl::ProcessAlShellNode(const TAlObjectPtr<AlDagNode>& ShellNode, FDagNodeInfo& ParentInfo, const FString& ShaderName)
-	{
-		FDagNodeInfo ShellInfo(ParentInfo);
-		ShellInfo.ExtractDagNodeInfo(*ShellNode);
-
-		TSharedPtr<IDatasmithMeshElement> MeshElement = FindOrAddMeshElement(ShellNode, ShellInfo, ShaderName);
-		if (!MeshElement.IsValid())
-		{
-			return;
-		}
-
-		TSharedPtr<IDatasmithMeshActorElement> ActorElement = FDatasmithSceneFactory::CreateMeshActor(*OpenModelUtils::UuidToString(ShellInfo.Uuid));
-		if (!ActorElement.IsValid())
-		{
-			return;
-		}
-
-		ActorElement->SetLabel(*ShellInfo.Label);
-		ActorElement->SetStaticMeshPathName(MeshElement->GetName());
-		ShellInfo.ActorElement = ActorElement;
-
-		ShellInfo.ExtractDagNodeLayer(*ShellNode);
-		if (!ShellInfo.bLayerIsVisible)
-		{
-			return;
-		}
-
-		// Apply materials on the current part
-		if (ShaderName.Len())
-		{
-			const TSharedPtr<IDatasmithMaterialIDElement>& MaterialIDElement = ShaderNameToUEMaterialId[ShaderName];
-			if (MaterialIDElement.IsValid())
-			{
-				for (int32 Index = 0; Index < MeshElement->GetMaterialSlotCount(); ++Index)
-				{
-					MaterialIDElement->SetId(MeshElement->GetMaterialSlotAt(Index)->GetId());
-					ActorElement->AddMaterialOverride(MaterialIDElement);
-				}
-			}
-		}
-
-		TSharedPtr<IDatasmithActorElement> ParentElement = FindOrAddParentActor(ParentInfo, ShellInfo.LayerName);
-		if (ParentElement.IsValid())
-		{
-			ParentElement->AddChild(ActorElement);
-		}
-		else
-		{
-			DatasmithScene->AddActor(ActorElement);
-		}
-
-		OpenModelUtils::SetActorTransform(ShellInfo.ActorElement, *ShellNode);
-	}
-
-	void FWireTranslatorImpl::ProcessBodyNode(const TSharedPtr<FBodyData>& Body, FDagNodeInfo& ParentInfo)
-	{
-
-		if (!Body.IsValid())
-		{
-			return;
-		}
-
-		if (Body->GetShells().Num() == 1)
-		{
-			DagForLeavesNoMerge(Body->GetShells()[0].Key, ParentInfo);
-			return;
-		}
-
-		FDagNodeInfo ShellInfo(ParentInfo);
-		ShellInfo.ExtractDagNodeInfo(*Body);
-
-		TSharedPtr<IDatasmithMeshElement> MeshElement = FindOrAddMeshElement(Body, ShellInfo);
-		if (!MeshElement.IsValid())
-		{
-			return;
-		}
-
-		TSharedPtr<IDatasmithMeshActorElement> ActorElement = FDatasmithSceneFactory::CreateMeshActor(*OpenModelUtils::UuidToString(ShellInfo.Uuid));
-		if (!ActorElement.IsValid())
-		{
-			return;
-		}
-
-		ActorElement->SetLabel(ShellInfo.Label.IsEmpty() ? *Body->GetLayerName() : *ShellInfo.Label);
-		ActorElement->SetStaticMeshPathName(MeshElement->GetName());
-		ShellInfo.ActorElement = ActorElement;
-
-		ActorElement->SetLayer(*Body->GetLayerName());
-
-		// Apply materials on the current part
-		int32 Index = 0;
-		const TSet<FString>& ShaderNames = Body->GetShaderNames();
-		for (const FString& ShaderName : ShaderNames)
-		{
-			TSharedPtr<IDatasmithMaterialIDElement> MaterialIDElement = ShaderNameToUEMaterialId[ShaderName];
-			if (MaterialIDElement.IsValid())
-			{
-				MaterialIDElement->SetId(MeshElement->GetMaterialSlotAt(Index)->GetId());
-				ActorElement->AddMaterialOverride(MaterialIDElement);
-				Index++;
-			}
-		}
-
-		TSharedPtr<IDatasmithActorElement> ParentElement = FindOrAddParentActor(ParentInfo, Body->GetLayerName());
-		if (ParentElement.IsValid())
-		{
-			ParentElement->AddChild(ActorElement);
-		}
-		else
-		{
-			DatasmithScene->AddActor(ActorElement);
-		}
-
-		return;
-	}
-
-	TAlObjectPtr<AlDagNode> GetNextNode(const TAlObjectPtr<AlDagNode>& DagNode)
-	{
-		// Grab the next sibling before deleting the node.
-		AlDagNode* SiblingNode = DagNode->nextNode();
-		if (AlIsValid(SiblingNode))
-		{
-			return TAlObjectPtr<AlDagNode>(SiblingNode);
-		}
-		else
-		{
-			return TAlObjectPtr<AlDagNode>();
-		}
-	}
-
-	uint32 GetBodyGroupUuid(const FString& ShaderName, const FString& LayerName, bool bCadData)
-	{
-		uint32 Uuid = HashCombine(GetTypeHash(LayerName), GetTypeHash(bCadData));
-		if (bGSewByMaterial)
-		{
-			Uuid = HashCombine(GetTypeHash(ShaderName), Uuid);
-		}
-		return Uuid;
-	}
-
-	uint32 GetPatchCount(const TAlObjectPtr<AlShell>& Shell)
-	{
-		uint32 PatchCount = 0;
-		TAlObjectPtr<AlTrimRegion> TrimRegion(Shell->firstTrimRegion());
-		while (TrimRegion.IsValid())
-		{
-			PatchCount++;
-			TrimRegion = TAlObjectPtr<AlTrimRegion>(TrimRegion->nextRegion());
-		}
-		return PatchCount;
-	}
-
-	void FWireTranslatorImpl::AddNodeInBodyGroup(TAlObjectPtr<AlDagNode>& DagNode, const FString& ShaderName, TMap<uint32, TSharedPtr<FBodyData>>& ShellToProcess, bool bIsAPatch, uint32 MaxSize)
-	{
-		FString LayerName;
-
-		AlLayer* LayerPtr = DagNode->layer();
-		if (AlIsValid(LayerPtr))
-		{
-			const TAlObjectPtr<AlLayer> Layer(LayerPtr);
-			LayerName = UTF8_TO_TCHAR(Layer->name());
-			if (Layer->invisible())
-			{
-				return;
-			}
-		}
-
-		uint32 SetId = GetBodyGroupUuid(ShaderName, LayerName, bIsAPatch);
-
-		TSharedPtr<FBodyData> Body;
-		if (TSharedPtr<FBodyData>* PBody = ShellToProcess.Find(SetId))
-		{
-			Body = *PBody;
-		}
-		else
-		{
-			TSharedPtr<FBodyData> BodyRef = MakeShared<FBodyData>(LayerName, bIsAPatch);
-			ShellToProcess.Add(SetId, BodyRef);
-			BodyRef->Reserve(MaxSize);
-			Body = BodyRef;
-		}
-
-		const FColor& Color = ShaderNameToColor[ShaderName];
-		Body->Add(DagNode, Color, ShaderName);
-	}
-
-	void FWireTranslatorImpl::RecurseDagForLeaves(const TAlObjectPtr<AlDagNode>& FirstDagNode, FDagNodeInfo& ParentInfo)
-	{
-		TAlObjectPtr<AlDagNode> DagNode = FirstDagNode;
-		uint32 MaxSize = 0;
-		while (DagNode)
-		{
-			MaxSize++;
-			DagNode = GetNextNode(DagNode);
-		}
-
-		DagNode = FirstDagNode;
-
-		TMap<uint32, TSharedPtr<FBodyData>> ShellToProcess;
-
-		const char* ShaderName = nullptr;
-
-		while (DagNode)
-		{
-			AlObjectType objectType = DagNode->type();
-
-			// Process the current node.
-			switch (objectType)
-			{
-				// Push all leaf nodes into 'leaves'
-			case kShellNodeType:
-			{
-				AlShellNode* ShellNode = DagNode->asShellNodePtr();
-				if (!AlIsValid(ShellNode))
-				{
-					break;
-				}
-
-				AlShell* ShellPtr = ShellNode->shell();
-				if (!AlIsValid(ShellPtr))
-				{
-					break;
-				}
-
-				TAlObjectPtr<AlShell> Shell(ShellPtr);
-				uint32 NbPatch = GetPatchCount(Shell);
-
-				TAlObjectPtr<AlShader> Shader(Shell->firstShader());
-				if (Shader.IsValid())
-				{
-					ShaderName = Shader->name();
-				}
-
-				if (NbPatch == 1)
-				{
-					AddNodeInBodyGroup(DagNode, ShaderName, ShellToProcess, true, MaxSize);
-				}
-				else
-				{
-					ProcessAlShellNode(DagNode, ParentInfo, ShaderName);
-				}
-				break;
-			}
-
-			case kSurfaceNodeType:
-			{
-				AlSurfaceNode* SurfaceNode = DagNode->asSurfaceNodePtr();
-				TAlObjectPtr<AlSurface> Surface(SurfaceNode->surface());
-				if (Surface.IsValid())
-				{
-					TAlObjectPtr<AlShader> Shader(Surface->firstShader());
-					if (Shader.IsValid())
-					{
-						ShaderName = Shader->name();
-					}
-				}
-				AddNodeInBodyGroup(DagNode, ShaderName, ShellToProcess, true, MaxSize);
-				break;
-			}
-
-			case kMeshNodeType:
-			{
-				AlMeshNode* MeshNode = DagNode->asMeshNodePtr();
-				TAlObjectPtr<AlMesh> Mesh(MeshNode->mesh());
-				if (Mesh.IsValid())
-				{
-					TAlObjectPtr<AlShader> Shader(Mesh->firstShader());
-					if (Shader.IsValid())
-					{
-						ShaderName = Shader->name();
-					}
-				}
-				AddNodeInBodyGroup(DagNode, ShaderName, ShellToProcess, false, MaxSize);
-				break;
-			}
-
-			// Traverse down through groups
-			case kGroupNodeType:
-			{
-				ProcessAlGroupNode(*DagNode, ParentInfo);
-				break;
-			}
-
-			default:
-			{
-				break;
-			}
-			}
-
-			DagNode = GetNextNode(DagNode);
-		}
-
-		for (const TPair<uint32, TSharedPtr<FBodyData>>& Body : ShellToProcess)
-		{
-			ProcessBodyNode(Body.Value, ParentInfo);
-		}
-	}
-
-	void FWireTranslatorImpl::RecurseDagForLeavesNoMerge(const TAlObjectPtr<AlDagNode>& FirstDagNode, FDagNodeInfo& ParentInfo)
-	{
-		TAlObjectPtr<AlDagNode> DagNode = FirstDagNode;
-		while (DagNode)
-		{
-			DagForLeavesNoMerge(DagNode, ParentInfo);
-			DagNode = GetNextNode(DagNode);
-		}
-	}
-
-	void FWireTranslatorImpl::DagForLeavesNoMerge(const TAlObjectPtr<AlDagNode>& DagNode, FDagNodeInfo& ParentInfo)
-	{
-		const char* ShaderName = nullptr;
-
-		// Process the current node.
-		AlObjectType objectType = DagNode->type();
-		switch (objectType)
-		{
-			// Push all leaf nodes into 'leaves'
-		case kShellNodeType:
-		{
-			AlShellNode* ShellNode = DagNode->asShellNodePtr();
-			TAlObjectPtr<AlShell> Shell(ShellNode->shell());
-			if (Shell.IsValid())
-			{
-				TAlObjectPtr<AlShader> Shader(Shell->firstShader());
-				if (Shader.IsValid())
-				{
-					ShaderName = Shader->name();
-				}
-
-				ProcessAlShellNode(DagNode, ParentInfo, ShaderName);
-			}
-			break;
-		}
-		case kSurfaceNodeType:
-		{
-			AlSurfaceNode* SurfaceNode = DagNode->asSurfaceNodePtr();
-			TAlObjectPtr<AlSurface> Surface(SurfaceNode->surface());
-			if (Surface.IsValid())
-			{
-				TAlObjectPtr<AlShader> Shader(Surface->firstShader());
-				if (Shader.IsValid())
-				{
-					ShaderName = Shader->name();
-				}
-			}
-			ProcessAlShellNode(DagNode, ParentInfo, ShaderName);
-			break;
-		}
-
-		case kMeshNodeType:
-		{
-			AlMeshNode* MeshNode = DagNode->asMeshNodePtr();
-			TAlObjectPtr<AlMesh> Mesh(MeshNode->mesh());
-			if (Mesh.IsValid())
-			{
-				TAlObjectPtr<AlShader> Shader(Mesh->firstShader());
-				if (Shader.IsValid())
-				{
-					ShaderName = Shader->name();
-				}
-			}
-			ProcessAlShellNode(DagNode, ParentInfo, ShaderName);
-			break;
-		}
-
-		// Traverse down through groups
-		case kGroupNodeType:
-		{
-			ProcessAlGroupNode(*DagNode, ParentInfo);
-			break;
-		}
-
-		default:
-		{
-			break;
-		}
-
-		}
-	}
-
-	TSharedPtr<CADLibrary::ICADModelConverter> FWireTranslatorImpl::GetModelConverter(IAliasBRepConverter*& BRepConverter) const
-	{
-		TSharedPtr<CADLibrary::ICADModelConverter> ModelConverter;
-
-		CADLibrary::FImportParameters ImportParameters;
-		if (CADLibrary::FImportParameters::bGDisableCADKernelTessellation)
-		{
-			TSharedPtr<FAliasModelToTechSoftConverter> AliasToTechSoftConverter = MakeShared<FAliasModelToTechSoftConverter>(ImportParameters);
-			ModelConverter = AliasToTechSoftConverter;
-			BRepConverter = AliasToTechSoftConverter.Get();
-		}
-		else
-		{
-			TSharedPtr<FAliasModelToCADKernelConverter> AliasToCADKernelConverter = MakeShared<FAliasModelToCADKernelConverter>(ImportParameters);
-			ModelConverter = AliasToCADKernelConverter;
-			BRepConverter = AliasToCADKernelConverter.Get();
-		}
-
-		ModelConverter->SetImportParameters(TessellationOptions.ChordTolerance, TessellationOptions.MaxEdgeLength, TessellationOptions.NormalTolerance, (CADLibrary::EStitchingTechnique)TessellationOptions.StitchingTechnique);
-
-		return ModelConverter;
-	}
-
-	TOptional<FMeshDescription> FWireTranslatorImpl::MeshDagNodeWithExternalMesher(AlDagNode& DagNode, TSharedPtr<IDatasmithMeshElement> MeshElement, CADLibrary::FMeshParameters& MeshParameters)
-	{
-		// Wire unit is cm
-		IAliasBRepConverter* BRepConverter = nullptr;
-		TSharedPtr<CADLibrary::ICADModelConverter> ModelConverter = GetModelConverter(BRepConverter);
-
-		ModelConverter->InitializeProcess();
-
-		EAliasObjectReference ObjectReference = EAliasObjectReference::LocalReference;
-
-		if (MeshParameters.bIsSymmetric)
-		{
-			// All actors of a Alias symmetric layer are defined in the world Reference i.e. they have identity transform. So Mesh actor has to be defined in the world reference.
-			ObjectReference = EAliasObjectReference::WorldReference;
-		}
-
-		const FColor* ColorPtr = nullptr;
-		const FString* ShaderName = MeshElementToShaderName.Find(MeshElement.Get());
-		if (ShaderName)
-		{
-			ColorPtr = ShaderNameToColor.Find(*ShaderName);
-		}
-		if (!ColorPtr) // Should never happen
-		{
-			ColorPtr = &DefaultColor;
-		}
-
-		if (!BRepConverter->AddBRep(DagNode, *ColorPtr, ObjectReference))
-		{
-			return TOptional<FMeshDescription>();
-		}
-
-		ModelConverter->RepairTopology();
-
-		CADModelConverter->SaveModel(*OutputPath, MeshElement);
-
-		FMeshDescription MeshDescription;
-		DatasmithMeshHelper::PrepareAttributeForStaticMesh(MeshDescription);
-
-		bool bRet = ModelConverter->Tessellate(MeshParameters, MeshDescription);
-		if (!bRet)
-		{
-			const TCHAR* StaticMeshLable = MeshElement->GetLabel();
-			const TCHAR* StaticMeshName = MeshElement->GetName();
-			UE_LOG(LogWireInterface, Warning, TEXT("Failed to generate the mesh of \"%s\" (%s) StaticMesh."), StaticMeshLable, StaticMeshName);
-		}
-
-		return MoveTemp(MeshDescription);
-	}
-
-	TOptional<FMeshDescription> FWireTranslatorImpl::MeshDagNodeWithExternalMesher(TSharedPtr<FBodyData> Body, TSharedPtr<IDatasmithMeshElement> MeshElement, CADLibrary::FMeshParameters& MeshParameters)
-	{
-		// Wire unit is cm
-		IAliasBRepConverter* BRepConverter = nullptr;
-		TSharedPtr<CADLibrary::ICADModelConverter> ModelConverter = GetModelConverter(BRepConverter);
-
-		ModelConverter->InitializeProcess();
-
-		EAliasObjectReference ObjectReference = EAliasObjectReference::LocalReference;
-		if (MeshParameters.bIsSymmetric)
-		{
-			// All actors of a Alias symmetric layer are defined in the world Reference i.e. they have identity transform. So Mesh actor has to be defined in the world reference.
-			ObjectReference = EAliasObjectReference::WorldReference;
-		}
-		else if (TessellationOptions.StitchingTechnique == EDatasmithCADStitchingTechnique::StitchingSew)
-		{
-			// In the case of StitchingSew, AlDagNode children of a GroupNode are merged together. To be merged, they have to be defined in the reference of parent GroupNode.
-			ObjectReference = EAliasObjectReference::ParentReference;
-		}
-
-		for (const TPair<TAlObjectPtr<AlDagNode>, FColor>& DagNode : Body->GetShells())
-		{
-			BRepConverter->AddBRep(*DagNode.Key, DagNode.Value, ObjectReference);
-		}
-
-		ModelConverter->RepairTopology();
-
-		ModelConverter->SaveModel(*OutputPath, MeshElement);
-
-		FMeshDescription MeshDescription;
-		DatasmithMeshHelper::PrepareAttributeForStaticMesh(MeshDescription);
-
-		ModelConverter->Tessellate(MeshParameters, MeshDescription);
-
-		return MoveTemp(MeshDescription);
-	}
-
-	TOptional<FMeshDescription> FWireTranslatorImpl::GetMeshOfShellNode(AlDagNode& DagNode, TSharedPtr<IDatasmithMeshElement> MeshElement, CADLibrary::FMeshParameters& MeshParameters)
-	{
-		if (CADModelConverter.IsValid() && CADModelConverter->IsSessionValid())
-		{
-			TOptional<FMeshDescription> UEMesh = MeshDagNodeWithExternalMesher(DagNode, MeshElement, MeshParameters);
-			return UEMesh;
-		}
-		else
-		{
-			AlMatrix4x4 AlMatrix;
-			DagNode.inverseGlobalTransformationMatrix(AlMatrix);
-			// TODO: the best way, should be to don't have to apply inverse global transform to the generated mesh
-			TAlObjectPtr<AlDagNode> TesselatedNode = OpenModelUtils::TesselateDagLeaf(DagNode, ETesselatorType::Fast, TessellationOptions.ChordTolerance);
-			if (TesselatedNode.IsValid())
-			{
-				// Get the meshes from the dag nodes. Note that removing the mesh's DAG.
-				// will also removes the meshes, so we have to do it later.
-				TOptional<FMeshDescription> UEMesh = GetMeshOfNodeMesh(*TesselatedNode, MeshElement, MeshParameters, &AlMatrix);
-				return UEMesh;
-			}
-		}
-		return TOptional<FMeshDescription>();
-	}
-
-	TOptional<FMeshDescription> FWireTranslatorImpl::GetMeshOfShellBody(TSharedPtr<FBodyData> Body, TSharedPtr<IDatasmithMeshElement> MeshElement, CADLibrary::FMeshParameters& MeshParameters)
-	{
-		TOptional<FMeshDescription> UEMesh = MeshDagNodeWithExternalMesher(Body, MeshElement, MeshParameters);
-		return UEMesh;
-	}
-
-	TOptional<FMeshDescription> FWireTranslatorImpl::GetMeshOfMeshBody(TSharedPtr<FBodyData> Body, TSharedPtr<IDatasmithMeshElement> MeshElement, CADLibrary::FMeshParameters& MeshParameters)
-	{
-		FMeshDescription MeshDescription;
-		DatasmithMeshHelper::PrepareAttributeForStaticMesh(MeshDescription);
-		MeshDescription.Empty();
-
-		for (const TPair<TAlObjectPtr<AlDagNode>, FColor>& DagNode : Body->GetShells())
-		{
-			AlMeshNode* MeshNode = DagNode.Key->asMeshNodePtr();
-			if (!AlIsValid(MeshNode))
-			{
-				continue;
-			}
-
-			AlMesh* MeshPtr = MeshNode->mesh();
-			if (!AlIsValid(MeshPtr))
-			{
-				continue;
-			}
-			TAlObjectPtr<AlMesh> Mesh(MeshPtr);
-
-			AlMatrix4x4 AlMatrix;
-			DagNode.Key->localTransformationMatrix(AlMatrix);
-
-			Mesh->transform(AlMatrix);
-
-			const FColor& ShaderColor = DagNode.Value;
-			const int32 ShaderId = CreateShaderId(ShaderColor);
-			const FString SlotMaterialName = FString::FromInt(ShaderId);
-
-			const bool bMerge = true;
-			OpenModelUtils::TransferAlMeshToMeshDescription(*MeshPtr, *SlotMaterialName, MeshDescription, MeshParameters, bMerge);
-		}
-
-		// Build edge meta data
-		FStaticMeshOperations::DetermineEdgeHardnessesFromVertexInstanceNormals(MeshDescription);
-
-		return MoveTemp(MeshDescription);
-	}
-
-	TOptional<FMeshDescription> FWireTranslatorImpl::GetMeshOfNodeMesh(AlDagNode& TesselatedNode, TSharedPtr<IDatasmithMeshElement> MeshElement, CADLibrary::FMeshParameters& MeshParameters, AlMatrix4x4* AlMeshInvGlobalMatrix)
-	{
-		AlMeshNode* MeshNode = TesselatedNode.asMeshNodePtr();
-		if (!AlIsValid(MeshNode))
-		{
-			return TOptional<FMeshDescription>();
-		}
-
-		AlMesh* Mesh = MeshNode->mesh();
-		if (!AlIsValid(Mesh))
-		{
-			return TOptional<FMeshDescription>();
-		}
-
-		TAlObjectPtr<AlMesh> SharedMesh(Mesh);
-		if (AlMeshInvGlobalMatrix != nullptr)
-		{
-			SharedMesh->transform(*AlMeshInvGlobalMatrix);
-		}
-
-		return ImportMesh(*Mesh, MeshElement, MeshParameters);
-	}
-
-	TOptional<FMeshDescription> FWireTranslatorImpl::GetMeshDescription(TSharedPtr<IDatasmithMeshElement> MeshElement, CADLibrary::FMeshParameters& MeshParameters, TSharedPtr<FBodyData> Body)
-	{
-		if (Body->GetShells().Num() == 0)
-		{
-			return TOptional<FMeshDescription>();
-		}
-
-		const TArray<TPair<TAlObjectPtr<AlDagNode>, FColor>>& Shells = Body->GetShells();
-		TAlObjectPtr<AlDagNode> DagNode = Shells[0].Key;
-		AlLayer* LayerPtr = DagNode->layer();
-		if (AlIsValid(LayerPtr))
-		{
-			const TAlObjectPtr<AlLayer> Layer(LayerPtr);
-			if (LayerPtr->isSymmetric())
-			{
-				MeshParameters.bIsSymmetric = true;
-				double Normal[3], Origin[3];
-				LayerPtr->symmetricNormal(Normal[0], Normal[1], Normal[2]);
-				LayerPtr->symmetricOrigin(Origin[0], Origin[1], Origin[2]);
-
-				MeshParameters.SymmetricOrigin.X = (float)Origin[0];
-				MeshParameters.SymmetricOrigin.Y = (float)Origin[1];
-				MeshParameters.SymmetricOrigin.Z = (float)Origin[2];
-				MeshParameters.SymmetricNormal.X = (float)Normal[0];
-				MeshParameters.SymmetricNormal.Y = (float)Normal[1];
-				MeshParameters.SymmetricNormal.Z = (float)Normal[2];
-			}
-		}
-
-		if (Body->IsCadData())
-		{
-			return GetMeshOfShellBody(Body, MeshElement, MeshParameters);
-		}
-		else
-		{
-			return GetMeshOfMeshBody(Body, MeshElement, MeshParameters);
-		}
-	}
-
-	TOptional<FMeshDescription> FWireTranslatorImpl::GetMeshDescription(TSharedPtr<IDatasmithMeshElement> MeshElement, CADLibrary::FMeshParameters& MeshParameters)
-	{
-		TAlObjectPtr<AlDagNode>* DagNodeTemp = MeshElementToAlDagNodeMap.Find(MeshElement.Get());
-		if (DagNodeTemp == nullptr || !DagNodeTemp->IsValid())
-		{
-			TSharedPtr<FBodyData>* BodyTemp = MeshElementToBodyMap.Find(MeshElement.Get());
-			if (BodyTemp != nullptr && BodyTemp->IsValid())
-			{
-				return GetMeshDescription(MeshElement, MeshParameters, (*BodyTemp));
-			}
-			return TOptional<FMeshDescription>();
-		}
-
-		AlDagNode& DagNode = **DagNodeTemp;
-		AlObjectType objectType = DagNode.type();
-
-		if (objectType == kShellNodeType || objectType == kSurfaceNodeType || objectType == kMeshNodeType)
-		{
-			boolean bAlOrientation;
-			DagNode.getSurfaceOrientation(bAlOrientation);
-			MeshParameters.bNeedSwapOrientation = (objectType == kMeshNodeType) ? (bool)bAlOrientation : false; 
-
-			AlLayer* LayerPtr = DagNode.layer();
-			if (AlIsValid(LayerPtr))
-			{
-				const TAlObjectPtr<AlLayer> Layer(LayerPtr);
-				if (LayerPtr->isSymmetric())
-				{
-					MeshParameters.bIsSymmetric = true;
-					double Normal[3], Origin[3];
-					LayerPtr->symmetricNormal(Normal[0], Normal[1], Normal[2]);
-					LayerPtr->symmetricOrigin(Origin[0], Origin[1], Origin[2]);
-
-					MeshParameters.SymmetricOrigin.X = (float)Origin[0];
-					MeshParameters.SymmetricOrigin.Y = (float)Origin[1];
-					MeshParameters.SymmetricOrigin.Z = (float)Origin[2];
-					MeshParameters.SymmetricNormal.X = (float)Normal[0];
-					MeshParameters.SymmetricNormal.Y = (float)Normal[1];
-					MeshParameters.SymmetricNormal.Z = (float)Normal[2];
-				}
-			}
-		}
-
-		switch (objectType)
-		{
-		case kShellNodeType:
-		case kSurfaceNodeType:
-		{
-			return GetMeshOfShellNode(DagNode, MeshElement, MeshParameters);
-			break;
-		}
-
-		case kMeshNodeType:
-		{
-			return GetMeshOfNodeMesh(DagNode, MeshElement, MeshParameters);
-			break;
-		}
-
-		default:
-			break;
-		}
-
-		return TOptional<FMeshDescription>();
-	}
-
-
-	// Note that Alias file unit is cm like UE
-	TOptional<FMeshDescription> FWireTranslatorImpl::ImportMesh(AlMesh& InMesh, TSharedPtr<IDatasmithMeshElement> MeshElement, CADLibrary::FMeshParameters& InMeshParameters)
-	{
-		FMeshDescription MeshDescription;
-		DatasmithMeshHelper::PrepareAttributeForStaticMesh(MeshDescription);
-
-		const FColor* ShaderColorPtr = nullptr;
-		const FString* ShaderName = MeshElementToShaderName.Find(MeshElement.Get());
-		if (ShaderName)
-		{
-			ShaderColorPtr = ShaderNameToColor.Find(*ShaderName);
-		}
-		if (!ShaderColorPtr) // Should never happen
-		{
-			ShaderColorPtr = &DefaultColor;
-		}
-		int32 ShaderId = CreateShaderId(*ShaderColorPtr);
-		FString SlotMaterialName = FString::FromInt(ShaderId);
-
-		const bool bMerge = false;
-		OpenModelUtils::TransferAlMeshToMeshDescription(InMesh, *SlotMaterialName, MeshDescription, InMeshParameters, bMerge);
-
-		// Build edge meta data
-		FStaticMeshOperations::DetermineEdgeHardnessesFromVertexInstanceNormals(MeshDescription);
-
-		return MoveTemp(MeshDescription);
-	}
-
-	#define TRACK_MESH 0
-
-	bool FWireTranslatorImpl::LoadStaticMesh(const TSharedPtr<IDatasmithMeshElement> MeshElement, FDatasmithMeshElementPayload& OutMeshPayload, const FDatasmithTessellationOptions& InTessellationOptions)
-	{
-	#if TRACK_MESH
-		static const FString MeshName(TEXT("shell_30181"));
-		bool bSkip = true;
-		if (MeshName.Equals(MeshElement->GetLabel()))
-		{
-			bSkip = false;
-		}
-		if (bSkip)
-		{
-			return false;
-		}
-	#endif
-		CADLibrary::FMeshParameters MeshParameters;
-
-		uint64 StartTime = FPlatformTime::Cycles64();
-
-		if (TOptional<FMeshDescription> Mesh = GetMeshDescription(MeshElement, MeshParameters))
-		{
-			OutMeshPayload.LodMeshes.Add(MoveTemp(Mesh.GetValue()));
-			const TCHAR* MeshFilename = MeshElement->GetFile();
-			if (!bAliasUseNative && FPaths::FileExists(MeshFilename))
-			{
-				CADModelConverter->AddSurfaceDataForMesh(MeshFilename, MeshParameters, InTessellationOptions, OutMeshPayload);
-
-				// Remove the file because it is temporary since caching is disabled.
-				if (!CADLibrary::FImportParameters::bGEnableCADCache)
-				{
-					IFileManager::Get().Delete(MeshFilename);
-				}
-			}
-		}
-
-		double ElapsedSeconds = FPlatformTime::ToSeconds64(FPlatformTime::Cycles64() - StartTime);
-
-		UE_LOG(LogWireInterface, Display, TEXT("Mesh description of %s retrieved in %.3f s"), MeshElement->GetLabel(), ElapsedSeconds);
-
-		return OutMeshPayload.LodMeshes.Num() > 0;
-	}
-
 	class FWireInterfaceModule : public IModuleInterface
@@ -2345,7 +2118,6 @@
 	public:
 		virtual void StartupModule() override
 		{
-#ifdef USE_OPENMODEL
 			uint64 AliasVersion = IWireInterface::GetRequiredAliasVersion();
 
 #ifdef OPEN_MODEL_2020
@@ -2369,7 +2141,6 @@
 					};
 				IWireInterface::RegisterInterface(UE_OPENMODEL_MAJOR_VERSION, UE_OPENMODEL_MAJOR_VERSION, MoveTemp(MakeInterfaceFunc));
 			}
-#endif
 		}
 
 		virtual void ShutdownModule() override
@@ -2381,7 +2152,7 @@
 
 #undef LOCTEXT_NAMESPACE // "DatasmithWireTranslator"
 
-#else
+#else // USE_OPENMODEL
 namespace UE_DATASMITHWIRETRANSLATOR_NAMESPACE
 {
 	class FWireInterfaceModule : public IModuleInterface
diff --git a/Engine/Plugins/Enterprise/DatasmithCADImporter/Source/WireInterface/Private/WireInterfaceImpl.h b/Engine/Plugins/Enterprise/DatasmithCADImporter/Source/WireInterface/Private/WireInterfaceImpl.h
--- a/Engine/Plugins/Enterprise/DatasmithCADImporter/Source/WireInterface/Private/WireInterfaceImpl.h
+++ b/Engine/Plugins/Enterprise/DatasmithCADImporter/Source/WireInterface/Private/WireInterfaceImpl.h
@@ -11,19 +11,19 @@
 
 #include "IWireInterface.h"
 
-#include "DatasmithImportOptions.h"
-
 #include "Templates/SharedPointer.h"
 
 class AlDagNode;
+class AlLayer;
 class AlMesh;
+class AlShader;
 class AlShaderNode;
 class IDatasmithActorElement;
+class IDatasmithBaseMaterialElement;
 class IDatasmithMaterialIDElement;
 class IDatasmithUEPbrMaterialElement;
 
 struct FMeshDescription;
-struct FDatasmithTessellationOptions;
 
 namespace CADLibrary
 {
@@ -47,48 +47,29 @@ namespace UE_DATASMITHWIRETRANSLATOR_NAMESPACE
 		/** Begin IWrieInterface */
 		virtual bool Initialize(const TCHAR* InSceneFullName) override;
 		virtual bool Load(TSharedPtr<IDatasmithScene> InScene) override;
-		virtual void SetTessellationOptions(const FDatasmithTessellationOptions& Options) override;
+		virtual void SetImportSettings(const FWireSettings& Options) override;
 		virtual void SetOutputPath(const FString& Path) { OutputPath = Path; }
 
 		bool LoadStaticMesh(const TSharedPtr<IDatasmithMeshElement> MeshElement, FDatasmithMeshElementPayload& OutMeshPayload, const FDatasmithTessellationOptions& InTessellationOptions);
 		/** End IWrieInterface */
 
 	private:
-		TOptional<FMeshDescription> GetMeshDescription(TSharedPtr<IDatasmithMeshElement> MeshElement, CADLibrary::FMeshParameters& OutMeshParameters);
-		TOptional<FMeshDescription> GetMeshDescription(TSharedPtr<IDatasmithMeshElement> MeshElement, CADLibrary::FMeshParameters& MeshParameters, TSharedPtr<FBodyData> BodyTemp);
-
-
-	private:
-
-		bool GetDagLeaves();
-		bool GetShader();
-		void RecurseDagForLeaves(const TAlObjectPtr<AlDagNode>& DagNode, FDagNodeInfo& ParentInfo);
-		void RecurseDagForLeavesNoMerge(const TAlObjectPtr<AlDagNode>& DagNode, FDagNodeInfo& ParentInfo);
-		void DagForLeavesNoMerge(const TAlObjectPtr<AlDagNode>& DagNode, FDagNodeInfo& ParentInfo);
-		bool ProcessAlGroupNode(AlDagNode& GroupNode, FDagNodeInfo& ParentInfo);
-		void ProcessAlShellNode(const TAlObjectPtr<AlDagNode>& ShellNode, FDagNodeInfo& ParentInfo, const FString& ShaderName);
-		void ProcessBodyNode(const TSharedPtr<FBodyData>& Body, FDagNodeInfo& ParentInfo);
-		TSharedPtr<IDatasmithMeshElement> FindOrAddMeshElement(const TSharedPtr<FBodyData>& Body, const FDagNodeInfo& ParentInfo);
-		TSharedPtr<IDatasmithMeshElement> FindOrAddMeshElement(const TAlObjectPtr<AlDagNode>& ShellNode, const FDagNodeInfo& ParentInfo, const FString& ShaderName);
-		TSharedPtr<IDatasmithActorElement> FindOrAddParentActor(FDagNodeInfo& ParentInfo, const FString& LayerName);
-		TSharedPtr<IDatasmithActorElement> FindOrAddLayerActor(const FString& LayerName);
-
-		void RecurseDeleteEmptyActor(TSharedPtr<IDatasmithActorElement> Actor);
-
-		TOptional<FMeshDescription> GetMeshOfShellNode(AlDagNode& DagNode, TSharedPtr<IDatasmithMeshElement> MeshElement, CADLibrary::FMeshParameters& MeshParameters);
-		TOptional<FMeshDescription> GetMeshOfNodeMesh(AlDagNode& DagNode, TSharedPtr<IDatasmithMeshElement> MeshElement, CADLibrary::FMeshParameters& MeshParameters, AlMatrix4x4* AlMeshInvGlobalMatrix = nullptr);
-		TOptional<FMeshDescription> GetMeshOfShellBody(TSharedPtr<FBodyData> DagNode, TSharedPtr<IDatasmithMeshElement> MeshElement, CADLibrary::FMeshParameters& MeshParameters);
-		TOptional<FMeshDescription> GetMeshOfMeshBody(TSharedPtr<FBodyData> DagNode, TSharedPtr<IDatasmithMeshElement> MeshElement, CADLibrary::FMeshParameters& MeshParameters);
-
-		void AddNodeInBodyGroup(TAlObjectPtr<AlDagNode>& DagNode, const FString& ShaderName, TMap<uint32, TSharedPtr<FBodyData>>& ShellToProcess, bool bIsAPatch, uint32 MaxSize);
-
-		TSharedPtr<CADLibrary::ICADModelConverter> GetModelConverter(IAliasBRepConverter*& BRepConverter) const;
-
-		TOptional<FMeshDescription> MeshDagNodeWithExternalMesher(AlDagNode& DagNode, TSharedPtr<IDatasmithMeshElement> MeshElement, CADLibrary::FMeshParameters& MeshParameters);
-		TOptional<FMeshDescription> MeshDagNodeWithExternalMesher(TSharedPtr<FBodyData> DagNode, TSharedPtr<IDatasmithMeshElement> MeshElement, CADLibrary::FMeshParameters& MeshParameters);
-
-		TOptional<FMeshDescription> ImportMesh(AlMesh& Mesh, TSharedPtr<IDatasmithMeshElement> MeshElement, CADLibrary::FMeshParameters& MeshParameters);
-
+		/** Model traversal */
+		bool TraverseModel();
+		TSharedPtr<IDatasmithActorElement> TraverseDag(const TAlDagNodePtr<AlDagNode>& DagNode);
+
+		TSharedPtr<IDatasmithActorElement> ProcessGeometryNode(const TAlDagNodePtr<AlDagNode>& GeomNode, const TAlObjectPtr<AlLayer>& ParentLayer = TAlObjectPtr<AlLayer>());
+		TSharedPtr<IDatasmithActorElement> TraverseGroupNode(const TAlDagNodePtr<AlGroupNode>& GroupNode, const TAlObjectPtr<AlLayer>& ParentLayer = TAlObjectPtr<AlLayer>());
+		TSharedPtr<IDatasmithActorElement> ProcessGroupNode(const TAlDagNodePtr<AlGroupNode>& GroupNode, const TAlObjectPtr<AlLayer>& ParentLayer = TAlObjectPtr<AlLayer>());
+		TSharedPtr<IDatasmithActorElement> ProcessBodyNode(TSharedPtr<FBodyNode>& BodyNode, const TAlDagNodePtr<AlGroupNode>& GroupNode, const TAlObjectPtr<AlLayer>& ParentLayer = TAlObjectPtr<AlLayer>());
+		TSharedPtr<IDatasmithActorElement> ProcessPatchMesh(TSharedPtr<FPatchMesh>& PatchMesh, const TAlDagNodePtr<AlGroupNode>& GroupNode, const TAlObjectPtr<AlLayer>& ParentLayer = TAlObjectPtr<AlLayer>());
+
+		TSharedPtr<IDatasmithActorElement> FindOrAddLayerActor(const TAlObjectPtr<AlLayer>& Layer);
+		TSharedPtr<IDatasmithMeshElement> FindOrAddMeshElement(const TAlDagNodePtr<AlDagNode>& GeomNode);
+		TSharedPtr<IDatasmithMeshElement> FindOrAddMeshElement(TSharedPtr<FBodyNode>& BodyNode);
+		TSharedPtr<IDatasmithMeshElement> FindOrAddMeshElement(TSharedPtr<FPatchMesh>& PatchMesh);
+
+		/** Material creation */
 		bool IsTransparent(FColor& TransparencyColor)
 		{
 			float Opacity = 1.0f - ((float)(TransparencyColor.R + TransparencyColor.G + TransparencyColor.B)) / 765.0f;
@@ -101,41 +82,42 @@ namespace UE_DATASMITHWIRETRANSLATOR_NAMESPACE
 		void AddAlLambertParameters(const TAlObjectPtr<AlShader>& Shader, TSharedPtr<IDatasmithUEPbrMaterialElement> MaterialElement);
 		void AddAlLightSourceParameters(const TAlObjectPtr<AlShader>& Shader, TSharedPtr<IDatasmithUEPbrMaterialElement> MaterialElement);
 		void AddAlPhongParameters(const TAlObjectPtr<AlShader>& Shader, TSharedPtr<IDatasmithUEPbrMaterialElement> MaterialElement);
+		
+		TSharedPtr<IDatasmithMaterialIDElement> FindOrAddMaterial(const TAlObjectPtr<AlShader>& Shader);
+
+		/** Geometry retrieval */
+		TOptional<FMeshDescription> GetMeshDescription(TSharedPtr<IDatasmithMeshElement> MeshElement, CADLibrary::FMeshParameters& OutMeshParameters);
+		TSharedPtr<CADLibrary::ICADModelConverter> GetModelConverter(IAliasBRepConverter*& BRepConverter) const;
+		TOptional<FMeshDescription> TessellateParametricNode(AlDagNode& DagNode, TSharedPtr<IDatasmithMeshElement> MeshElement, CADLibrary::FMeshParameters& MeshParameters);
+		TOptional<FMeshDescription> GetMeshDescriptionFromBodyNode(TSharedPtr<FBodyNode>& BodyNode, TSharedPtr<IDatasmithMeshElement> MeshElement, CADLibrary::FMeshParameters& MeshParameters);
+		TOptional<FMeshDescription> GetMeshDescriptionFromPatchMesh(TSharedPtr<FPatchMesh>& PatchMesh, TSharedPtr<IDatasmithMeshElement> MeshElement, CADLibrary::FMeshParameters& MeshParameters);
+		TOptional<FMeshDescription> GetMeshDescriptionFromParametricNode(AlDagNode& DagNode, TSharedPtr<IDatasmithMeshElement> MeshElement, CADLibrary::FMeshParameters& MeshParameters);
+		TOptional<FMeshDescription> GetMeshDescriptionFromMeshNode(const TAlDagNodePtr<AlMeshNode>& MeshNode, TSharedPtr<IDatasmithMeshElement> MeshElement, CADLibrary::FMeshParameters& MeshParameters, AlMatrix4x4* AlMeshInvGlobalMatrix = nullptr);
 
 	private:
 		TSharedPtr<IDatasmithScene> DatasmithScene;
-		FString SceneName;
-		FString CurrentPath;
 		FString OutputPath;
 		FString SceneFullPath;
 
-		// Hash value of the scene file used to check if the file has been modified for re-import
-		uint32 SceneFileHash = 0;
-
-		TAlObjectPtr<AlDagNode> AlRootNode;
-
-		/** Table of correspondence between mesh identifier and associated Datasmith mesh element */
-		TMap<uint32, TSharedPtr<IDatasmithMeshElement>> ShellUuidToMeshElementMap;
-		TMap<uint32, TSharedPtr<IDatasmithMeshElement>> BodyUuidToMeshElementMap;
+		FWireSettings WireSettings;
 
-		/** Datasmith mesh elements to OpenModel objects */
-		TMap<IDatasmithMeshElement*, TAlObjectPtr<AlDagNode>> MeshElementToAlDagNodeMap;
+		TSharedPtr<CADLibrary::ICADModelConverter> CADModelConverter;
+		TSharedPtr<IAliasBRepConverter> AliasBRepConverter;
 
-		TMap<IDatasmithMeshElement*, TSharedPtr<FBodyData>> MeshElementToBodyMap;
+		bool bSceneLoaded = false;
 
-		TMap<IDatasmithMeshElement*, FString> MeshElementToShaderName;
-		TMap<FString, FColor> ShaderNameToColor;
-		TMap<FString, TSharedPtr<IDatasmithMaterialIDElement>> ShaderNameToUEMaterialId;
+		TMap<FString, TSharedPtr<IDatasmithBaseMaterialElement>> ShaderNameToMaterial;
 
-		TMap<FString, TSharedPtr<IDatasmithActorElement>> LayerNameToActor;
+		TMap<uint32, TSharedPtr<IDatasmithMeshElement>> GeomNodeToMeshElement;
+		TMap<TSharedPtr<IDatasmithMeshElement>, AlDagNode*> MeshElementToGeomNode;
 
-		FDatasmithTessellationOptions TessellationOptions;
+		TMap<uint32, TSharedPtr<IDatasmithMeshElement>> BodyNodeToMeshElement;
+		TMap<TSharedPtr<IDatasmithMeshElement>, TSharedPtr<FBodyNode>> MeshElementToBodyNode;
 
-		TSharedPtr<CADLibrary::ICADModelConverter> CADModelConverter;
-		TSharedPtr<IAliasBRepConverter> AliasBRepConverter;
+		TMap<uint32, TSharedPtr<IDatasmithMeshElement>> PatchMeshToMeshElement;
+		TMap<TSharedPtr<IDatasmithMeshElement>, TSharedPtr<FPatchMesh>> MeshElementToPatchMesh;
 
-		bool bAliasUseNative = false;
-		bool bSceneLoaded = false;
+		TMap<uint32, TSharedPtr<IDatasmithActorElement>> LayerToActor;
 	};
 } // namespace
 #endif
