Code refactor - take #2
- Simplified ICADCOnverter interface
- Better sorting of geometry to tessellate
- Tentatively fixed the transform issue. Still potential issues if geometry is symmetrical and groupe node has transform.

#rnx

diff --git a/Engine/Plugins/Enterprise/DatasmithCADImporter/Source/CADKernelSurface/Public/CADModelToCADKernelConverterBase.h b/Engine/Plugins/Enterprise/DatasmithCADImporter/Source/CADKernelSurface/Public/CADModelToCADKernelConverterBase.h
--- a/Engine/Plugins/Enterprise/DatasmithCADImporter/Source/CADKernelSurface/Public/CADModelToCADKernelConverterBase.h
+++ b/Engine/Plugins/Enterprise/DatasmithCADImporter/Source/CADKernelSurface/Public/CADModelToCADKernelConverterBase.h
@@ -86,6 +86,8 @@ public:
 		return true;
 	}
 
+	virtual bool AddGeometry(const CADLibrary::FCADModelGeometry& Geometry) override { return false; }
+
 	virtual void AddSurfaceDataForMesh(const TCHAR* InFilePath, const CADLibrary::FMeshParameters& InMeshParameters, const FDatasmithTessellationOptions& InTessellationOptions, FDatasmithMeshElementPayload& OutMeshPayload) const override
 	{
 		CADKernelSurface::AddSurfaceDataForMesh(InFilePath, ImportParameters, InMeshParameters, InTessellationOptions, OutMeshPayload);
diff --git a/Engine/Plugins/Enterprise/DatasmithCADImporter/Source/DatasmithWireTranslator/Public/IWireInterface.h b/Engine/Plugins/Enterprise/DatasmithCADImporter/Source/DatasmithWireTranslator/Public/IWireInterface.h
--- a/Engine/Plugins/Enterprise/DatasmithCADImporter/Source/DatasmithWireTranslator/Public/IWireInterface.h
+++ b/Engine/Plugins/Enterprise/DatasmithCADImporter/Source/DatasmithWireTranslator/Public/IWireInterface.h
@@ -27,7 +27,7 @@ struct FWireSettings : public FDatasmithTessellationOptions
 	bool bUseLayerAsActor = true;
 
 	UPROPERTY(config, EditAnywhere, BlueprintReadWrite, Category = "Wire Translation Options", meta = (ToolTip = "If set to true, all geometry nodes under a group are merged. Default is true."))
-	bool bMergeGeomtryByGroup = true;
+	bool bMergeGeometryByGroup = true;
 
 	uint32 GetHash() const
 	{
diff --git a/Engine/Plugins/Enterprise/DatasmithCADImporter/Source/ParametricSurface/Public/CADModelConverter.h b/Engine/Plugins/Enterprise/DatasmithCADImporter/Source/ParametricSurface/Public/CADModelConverter.h
--- a/Engine/Plugins/Enterprise/DatasmithCADImporter/Source/ParametricSurface/Public/CADModelConverter.h
+++ b/Engine/Plugins/Enterprise/DatasmithCADImporter/Source/ParametricSurface/Public/CADModelConverter.h
@@ -31,6 +31,10 @@ struct FMeshParameters;
  *
  * In a second step, the saved file is associated to its UStaticMesh with AddSurfaceDataForMesh function
  */
+struct FCADModelGeometry
+{
+	int32 Type = -1;
+};
 
 class ICADModelConverter
 {
@@ -54,6 +58,7 @@ public:
 	 */
 	virtual void SetImportParameters(double ChordTolerance, double MaxEdgeLength, double NormalTolerance, CADLibrary::EStitchingTechnique StitchingTechnique) = 0;
 
+	virtual bool AddGeometry(const FCADModelGeometry& Geometry) = 0;
 	virtual bool IsSessionValid() = 0;
 
 	virtual void AddSurfaceDataForMesh(const TCHAR* InFilePath, const FMeshParameters& InMeshParameters, const FDatasmithTessellationOptions& InTessellationOptions, FDatasmithMeshElementPayload& OutMeshPayload) const = 0;
diff --git a/Engine/Plugins/Enterprise/DatasmithCADImporter/Source/ParametricSurface/Public/CADModelToTechSoftConverterBase.h b/Engine/Plugins/Enterprise/DatasmithCADImporter/Source/ParametricSurface/Public/CADModelToTechSoftConverterBase.h
--- a/Engine/Plugins/Enterprise/DatasmithCADImporter/Source/ParametricSurface/Public/CADModelToTechSoftConverterBase.h
+++ b/Engine/Plugins/Enterprise/DatasmithCADImporter/Source/ParametricSurface/Public/CADModelToTechSoftConverterBase.h
@@ -47,6 +47,8 @@ public:
 		return true;
 	}
 
+	virtual bool AddGeometry(const CADLibrary::FCADModelGeometry& Geometry) override { return false; }
+
 	virtual void AddSurfaceDataForMesh(const TCHAR* InFilePath, const CADLibrary::FMeshParameters& InMeshParameters, const FDatasmithTessellationOptions& InTessellationOptions, FDatasmithMeshElementPayload& OutMeshPayload) const override;
 
 protected:
diff --git a/Engine/Plugins/Enterprise/DatasmithCADImporter/Source/WireInterface/Private/AliasBrepConverter.h b/Engine/Plugins/Enterprise/DatasmithCADImporter/Source/WireInterface/Private/AliasBrepConverter.h
deleted file mode 100644
diff --git a/Engine/Plugins/Enterprise/DatasmithCADImporter/Source/WireInterface/Private/AliasModelToCADKernelConverter.cpp b/Engine/Plugins/Enterprise/DatasmithCADImporter/Source/WireInterface/Private/AliasModelToCADKernelConverter.cpp
--- a/Engine/Plugins/Enterprise/DatasmithCADImporter/Source/WireInterface/Private/AliasModelToCADKernelConverter.cpp
+++ b/Engine/Plugins/Enterprise/DatasmithCADImporter/Source/WireInterface/Private/AliasModelToCADKernelConverter.cpp
@@ -402,6 +402,36 @@ bool FAliasModelToCADKernelConverter::RepairTopology()
 	return true;
 }
 
+bool FAliasModelToCADKernelConverter::AddGeometry(const CADLibrary::FCADModelGeometry& Geometry)
+{
+	if (Geometry.Type == (int32)ECADModelGeometryType::DagNode)
+	{
+		const FDagNodeGeometry& DagNodeGeometry = static_cast<const FDagNodeGeometry&>(Geometry);
+
+		return AddBRep(*DagNodeGeometry.DagNode, 0, DagNodeGeometry.Reference);
+	}
+	else if (Geometry.Type == (int32)ECADModelGeometryType::BodyNode)
+	{
+		const FBodyNodeGeometry& BodyNodeGeometry = static_cast<const FBodyNodeGeometry&>(Geometry);
+
+		bool bBRepAdded = true;
+		BodyNodeGeometry.BodyNode->IterateOnSurfaceNodes([&](const TAlDagNodePtr<AlSurfaceNode>& SurfaceNode)
+		{
+			bBRepAdded &= AddBRep(*SurfaceNode, BodyNodeGeometry.BodyNode->GetSlotIndex(SurfaceNode.Get()), BodyNodeGeometry.Reference);
+		});
+
+		BodyNodeGeometry.BodyNode->IterateOnShellNodes([&](const TAlDagNodePtr<AlShellNode>& ShellNode)
+		{
+			bBRepAdded &= AddBRep(*ShellNode, BodyNodeGeometry.BodyNode->GetSlotIndex(ShellNode.Get()), BodyNodeGeometry.Reference);
+		});
+
+		ensure(bBRepAdded);
+		return bBRepAdded;
+	}
+
+	return false;
+}
+
 }
 
 #endif
\ No newline at end of file
diff --git a/Engine/Plugins/Enterprise/DatasmithCADImporter/Source/WireInterface/Private/AliasModelToCADKernelConverter.h b/Engine/Plugins/Enterprise/DatasmithCADImporter/Source/WireInterface/Private/AliasModelToCADKernelConverter.h
--- a/Engine/Plugins/Enterprise/DatasmithCADImporter/Source/WireInterface/Private/AliasModelToCADKernelConverter.h
+++ b/Engine/Plugins/Enterprise/DatasmithCADImporter/Source/WireInterface/Private/AliasModelToCADKernelConverter.h
@@ -5,7 +5,6 @@
 #include "CoreMinimal.h"
 
 #ifdef USE_OPENMODEL
-#include "AliasBRepConverter.h"
 #include "OpenModelUtils.h"
 
 #include "CADModelToCADKernelConverterBase.h"
@@ -32,7 +31,7 @@ namespace UE::CADKernel
 namespace UE_DATASMITHWIRETRANSLATOR_NAMESPACE
 {
 
-class FAliasModelToCADKernelConverter : public FCADModelToCADKernelConverterBase, public IAliasBRepConverter
+class FAliasModelToCADKernelConverter : public FCADModelToCADKernelConverterBase
 {
 
 public:
@@ -44,10 +43,12 @@ public:
 	// Begin FCADModelToCADKernelConverterBase overrides
 	virtual bool Tessellate(const CADLibrary::FMeshParameters& InMeshParameters, FMeshDescription& OutMeshDescription) override;
 	virtual bool RepairTopology() override;
+
+	virtual bool AddGeometry(const CADLibrary::FCADModelGeometry& Geometry) override;
 	// End FCADModelToCADKernelConverterBase overrides
 
-	virtual bool AddBRep(AlDagNode& DagNode, const FColor& Color, EAliasObjectReference ObjectReference) override;
-	virtual bool AddBRep(AlDagNode& DagNode, uint32 SlotID, EAliasObjectReference ObjectReference) override;
+	bool AddBRep(AlDagNode& DagNode, const FColor& Color, EAliasObjectReference ObjectReference);
+	bool AddBRep(AlDagNode& DagNode, uint32 SlotID, EAliasObjectReference ObjectReference);
 
 protected:
 	TSharedPtr<UE::CADKernel::FTopologicalEdge> AddEdge(const AlTrimCurve& TrimCurve, TSharedPtr<UE::CADKernel::FSurface>& CarrierSurface);
diff --git a/Engine/Plugins/Enterprise/DatasmithCADImporter/Source/WireInterface/Private/AliasModelToTechSoftConverter.cpp b/Engine/Plugins/Enterprise/DatasmithCADImporter/Source/WireInterface/Private/AliasModelToTechSoftConverter.cpp
--- a/Engine/Plugins/Enterprise/DatasmithCADImporter/Source/WireInterface/Private/AliasModelToTechSoftConverter.cpp
+++ b/Engine/Plugins/Enterprise/DatasmithCADImporter/Source/WireInterface/Private/AliasModelToTechSoftConverter.cpp
@@ -478,5 +478,35 @@ bool FAliasModelToTechSoftConverter::AddBRep(AlDagNode& DagNode, const FColor& C
 	return false;
 }
 
+bool FAliasModelToTechSoftConverter::AddGeometry(const CADLibrary::FCADModelGeometry& Geometry)
+{
+	if (Geometry.Type == (int32)ECADModelGeometryType::DagNode)
+	{
+		const FDagNodeGeometry& DagNodeGeometry = static_cast<const FDagNodeGeometry&>(Geometry);
+
+		return AddBRep(*DagNodeGeometry.DagNode, 0, DagNodeGeometry.Reference);
+	}
+	else if (Geometry.Type == (int32)ECADModelGeometryType::BodyNode)
+	{
+		const FBodyNodeGeometry& BodyNodeGeometry = static_cast<const FBodyNodeGeometry&>(Geometry);
+
+		bool bBRepAdded = true;
+		BodyNodeGeometry.BodyNode->IterateOnSurfaceNodes([&](const TAlDagNodePtr<AlSurfaceNode>& SurfaceNode)
+			{
+				bBRepAdded &= AddBRep(*SurfaceNode, BodyNodeGeometry.BodyNode->GetSlotIndex(SurfaceNode.Get()), BodyNodeGeometry.Reference);
+			});
+
+		BodyNodeGeometry.BodyNode->IterateOnShellNodes([&](const TAlDagNodePtr<AlShellNode>& ShellNode)
+			{
+				bBRepAdded &= AddBRep(*ShellNode, BodyNodeGeometry.BodyNode->GetSlotIndex(ShellNode.Get()), BodyNodeGeometry.Reference);
+			});
+
+		ensure(bBRepAdded);
+		return bBRepAdded;
+	}
+
+	return false;
+}
+
 }
 #endif
diff --git a/Engine/Plugins/Enterprise/DatasmithCADImporter/Source/WireInterface/Private/AliasModelToTechSoftConverter.h b/Engine/Plugins/Enterprise/DatasmithCADImporter/Source/WireInterface/Private/AliasModelToTechSoftConverter.h
--- a/Engine/Plugins/Enterprise/DatasmithCADImporter/Source/WireInterface/Private/AliasModelToTechSoftConverter.h
+++ b/Engine/Plugins/Enterprise/DatasmithCADImporter/Source/WireInterface/Private/AliasModelToTechSoftConverter.h
@@ -5,7 +5,6 @@
 #include "CoreMinimal.h"
 
 #ifdef USE_OPENMODEL
-#include "AliasBRepConverter.h"
 #include "OpenModelUtils.h"
 
 #include "CADModelToTechSoftConverterBase.h"
@@ -29,7 +28,7 @@ typedef void A3DCrvBase;
 namespace UE_DATASMITHWIRETRANSLATOR_NAMESPACE
 {
 
-class FAliasModelToTechSoftConverter : public FCADModelToTechSoftConverterBase, public IAliasBRepConverter
+class FAliasModelToTechSoftConverter : public FCADModelToTechSoftConverterBase
 {
 
 public:
@@ -38,8 +37,10 @@ public:
 	{
 	}
 
-	virtual bool AddBRep(AlDagNode& DagNode, const FColor& Color, EAliasObjectReference ObjectReference) override;
-	virtual bool AddBRep(AlDagNode& DagNode, uint32 SlotID, EAliasObjectReference ObjectReference) override;
+	virtual bool AddGeometry(const CADLibrary::FCADModelGeometry& Geometry) override;
+
+	bool AddBRep(AlDagNode& DagNode, const FColor& Color, EAliasObjectReference ObjectReference);
+	bool AddBRep(AlDagNode& DagNode, uint32 SlotID, EAliasObjectReference ObjectReference);
 
 protected:
 	TMap<void*, A3DTopoCoEdge*> AlEdgeToTSCoEdge;
diff --git a/Engine/Plugins/Enterprise/DatasmithCADImporter/Source/WireInterface/Private/OpenModelUtils.cpp b/Engine/Plugins/Enterprise/DatasmithCADImporter/Source/WireInterface/Private/OpenModelUtils.cpp
--- a/Engine/Plugins/Enterprise/DatasmithCADImporter/Source/WireInterface/Private/OpenModelUtils.cpp
+++ b/Engine/Plugins/Enterprise/DatasmithCADImporter/Source/WireInterface/Private/OpenModelUtils.cpp
@@ -2,8 +2,6 @@
 #include "OpenModelUtils.h"
 
 #ifdef USE_OPENMODEL
-#include "AliasBRepConverter.h"
-
 #include "CADOptions.h"
 #include "DatasmithUtils.h"
 #include "DatasmithTranslator.h"
@@ -72,7 +70,7 @@ namespace UE_DATASMITHWIRETRANSLATOR_NAMESPACE
 
 	int32 FBodyNode::GetSlotIndex(const TAlDagNodePtr<AlDagNode>& DagNode)
 	{
-		// TODO: Add support for AlShell
+		// #wire_import: Add support for AlShell
 		if (TAlDagNodePtr<AlSurfaceNode> SurfaceNode = DagNode->asSurfaceNodePtr())
 		{
 			if (TAlObjectPtr<AlSurface> Surface = SurfaceNode->surface())
@@ -125,7 +123,7 @@ namespace UE_DATASMITHWIRETRANSLATOR_NAMESPACE
 
 			for (const TAlDagNodePtr<AlShellNode>& ShellNode : ShellNodes)
 			{
-				// TODO: Extract all shaders
+				// #wire_import: Extract all shaders
 				TAlObjectPtr<AlShell> Shell(ShellNode->shell());
 				if (!Shell)
 				{
@@ -223,17 +221,16 @@ namespace UE_DATASMITHWIRETRANSLATOR_NAMESPACE
 		return ActorElement->IsA(EDatasmithElementType::StaticMeshActor) || ActorElement->GetChildrenCount() > 0;
 	}
 
-	void OpenModelUtils::SetActorTransform(IDatasmithActorElement& ActorElement, const AlDagNode& InDagNode)
+	void OpenModelUtils::SetActorTransform(IDatasmithActorElement& ActorElement, const TAlDagNodePtr<AlDagNode>& DagNode)
 	{
-		// TODO: Find why no transform applied if layer is symmetrical???
-		TAlObjectPtr<AlLayer> Layer(InDagNode.layer());
-		if (Layer && Layer->isSymmetric())
+		// #wire_import: Find why no transform applied if layer is symmetrical???
+		if (DagNode.HasSymmetry())
 		{
 			return;
 		}
 
 		AlMatrix4x4 AlGlobalMatrix;
-		InDagNode.globalTransformationMatrix(AlGlobalMatrix);
+		DagNode->globalTransformationMatrix(AlGlobalMatrix);
 
 		FMatrix GlobalMatrix;
 		double* MatrixFloats = (double*)GlobalMatrix.M;
diff --git a/Engine/Plugins/Enterprise/DatasmithCADImporter/Source/WireInterface/Private/OpenModelUtils.h b/Engine/Plugins/Enterprise/DatasmithCADImporter/Source/WireInterface/Private/OpenModelUtils.h
--- a/Engine/Plugins/Enterprise/DatasmithCADImporter/Source/WireInterface/Private/OpenModelUtils.h
+++ b/Engine/Plugins/Enterprise/DatasmithCADImporter/Source/WireInterface/Private/OpenModelUtils.h
@@ -33,7 +33,7 @@ struct FMeshDescription;
 #define GROUPNODE_TYPE      TEXT("GroupNode")
 #define MESH_TYPE           TEXT("Mesh")
 #define MESHNODE_TYPE       TEXT("MeshNode")
-#define SHADER_TYPE         TEXT("Shell")
+#define SHADER_TYPE         TEXT("Shader")
 #define SHELLNODE_TYPE      TEXT("ShellNode")
 #define SHELL_TYPE          TEXT("Shell")
 #define SURFACE_TYPE        TEXT("Surface")
@@ -41,10 +41,7 @@ struct FMeshDescription;
 
 namespace UE_DATASMITHWIRETRANSLATOR_NAMESPACE
 {
-	namespace OpenModelUtils
-	{
-
-	}
+	typedef double AlMatrix4x4[4][4];
 
 	enum class ETesselatorType : uint8
 	{
@@ -153,6 +150,13 @@ namespace UE_DATASMITHWIRETRANSLATOR_NAMESPACE
 			return TAlObjectPtr<AlLayer>(LocalPtr && LocalPtr->layer() ? LocalPtr->layer() : nullptr);
 		}
 
+		bool HasSymmetry() const
+		{
+			static_assert(std::is_base_of<AlDagNode, T>::value);
+			T* LocalPtr = TAlObjectPtr<T>::Ptr;
+			return (AlIsValid(LocalPtr) && LocalPtr->layer()) ? (bool)LocalPtr->layer()->isSymmetric() : false;
+		}
+
 		TAlObjectPtr<AlMesh> GetMesh() const
 		{
 			static_assert(std::is_base_of<AlDagNode, T>::value);
@@ -332,6 +336,47 @@ namespace UE_DATASMITHWIRETRANSLATOR_NAMESPACE
 		bool bInitialized = false;
 	};
 
+	enum class ECADModelGeometryType : int32
+	{
+		DagNode,
+		MeshNode,
+		BodyNode,
+		PatchMesh,
+	};
+
+	enum class EAliasObjectReference
+	{
+		LocalReference,  
+		ParentReference, 
+		WorldReference, 
+	};
+	
+	struct FAliasGeometry : public CADLibrary::FCADModelGeometry
+	{
+		EAliasObjectReference Reference = EAliasObjectReference::LocalReference;
+	};
+
+	struct FDagNodeGeometry : public FAliasGeometry
+	{
+		TAlDagNodePtr<AlDagNode> DagNode;
+	};
+
+	struct FBodyNodeGeometry : public FAliasGeometry
+	{
+		TSharedPtr<FBodyNode> BodyNode;
+	};
+
+	struct FPatchMeshGeometry : public CADLibrary::FCADModelGeometry
+	{
+		TSharedPtr<FPatchMesh> PatchMesh;
+	};
+
+	struct FMeshNodeGeometry : public CADLibrary::FCADModelGeometry
+	{
+		TAlDagNodePtr<AlMeshNode> MeshNode;
+	};
+
+
 	namespace OpenModelUtils
 	{
 		/** Following layer hierarchy, get list of layers an actor would be in as a csv string*/
@@ -339,7 +384,7 @@ namespace UE_DATASMITHWIRETRANSLATOR_NAMESPACE
 
 		bool ActorHasContent(const TSharedPtr<IDatasmithActorElement>& ActorElement);
 
-		void SetActorTransform(IDatasmithActorElement& OutActorElement, const AlDagNode& InDagNode);
+		void SetActorTransform(IDatasmithActorElement& OutActorElement, const TAlDagNodePtr<AlDagNode>& InDagNode);
 
 		bool IsValidActor(const TSharedPtr<IDatasmithActorElement>& ActorElement);
 
diff --git a/Engine/Plugins/Enterprise/DatasmithCADImporter/Source/WireInterface/Private/WireInterfaceImpl.cpp b/Engine/Plugins/Enterprise/DatasmithCADImporter/Source/WireInterface/Private/WireInterfaceImpl.cpp
--- a/Engine/Plugins/Enterprise/DatasmithCADImporter/Source/WireInterface/Private/WireInterfaceImpl.cpp
+++ b/Engine/Plugins/Enterprise/DatasmithCADImporter/Source/WireInterface/Private/WireInterfaceImpl.cpp
@@ -187,14 +187,14 @@ namespace UE_DATASMITHWIRETRANSLATOR_NAMESPACE
 
 	bool FWireTranslatorImpl::Initialize(const TCHAR* InSceneFullName)
 	{
-		// TODO: Check file can be loaded with current SDK
+		// #wire_import: Check file can be loaded with current SDK
 
 		if (InSceneFullName)
 		{
 			SceneFullPath = InSceneFullName;
 		}
 
-		// TODO: Check whether this implementation can load it
+		// #wire_import: Check whether this implementation can load it
 		return true;
 	}
 
@@ -234,17 +234,18 @@ namespace UE_DATASMITHWIRETRANSLATOR_NAMESPACE
 			CADLibrary::FImportParameters ImportParameters;
 			if (CADLibrary::FImportParameters::bGDisableCADKernelTessellation)
 			{
-				TSharedPtr<FAliasModelToTechSoftConverter> AliasToTechSoftConverter = MakeShared<FAliasModelToTechSoftConverter>(ImportParameters);
-				CADModelConverter = AliasToTechSoftConverter;
-				AliasBRepConverter = AliasToTechSoftConverter;
+				CADModelConverter = MakeShared<FAliasModelToTechSoftConverter>(ImportParameters);
 			}
 			else
 			{
-				TSharedPtr<FAliasModelToCADKernelConverter> AliasToCADKernelConverter = MakeShared<FAliasModelToCADKernelConverter>(ImportParameters);
-				CADModelConverter = AliasToCADKernelConverter;
-				AliasBRepConverter = AliasToCADKernelConverter;
+				CADModelConverter = MakeShared<FAliasModelToCADKernelConverter>(ImportParameters);
 			}
 		}
+		else
+		{
+			// Merge by group when using Alias' tessellator
+			WireSettings.bMergeGeometryByGroup = false;
+		}
 
 		bGSewByMaterial = GetConsoleBoolValue(TEXT("ds.CADTranslator.Alias.SewByMaterial"), false);
 
@@ -281,7 +282,7 @@ namespace UE_DATASMITHWIRETRANSLATOR_NAMESPACE
 		TAlObjectPtr<AlSet> Set(AlUniverse::firstSet());
 		while (Set)
 		{
-			// TODO: Add an actor to represent the set.
+			// #wire_import: Add an actor to represent the set.
 			TAlObjectPtr<AlSetMember> SetMember(Set->firstMember());
 			while (SetMember)
 			{
@@ -309,7 +310,7 @@ namespace UE_DATASMITHWIRETRANSLATOR_NAMESPACE
 		TAlDagNodePtr<AlGroupNode> GroupNode(RootNode->asGroupNodePtr());
 		if (GroupNode)
 		{
-			return WireSettings.bMergeGeomtryByGroup ? ProcessGroupNode(GroupNode) : TraverseGroupNode(GroupNode);
+			return WireSettings.bMergeGeometryByGroup ? ProcessGroupNode(GroupNode) : TraverseGroupNode(GroupNode);
 		}
 		else if (OpenModelUtils::IsGeometryValid(RootNode))
 		{
@@ -381,7 +382,7 @@ namespace UE_DATASMITHWIRETRANSLATOR_NAMESPACE
 			ActorElement->SetLayer(*CsvLayerString);
 		}
 
-		OpenModelUtils::SetActorTransform(*ActorElement, *GroupNode);
+		OpenModelUtils::SetActorTransform(*ActorElement, GroupNode.Get());
 
 		for (TSharedPtr<IDatasmithActorElement>& ChildActor : ChildActors)
 		{
@@ -492,7 +493,7 @@ namespace UE_DATASMITHWIRETRANSLATOR_NAMESPACE
 		}
 
 		TSharedPtr<IDatasmithActorElement> ActorElement;
-		if (ChildActors.Num() == 1)
+		if (WireSettings.bMergeGeometryByGroup && ChildActors.Num() == 1)
 		{
 			ActorElement = ChildActors[0];
 		}
@@ -516,9 +517,9 @@ namespace UE_DATASMITHWIRETRANSLATOR_NAMESPACE
 			ActorElement->SetLayer(*CsvLayerString);
 		}
 
-		if (ChildActors.Num() > 1)
+		if (WireSettings.bMergeGeometryByGroup && ChildActors.Num() > 1)
 		{
-			OpenModelUtils::SetActorTransform(*ActorElement, *GroupNode);
+			OpenModelUtils::SetActorTransform(*ActorElement, GroupNode.Get());
 
 			for (TSharedPtr<IDatasmithActorElement>& ChildActor : ChildActors)
 			{
@@ -571,7 +572,7 @@ namespace UE_DATASMITHWIRETRANSLATOR_NAMESPACE
 			ActorElement->SetLayer(*CsvLayerString);
 		}
 
-		OpenModelUtils::SetActorTransform(*ActorElement, *GeomNode);
+		OpenModelUtils::SetActorTransform(*ActorElement, GeomNode);
 
 		TAlDagNodePtr<AlShellNode> ShellNode(GeomNode->asShellNodePtr());
 		if (ShellNode && ShellNode->shell())
@@ -660,7 +661,7 @@ namespace UE_DATASMITHWIRETRANSLATOR_NAMESPACE
 
 		if (!BodyNode->GetLayer()->isSymmetric())
 		{
-			OpenModelUtils::SetActorTransform(*ActorElement, *GroupNode);
+			OpenModelUtils::SetActorTransform(*ActorElement, GroupNode.Get());
 		}
 
 		if (WireSettings.bUseLayerAsActor && BodyNode->GetLayer() != ParentLayer)
@@ -711,7 +712,7 @@ namespace UE_DATASMITHWIRETRANSLATOR_NAMESPACE
 			ActorElement->SetLayer(*CsvLayerString);
 		}
 
-		OpenModelUtils::SetActorTransform(*ActorElement, *GroupNode);
+		OpenModelUtils::SetActorTransform(*ActorElement, GroupNode.Get());
 
 		if (WireSettings.bUseLayerAsActor && PatchMesh->GetLayer() != ParentLayer)
 		{
@@ -737,7 +738,7 @@ namespace UE_DATASMITHWIRETRANSLATOR_NAMESPACE
 
 		if (!OpenModelUtils::IsGeometryValid(GeomNode))
 		{
-			// TODO: Log an error
+			// #wire_import: Log an error
 			return TSharedPtr<IDatasmithMeshElement>();
 		}
 
@@ -769,14 +770,17 @@ namespace UE_DATASMITHWIRETRANSLATOR_NAMESPACE
 				ApplyMaterial(Shader, SlotIndex++);
 				Shader = TAlObjectPtr<AlShader>(Shell->nextShader(Shader.Get()));
 			}
-			// TODO: Check There are as many shaders as trim regions
+			// #wire_import: Check There are as many shaders as trim regions
+			MeshElementToParametricNode.Add(MeshElement, GeomNode.Get());
 		}
 		else
 		{
 			TAlDagNodePtr<AlSurfaceNode> SurfaceNode(GeomNode->asSurfaceNodePtr());
-			if (SurfaceNode && SurfaceNode->surface())
+			if (SurfaceNode && AlIsValid(SurfaceNode->surface()))
 			{
+				// #wire_import: Check for trim regions
 				ApplyMaterial(SurfaceNode->surface()->firstShader(), 0);
+				MeshElementToParametricNode.Add(MeshElement, GeomNode.Get());
 			}
 			else
 			{
@@ -784,13 +788,18 @@ namespace UE_DATASMITHWIRETRANSLATOR_NAMESPACE
 				if (MeshNode && MeshNode->mesh())
 				{
 					ApplyMaterial(MeshNode->mesh()->firstShader(), 0);
+					MeshElementToMeshNode.Add(MeshElement, MeshNode.Get());
+				}
+				else
+				{
+					// #wire_import: Log an error
+					return TSharedPtr<IDatasmithMeshElement>();
 				}
 			}
 		}
 
 		DatasmithScene->AddMesh(MeshElement);
 		GeomNodeToMeshElement.Add(GeomNode.GetHash(), MeshElement);
-		MeshElementToGeomNode.Add(MeshElement, GeomNode.Get());
 
 		return MeshElement;
 	}
@@ -805,7 +814,7 @@ namespace UE_DATASMITHWIRETRANSLATOR_NAMESPACE
 
 		if (!BodyNode->HasContent())
 		{
-			// TODO: Log an error
+			// #wire_import: Log an error
 			return TSharedPtr<IDatasmithMeshElement>();
 		}
 
@@ -843,7 +852,7 @@ namespace UE_DATASMITHWIRETRANSLATOR_NAMESPACE
 
 		if (!PatchMesh->HasContent())
 		{
-			// TODO: Log an error
+			// #wire_import: Log an error
 			return TSharedPtr<IDatasmithMeshElement>();
 		}
 
@@ -970,78 +979,86 @@ namespace UE_DATASMITHWIRETRANSLATOR_NAMESPACE
 		return OutMeshPayload.LodMeshes.Num() > 0;
 	}
 
-	TOptional<FMeshDescription> FWireTranslatorImpl::GetMeshDescription(TSharedPtr<IDatasmithMeshElement> MeshElement, CADLibrary::FMeshParameters& MeshParameters)
+	TOptional<FMeshDescription> FWireTranslatorImpl::GetMeshDescription(TSharedPtr<IDatasmithMeshElement> MeshElement, CADLibrary::FMeshParameters& OutMeshParameters)
 	{
-		if (TSharedPtr<FBodyNode>* BodyNodePtr = MeshElementToBodyNode.Find(MeshElement))
+		if (WireSettings.bAliasUseNative)
 		{
-			return GetMeshDescriptionFromBodyNode(*BodyNodePtr, MeshElement, MeshParameters);
+			if (AlDagNode** GeomNodePtr = MeshElementToParametricNode.Find(MeshElement))
+			{
+				TAlDagNodePtr<AlDagNode> GeomNode(*GeomNodePtr);
+
+				// #wire_import: the best way, should be to don't have to apply inverse global transform to the generated mesh
+				TAlDagNodePtr<AlMeshNode> MeshNode = OpenModelUtils::TesselateDagLeaf(*GeomNode, ETesselatorType::Fast, WireSettings.ChordTolerance);
+
+				if (MeshNode && MeshNode->mesh())
+				{
+					AlMatrix4x4 AlMatrix;
+					GeomNode->inverseGlobalTransformationMatrix(AlMatrix);
+					MeshNode->mesh()->transform(AlMatrix);
+
+					// Get the meshes from the dag nodes. Note that removing the mesh's DAG.
+					// will also removes the meshes, so we have to do it later.
+					return GetMeshDescriptionFromMeshNode(MeshNode, MeshElement, OutMeshParameters);
+				}
+			}
+			
+			if (AlMeshNode** MeshNodePtr = MeshElementToMeshNode.Find(MeshElement))
+			{
+				TAlDagNodePtr<AlMeshNode> MeshNode(*MeshNodePtr);
+				if (MeshNode)
+				{
+					return GetMeshDescriptionFromMeshNode(MeshNode, MeshElement, OutMeshParameters);
+				}
+			}
+
+			return TOptional<FMeshDescription>();
 		}
 
-		if (TSharedPtr<FPatchMesh>* PatchMeshPtr = MeshElementToPatchMesh.Find(MeshElement))
+		if (TSharedPtr<FBodyNode>* BodyNodePtr = MeshElementToBodyNode.Find(MeshElement))
 		{
-			return GetMeshDescriptionFromPatchMesh(*PatchMeshPtr, MeshElement, MeshParameters);
+			return GetMeshDescriptionFromBodyNode(*BodyNodePtr, MeshElement, OutMeshParameters);
 		}
 
-		if (AlDagNode** GeomNodePtr = MeshElementToGeomNode.Find(MeshElement))
+		if (AlDagNode** GeomNodePtr = MeshElementToParametricNode.Find(MeshElement))
 		{
-			TAlDagNodePtr<AlDagNode> GeomNode(*GeomNodePtr);
-			if (GeomNode)
-			{
-				MeshParameters = OpenModelUtils::GetMeshParameters(*GeomNode);
-				switch (GeomNode->type())
-				{
-				case kShellNodeType:
-				case kSurfaceNodeType:
-				{
-					return GetMeshDescriptionFromParametricNode(*GeomNode, MeshElement, MeshParameters);
-					break;
-				}
+			return GetMeshDescriptionFromParametricNode(*GeomNodePtr, MeshElement, OutMeshParameters);
+		}
 
-				case kMeshNodeType:
-				{
-					return GetMeshDescriptionFromMeshNode(GeomNode->asMeshNodePtr(), MeshElement, MeshParameters);
-					break;
-				}
+		if (TSharedPtr<FPatchMesh>* PatchMeshPtr = MeshElementToPatchMesh.Find(MeshElement))
+		{
+			return GetMeshDescriptionFromPatchMesh(*PatchMeshPtr, MeshElement, OutMeshParameters);
+		}
 
-				default:
-					break;
-				}
-			}
+		if (AlMeshNode** MeshNodePtr = MeshElementToMeshNode.Find(MeshElement))
+		{
+			return GetMeshDescriptionFromMeshNode(*MeshNodePtr, MeshElement, OutMeshParameters);
 		}
 
 		return TOptional<FMeshDescription>();
 	}
 
-	TOptional<FMeshDescription> FWireTranslatorImpl::GetMeshDescriptionFromBodyNode(TSharedPtr<FBodyNode>& BodyNode, TSharedPtr<IDatasmithMeshElement> MeshElement, CADLibrary::FMeshParameters& MeshParameters)
+	TOptional<FMeshDescription> FWireTranslatorImpl::GetMeshDescriptionFromBodyNode(TSharedPtr<FBodyNode>& BodyNode, TSharedPtr<IDatasmithMeshElement> MeshElement, CADLibrary::FMeshParameters& OutMeshParameters)
 	{
-		MeshParameters = OpenModelUtils::GetMeshParameters(BodyNode->GetLayer());
+		OutMeshParameters = OpenModelUtils::GetMeshParameters(BodyNode->GetLayer());
 
-		IAliasBRepConverter* BRepConverter = nullptr;
-		TSharedPtr<CADLibrary::ICADModelConverter> ModelConverter = GetModelConverter(BRepConverter);
+		TSharedPtr<CADLibrary::ICADModelConverter> ModelConverter = GetModelConverter();
 
 		ModelConverter->InitializeProcess();
 
 		EAliasObjectReference ObjectReference = EAliasObjectReference::LocalReference;
-		if (MeshParameters.bIsSymmetric)
+		if (OutMeshParameters.bIsSymmetric)
 		{
 			// All actors of a Alias symmetric layer are defined in the world Reference i.e. they have identity transform. So Mesh actor has to be defined in the world reference.
 			ObjectReference = EAliasObjectReference::WorldReference;
 		}
-		else if (WireSettings.StitchingTechnique != EDatasmithCADStitchingTechnique::StitchingNone)
+		else if (WireSettings.bMergeGeometryByGroup)
 		{
 			// In the case of StitchingSew, AlDagNode children of a GroupNode are merged together. To be merged, they have to be defined in the reference of parent GroupNode.
 			ObjectReference = EAliasObjectReference::ParentReference;
 		}
 
-		BodyNode->IterateOnSurfaceNodes([&](const TAlDagNodePtr<AlSurfaceNode>& SurfaceNode)
-			{
-				BRepConverter->AddBRep(*SurfaceNode, BodyNode->GetSlotIndex(SurfaceNode.Get()), ObjectReference);
-			});
-
-		BodyNode->IterateOnShellNodes([&](const TAlDagNodePtr<AlShellNode>& ShellNode)
-			{
-				BRepConverter->AddBRep(*ShellNode, BodyNode->GetSlotIndex(ShellNode.Get()), ObjectReference);
-			});
+		FBodyNodeGeometry BodyNodeGeometry{ (int32)ECADModelGeometryType::BodyNode, ObjectReference, BodyNode};
+		ModelConverter->AddGeometry(BodyNodeGeometry);
 
 		ModelConverter->RepairTopology();
 
@@ -1050,14 +1067,21 @@ namespace UE_DATASMITHWIRETRANSLATOR_NAMESPACE
 		FMeshDescription MeshDescription;
 		DatasmithMeshHelper::PrepareAttributeForStaticMesh(MeshDescription);
 
-		ModelConverter->Tessellate(MeshParameters, MeshDescription);
+		if (ModelConverter->Tessellate(OutMeshParameters, MeshDescription))
+		{
+			return TOptional<FMeshDescription>(MoveTemp(MeshDescription));
+		}
 
-		return MoveTemp(MeshDescription);
+		const TCHAR* StaticMeshLabel = MeshElement->GetLabel();
+		const TCHAR* StaticMeshName = MeshElement->GetName();
+		UE_LOG(LogWireInterface, Warning, TEXT("Failed to generate the mesh of \"%s\" (%s) StaticMesh."), StaticMeshLabel, StaticMeshName);
+
+		return TOptional<FMeshDescription>();
 	}
 
-	TOptional<FMeshDescription> FWireTranslatorImpl::GetMeshDescriptionFromPatchMesh(TSharedPtr<FPatchMesh>& PatchMesh, TSharedPtr<IDatasmithMeshElement> MeshElement, CADLibrary::FMeshParameters& MeshParameters)
+	TOptional<FMeshDescription> FWireTranslatorImpl::GetMeshDescriptionFromPatchMesh(TSharedPtr<FPatchMesh>& PatchMesh, TSharedPtr<IDatasmithMeshElement> MeshElement, CADLibrary::FMeshParameters& OutMeshParameters)
 	{
-		MeshParameters = OpenModelUtils::GetMeshParameters(PatchMesh->GetLayer());
+		OutMeshParameters = OpenModelUtils::GetMeshParameters(PatchMesh->GetLayer());
 
 		FMeshDescription MeshDescription;
 		DatasmithMeshHelper::PrepareAttributeForStaticMesh(MeshDescription);
@@ -1077,7 +1101,7 @@ namespace UE_DATASMITHWIRETRANSLATOR_NAMESPACE
 					TSharedPtr<IDatasmithMaterialIDElement> MaterialIDElement = MeshElement->GetMaterialSlotAt(0);
 					const FString SlotMaterialName = MaterialIDElement ? MaterialIDElement->GetName() : FString();
 
-					OpenModelUtils::TransferAlMeshToMeshDescription(*Mesh, *SlotMaterialName, MeshDescription, MeshParameters, bMerge);
+					OpenModelUtils::TransferAlMeshToMeshDescription(*Mesh, *SlotMaterialName, MeshDescription, OutMeshParameters, bMerge);
 				}
 			});
 
@@ -1087,33 +1111,55 @@ namespace UE_DATASMITHWIRETRANSLATOR_NAMESPACE
 		return MoveTemp(MeshDescription);
 	}
 
-	TOptional<FMeshDescription> FWireTranslatorImpl::GetMeshDescriptionFromParametricNode(AlDagNode& DagNode, TSharedPtr<IDatasmithMeshElement> MeshElement, CADLibrary::FMeshParameters& MeshParameters)
+	// #wire_import: AlSurfaceNode can have trim regions. This should be handle at this stage
+	TOptional<FMeshDescription> FWireTranslatorImpl::GetMeshDescriptionFromParametricNode(const TAlDagNodePtr<AlDagNode>& DagNode, TSharedPtr<IDatasmithMeshElement> MeshElement, CADLibrary::FMeshParameters& OutMeshParameters)
 	{
-		if (!WireSettings.bAliasUseNative)
+		OutMeshParameters = OpenModelUtils::GetMeshParameters(DagNode.GetLayer());
+
+		TSharedPtr<CADLibrary::ICADModelConverter> ModelConverter = GetModelConverter();
+
+		ModelConverter->InitializeProcess();
+
+		EAliasObjectReference ObjectReference = EAliasObjectReference::LocalReference;
+
+		// All geometry are processed in world space if layers are converted to actors
+		// All actors of a Alias symmetric layer are defined in the world Reference
+		// i.e. they have identity transform. So Mesh actor has to be defined in the world reference.
+		if (OutMeshParameters.bIsSymmetric)
 		{
-			return TessellateParametricNode(DagNode, MeshElement, MeshParameters);
+			ObjectReference = EAliasObjectReference::WorldReference;
 		}
-		else
+
+		ensure(MeshElement->GetMaterialSlotCount() == 1);
+
+		FDagNodeGeometry DagNodeGeometry{ (int32)ECADModelGeometryType::DagNode, ObjectReference, DagNode };
+		if (!ModelConverter->AddGeometry(DagNodeGeometry))
 		{
-			// TODO: the best way, should be to don't have to apply inverse global transform to the generated mesh
-			TAlDagNodePtr<AlMeshNode> MeshNode = OpenModelUtils::TesselateDagLeaf(DagNode, ETesselatorType::Fast, WireSettings.ChordTolerance);
+			return TOptional<FMeshDescription>();
+		}
 
-			if (MeshNode && MeshNode->mesh())
-			{
-				AlMatrix4x4 AlMatrix;
-				DagNode.inverseGlobalTransformationMatrix(AlMatrix);
-				MeshNode->mesh()->transform(AlMatrix);
+		ModelConverter->RepairTopology();
 
-				// Get the meshes from the dag nodes. Note that removing the mesh's DAG.
-				// will also removes the meshes, so we have to do it later.
-				return GetMeshDescriptionFromMeshNode(MeshNode, MeshElement, MeshParameters);
-			}
+		ModelConverter->SaveModel(*OutputPath, MeshElement);
+
+		FMeshDescription MeshDescription;
+		DatasmithMeshHelper::PrepareAttributeForStaticMesh(MeshDescription);
+
+		OutMeshParameters = OpenModelUtils::GetMeshParameters(*DagNode);
+
+		if (ModelConverter->Tessellate(OutMeshParameters, MeshDescription))
+		{
+			return TOptional<FMeshDescription>(MoveTemp(MeshDescription));
 		}
 
+		const TCHAR* StaticMeshLabel = MeshElement->GetLabel();
+		const TCHAR* StaticMeshName = MeshElement->GetName();
+		UE_LOG(LogWireInterface, Warning, TEXT("Failed to generate the mesh of \"%s\" (%s) StaticMesh."), StaticMeshLabel, StaticMeshName);
+
 		return TOptional<FMeshDescription>();
 	}
 
-	TOptional<FMeshDescription> FWireTranslatorImpl::GetMeshDescriptionFromMeshNode(const TAlDagNodePtr<AlMeshNode>& MeshNode, TSharedPtr<IDatasmithMeshElement> MeshElement, CADLibrary::FMeshParameters& MeshParameters, AlMatrix4x4* AlMeshInvGlobalMatrix)
+	TOptional<FMeshDescription> FWireTranslatorImpl::GetMeshDescriptionFromMeshNode(const TAlDagNodePtr<AlMeshNode>& MeshNode, TSharedPtr<IDatasmithMeshElement> MeshElement, CADLibrary::FMeshParameters& OutMeshParameters, AlMatrix4x4* AlMeshInvGlobalMatrix)
 	{
 		if (!MeshNode)
 		{
@@ -1137,8 +1183,10 @@ namespace UE_DATASMITHWIRETRANSLATOR_NAMESPACE
 		TSharedPtr<IDatasmithMaterialIDElement> MaterialIDElement = MeshElement->GetMaterialSlotAt(0);
 		const FString SlotMaterialName = MaterialIDElement ? MaterialIDElement->GetName() : FString();
 
+		OutMeshParameters = OpenModelUtils::GetMeshParameters(*MeshNode);
+
 		const bool bMerge = false;
-		OpenModelUtils::TransferAlMeshToMeshDescription(*Mesh, *SlotMaterialName, MeshDescription, MeshParameters, bMerge);
+		OpenModelUtils::TransferAlMeshToMeshDescription(*Mesh, *SlotMaterialName, MeshDescription, OutMeshParameters, bMerge);
 
 		// Build edge meta data
 		FStaticMeshOperations::DetermineEdgeHardnessesFromVertexInstanceNormals(MeshDescription);
@@ -1146,68 +1194,26 @@ namespace UE_DATASMITHWIRETRANSLATOR_NAMESPACE
 		return MoveTemp(MeshDescription);
 	}
 
-	TSharedPtr<CADLibrary::ICADModelConverter> FWireTranslatorImpl::GetModelConverter(IAliasBRepConverter*& BRepConverter) const
+	TSharedPtr<CADLibrary::ICADModelConverter> FWireTranslatorImpl::GetModelConverter() const
 	{
 		TSharedPtr<CADLibrary::ICADModelConverter> ModelConverter;
 
 		CADLibrary::FImportParameters ImportParameters;
 		if (CADLibrary::FImportParameters::bGDisableCADKernelTessellation)
 		{
-			TSharedPtr<FAliasModelToTechSoftConverter> AliasToTechSoftConverter = MakeShared<FAliasModelToTechSoftConverter>(ImportParameters);
-			ModelConverter = AliasToTechSoftConverter;
-			BRepConverter = AliasToTechSoftConverter.Get();
+			ModelConverter = MakeShared<FAliasModelToTechSoftConverter>(ImportParameters);
 		}
 		else
 		{
-			TSharedPtr<FAliasModelToCADKernelConverter> AliasToCADKernelConverter = MakeShared<FAliasModelToCADKernelConverter>(ImportParameters);
-			ModelConverter = AliasToCADKernelConverter;
-			BRepConverter = AliasToCADKernelConverter.Get();
-		}
-
-		ModelConverter->SetImportParameters(WireSettings.ChordTolerance, WireSettings.MaxEdgeLength, WireSettings.NormalTolerance, (CADLibrary::EStitchingTechnique)WireSettings.StitchingTechnique);
-
-		return ModelConverter;
-	}
-
-	TOptional<FMeshDescription> FWireTranslatorImpl::TessellateParametricNode(AlDagNode& DagNode, TSharedPtr<IDatasmithMeshElement> MeshElement, CADLibrary::FMeshParameters& MeshParameters)
-	{
-		// Wire unit is cm
-		IAliasBRepConverter* BRepConverter = nullptr;
-		TSharedPtr<CADLibrary::ICADModelConverter> ModelConverter = GetModelConverter(BRepConverter);
-
-		ModelConverter->InitializeProcess();
-
-		EAliasObjectReference ObjectReference = EAliasObjectReference::LocalReference;
-
-		if (MeshParameters.bIsSymmetric)
-		{
-			// All actors of a Alias symmetric layer are defined in the world Reference i.e. they have identity transform. So Mesh actor has to be defined in the world reference.
-			ObjectReference = EAliasObjectReference::WorldReference;
-		}
-
-		ensure(MeshElement->GetMaterialSlotCount() == 1);
-
-		if (!BRepConverter->AddBRep(DagNode, 0, ObjectReference))
-		{
-			return TOptional<FMeshDescription>();
+			ModelConverter = MakeShared<FAliasModelToCADKernelConverter>(ImportParameters);
 		}
 
-		ModelConverter->RepairTopology();
-
-		ModelConverter->SaveModel(*OutputPath, MeshElement);
-
-		FMeshDescription MeshDescription;
-		DatasmithMeshHelper::PrepareAttributeForStaticMesh(MeshDescription);
-
-		bool bRet = ModelConverter->Tessellate(MeshParameters, MeshDescription);
-		if (!bRet)
+		if (ModelConverter)
 		{
-			const TCHAR* StaticMeshLable = MeshElement->GetLabel();
-			const TCHAR* StaticMeshName = MeshElement->GetName();
-			UE_LOG(LogWireInterface, Warning, TEXT("Failed to generate the mesh of \"%s\" (%s) StaticMesh."), StaticMeshLable, StaticMeshName);
+			ModelConverter->SetImportParameters(WireSettings.ChordTolerance, WireSettings.MaxEdgeLength, WireSettings.NormalTolerance, (CADLibrary::EStitchingTechnique)WireSettings.StitchingTechnique);
 		}
 
-		return MoveTemp(MeshDescription);
+		return ModelConverter;
 	}
 
 	// Material creation
@@ -1857,7 +1863,6 @@ namespace UE_DATASMITHWIRETRANSLATOR_NAMESPACE
 		else {
 			MaterialElement->SetParentLabel(TEXT("M_DatasmithAliasLightSource"));
 		}
-
 	}
 
 	void FWireTranslatorImpl::AddAlPhongParameters(const TAlObjectPtr<AlShader>& Shader, TSharedPtr<IDatasmithUEPbrMaterialElement> MaterialElement)
diff --git a/Engine/Plugins/Enterprise/DatasmithCADImporter/Source/WireInterface/Private/WireInterfaceImpl.h b/Engine/Plugins/Enterprise/DatasmithCADImporter/Source/WireInterface/Private/WireInterfaceImpl.h
--- a/Engine/Plugins/Enterprise/DatasmithCADImporter/Source/WireInterface/Private/WireInterfaceImpl.h
+++ b/Engine/Plugins/Enterprise/DatasmithCADImporter/Source/WireInterface/Private/WireInterfaceImpl.h
@@ -6,7 +6,6 @@
 
 #ifdef USE_OPENMODEL
 
-#include "AliasBrepConverter.h"
 #include "OpenModelUtils.h"
 
 #include "IWireInterface.h"
@@ -35,7 +34,6 @@ namespace UE_DATASMITHWIRETRANSLATOR_NAMESPACE
 {
 	class FBodyData;
 	struct FDagNodeInfo;
-	class IAliasBRepConverter;
 
 	class FWireTranslatorImpl : public IWireInterface
 	{
@@ -87,11 +85,10 @@ namespace UE_DATASMITHWIRETRANSLATOR_NAMESPACE
 
 		/** Geometry retrieval */
 		TOptional<FMeshDescription> GetMeshDescription(TSharedPtr<IDatasmithMeshElement> MeshElement, CADLibrary::FMeshParameters& OutMeshParameters);
-		TSharedPtr<CADLibrary::ICADModelConverter> GetModelConverter(IAliasBRepConverter*& BRepConverter) const;
-		TOptional<FMeshDescription> TessellateParametricNode(AlDagNode& DagNode, TSharedPtr<IDatasmithMeshElement> MeshElement, CADLibrary::FMeshParameters& MeshParameters);
+		TSharedPtr<CADLibrary::ICADModelConverter> GetModelConverter() const;
 		TOptional<FMeshDescription> GetMeshDescriptionFromBodyNode(TSharedPtr<FBodyNode>& BodyNode, TSharedPtr<IDatasmithMeshElement> MeshElement, CADLibrary::FMeshParameters& MeshParameters);
 		TOptional<FMeshDescription> GetMeshDescriptionFromPatchMesh(TSharedPtr<FPatchMesh>& PatchMesh, TSharedPtr<IDatasmithMeshElement> MeshElement, CADLibrary::FMeshParameters& MeshParameters);
-		TOptional<FMeshDescription> GetMeshDescriptionFromParametricNode(AlDagNode& DagNode, TSharedPtr<IDatasmithMeshElement> MeshElement, CADLibrary::FMeshParameters& MeshParameters);
+		TOptional<FMeshDescription> GetMeshDescriptionFromParametricNode(const TAlDagNodePtr<AlDagNode>& DagNode, TSharedPtr<IDatasmithMeshElement> MeshElement, CADLibrary::FMeshParameters& MeshParameters);
 		TOptional<FMeshDescription> GetMeshDescriptionFromMeshNode(const TAlDagNodePtr<AlMeshNode>& MeshNode, TSharedPtr<IDatasmithMeshElement> MeshElement, CADLibrary::FMeshParameters& MeshParameters, AlMatrix4x4* AlMeshInvGlobalMatrix = nullptr);
 
 	private:
@@ -102,14 +99,14 @@ namespace UE_DATASMITHWIRETRANSLATOR_NAMESPACE
 		FWireSettings WireSettings;
 
 		TSharedPtr<CADLibrary::ICADModelConverter> CADModelConverter;
-		TSharedPtr<IAliasBRepConverter> AliasBRepConverter;
 
 		bool bSceneLoaded = false;
 
 		TMap<FString, TSharedPtr<IDatasmithBaseMaterialElement>> ShaderNameToMaterial;
 
 		TMap<uint32, TSharedPtr<IDatasmithMeshElement>> GeomNodeToMeshElement;
-		TMap<TSharedPtr<IDatasmithMeshElement>, AlDagNode*> MeshElementToGeomNode;
+		TMap<TSharedPtr<IDatasmithMeshElement>, AlDagNode*> MeshElementToParametricNode;
+		TMap<TSharedPtr<IDatasmithMeshElement>, AlMeshNode*> MeshElementToMeshNode;
 
 		TMap<uint32, TSharedPtr<IDatasmithMeshElement>> BodyNodeToMeshElement;
 		TMap<TSharedPtr<IDatasmithMeshElement>, TSharedPtr<FBodyNode>> MeshElementToBodyNode;
