Fixed build error

#rnx

diff --git a/Engine/Plugins/Enterprise/DatasmithCADImporter/Source/WireInterface/Private/OpenModelUtils.h b/Engine/Plugins/Enterprise/DatasmithCADImporter/Source/WireInterface/Private/OpenModelUtils.h
--- a/Engine/Plugins/Enterprise/DatasmithCADImporter/Source/WireInterface/Private/OpenModelUtils.h
+++ b/Engine/Plugins/Enterprise/DatasmithCADImporter/Source/WireInterface/Private/OpenModelUtils.h
@@ -359,24 +359,27 @@ namespace UE_DATASMITHWIRETRANSLATOR_NAMESPACE
 	struct FDagNodeGeometry : public FAliasGeometry
 	{
 		TAlDagNodePtr<AlDagNode> DagNode;
+
+		FDagNodeGeometry(int32 InType, EAliasObjectReference InReference, const TAlDagNodePtr<AlDagNode>& InDagNode)
+		{
+			Type = InType;
+			Reference = InReference;
+			DagNode = InDagNode;
+		}
 	};
 
 	struct FBodyNodeGeometry : public FAliasGeometry
 	{
 		TSharedPtr<FBodyNode> BodyNode;
-	};
 
-	struct FPatchMeshGeometry : public CADLibrary::FCADModelGeometry
-	{
-		TSharedPtr<FPatchMesh> PatchMesh;
-	};
-
-	struct FMeshNodeGeometry : public CADLibrary::FCADModelGeometry
-	{
-		TAlDagNodePtr<AlMeshNode> MeshNode;
+		FBodyNodeGeometry(int32 InType, EAliasObjectReference InReference, const TSharedPtr<FBodyNode>& InBodyNode)
+		{
+			Type = InType;
+			Reference = InReference;
+			BodyNode = InBodyNode;
+		}
 	};
 
-
 	namespace OpenModelUtils
 	{
 		/** Following layer hierarchy, get list of layers an actor would be in as a csv string*/
diff --git a/Engine/Plugins/Enterprise/DatasmithCADImporter/Source/WireInterface/Private/WireInterfaceImpl.cpp b/Engine/Plugins/Enterprise/DatasmithCADImporter/Source/WireInterface/Private/WireInterfaceImpl.cpp
--- a/Engine/Plugins/Enterprise/DatasmithCADImporter/Source/WireInterface/Private/WireInterfaceImpl.cpp
+++ b/Engine/Plugins/Enterprise/DatasmithCADImporter/Source/WireInterface/Private/WireInterfaceImpl.cpp
@@ -1057,7 +1057,7 @@ namespace UE_DATASMITHWIRETRANSLATOR_NAMESPACE
 			ObjectReference = EAliasObjectReference::ParentReference;
 		}
 
-		FBodyNodeGeometry BodyNodeGeometry{ (int32)ECADModelGeometryType::BodyNode, ObjectReference, BodyNode};
+		const FBodyNodeGeometry BodyNodeGeometry( (int32)ECADModelGeometryType::BodyNode, ObjectReference, BodyNode);
 		ModelConverter->AddGeometry(BodyNodeGeometry);
 
 		ModelConverter->RepairTopology();
@@ -1132,7 +1132,7 @@ namespace UE_DATASMITHWIRETRANSLATOR_NAMESPACE
 
 		ensure(MeshElement->GetMaterialSlotCount() == 1);
 
-		FDagNodeGeometry DagNodeGeometry{ (int32)ECADModelGeometryType::DagNode, ObjectReference, DagNode };
+		const FDagNodeGeometry DagNodeGeometry( (int32)ECADModelGeometryType::DagNode, ObjectReference, DagNode );
 		if (!ModelConverter->AddGeometry(DagNodeGeometry))
 		{
 			return TOptional<FMeshDescription>();
