Fixed infinite loop when importing shell nodes.

diff --git a/Engine/Plugins/Enterprise/DatasmithCADImporter/Source/WireInterface/Private/WireInterfaceImpl.cpp b/Engine/Plugins/Enterprise/DatasmithCADImporter/Source/WireInterface/Private/WireInterfaceImpl.cpp
--- a/Engine/Plugins/Enterprise/DatasmithCADImporter/Source/WireInterface/Private/WireInterfaceImpl.cpp
+++ b/Engine/Plugins/Enterprise/DatasmithCADImporter/Source/WireInterface/Private/WireInterfaceImpl.cpp
@@ -577,13 +577,15 @@ namespace UE_DATASMITHWIRETRANSLATOR_NAMESPACE
 		TAlDagNodePtr<AlShellNode> ShellNode(GeomNode->asShellNodePtr());
 		if (ShellNode && ShellNode->shell())
 		{
-			TAlObjectPtr<AlShader> Shader(ShellNode->shell()->firstShader());
+			TAlObjectPtr<AlShell> Shell(ShellNode->shell());
+			TAlObjectPtr<AlShader> Shader(Shell->firstShader());
 			int32 SlotIndex = 0;
 			while (Shader.IsValid())
 			{
 				TSharedPtr<IDatasmithMaterialIDElement> MaterialIDElement = FindOrAddMaterial(Shader);
 				MaterialIDElement->SetId(SlotIndex++);
 				ActorElement->AddMaterialOverride(MaterialIDElement);
+				Shader = Shell->nextShader(Shader.Get());
 			}
 		}
 		else
