Completed fix for position issue

#rnx

diff --git a/Engine/Plugins/Enterprise/DatasmithCADImporter/Source/WireInterface/Private/WireInterfaceImpl.cpp b/Engine/Plugins/Enterprise/DatasmithCADImporter/Source/WireInterface/Private/WireInterfaceImpl.cpp
--- a/Engine/Plugins/Enterprise/DatasmithCADImporter/Source/WireInterface/Private/WireInterfaceImpl.cpp
+++ b/Engine/Plugins/Enterprise/DatasmithCADImporter/Source/WireInterface/Private/WireInterfaceImpl.cpp
@@ -549,6 +549,12 @@ namespace UE_DATASMITHWIRETRANSLATOR_NAMESPACE
 
 	TSharedPtr<IDatasmithActorElement> FWireTranslatorImpl::ProcessGeometryNode(const TAlDagNodePtr<AlDagNode>& GeomNode, const TAlObjectPtr<AlLayer>& ParentLayer)
 	{
+		TAlObjectPtr<AlLayer> Layer(GeomNode->layer());
+		if (Layer && Layer->invisible())
+		{
+			return TSharedPtr<IDatasmithActorElement>();
+		}
+
 		TSharedPtr<IDatasmithMeshElement> MeshElement = FindOrAddMeshElement(GeomNode);
 		if (!MeshElement.IsValid())
 		{
@@ -565,7 +571,6 @@ namespace UE_DATASMITHWIRETRANSLATOR_NAMESPACE
 		ActorElement->SetLabel(Label.Len() > 0 ? *Label : TEXT("NoName"));
 		ActorElement->SetStaticMeshPathName(MeshElement->GetName());
 
-		TAlObjectPtr<AlLayer> Layer(GeomNode->layer());
 		FString CsvLayerString;
 		if (OpenModelUtils::GetCsvLayerString(Layer, CsvLayerString))
 		{
@@ -687,6 +692,12 @@ namespace UE_DATASMITHWIRETRANSLATOR_NAMESPACE
 			return TSharedPtr<IDatasmithActorElement>();
 		}
 
+		TAlObjectPtr<AlLayer> Layer(GroupNode->layer());
+		if (Layer && Layer->invisible())
+		{
+			return TSharedPtr<IDatasmithActorElement>();
+		}
+
 		if (PatchMesh->HasSingleContent())
 		{
 			return ProcessGeometryNode(PatchMesh->GetSingleContent().Get(), ParentLayer);
@@ -763,7 +774,6 @@ namespace UE_DATASMITHWIRETRANSLATOR_NAMESPACE
 		if (ShellNode)
 		{
 			TAlObjectPtr<AlShell> Shell(ShellNode->shell());
-			ensure(Shell);
 
 			TAlObjectPtr<AlShader> Shader(Shell->firstShader());
 			int32 SlotIndex = 0;
@@ -787,7 +797,7 @@ namespace UE_DATASMITHWIRETRANSLATOR_NAMESPACE
 			else
 			{
 				TAlDagNodePtr<AlMeshNode> MeshNode(GeomNode->asMeshNodePtr());
-				if (MeshNode && MeshNode->mesh())
+				if (MeshNode && AlIsValid(MeshNode->mesh()))
 				{
 					ApplyMaterial(MeshNode->mesh()->firstShader(), 0);
 					MeshElementToMeshNode.Add(MeshElement, MeshNode.Get());
@@ -989,6 +999,7 @@ namespace UE_DATASMITHWIRETRANSLATOR_NAMESPACE
 			{
 				TAlDagNodePtr<AlDagNode> GeomNode(*GeomNodePtr);
 
+				// #wire_import: Check whether parametric geometry with symmetry keeps the symmetry
 				// #wire_import: the best way, should be to don't have to apply inverse global transform to the generated mesh
 				TAlDagNodePtr<AlMeshNode> MeshNode = OpenModelUtils::TesselateDagLeaf(*GeomNode, ETesselatorType::Fast, WireSettings.ChordTolerance);
 
@@ -1096,7 +1107,14 @@ namespace UE_DATASMITHWIRETRANSLATOR_NAMESPACE
 				if (Mesh)
 				{
 					AlMatrix4x4 AlMatrix;
-					MeshNode->localTransformationMatrix(AlMatrix);
+					if (OutMeshParameters.bIsSymmetric)
+					{
+						MeshNode->globalTransformationMatrix(AlMatrix);
+					}
+					else
+					{
+						MeshNode->localTransformationMatrix(AlMatrix);
+					}
 
 					Mesh->transform(AlMatrix);
 
@@ -1161,7 +1179,7 @@ namespace UE_DATASMITHWIRETRANSLATOR_NAMESPACE
 		return TOptional<FMeshDescription>();
 	}
 
-	TOptional<FMeshDescription> FWireTranslatorImpl::GetMeshDescriptionFromMeshNode(const TAlDagNodePtr<AlMeshNode>& MeshNode, TSharedPtr<IDatasmithMeshElement> MeshElement, CADLibrary::FMeshParameters& OutMeshParameters, AlMatrix4x4* AlMeshInvGlobalMatrix)
+	TOptional<FMeshDescription> FWireTranslatorImpl::GetMeshDescriptionFromMeshNode(const TAlDagNodePtr<AlMeshNode>& MeshNode, TSharedPtr<IDatasmithMeshElement> MeshElement, CADLibrary::FMeshParameters& OutMeshParameters)
 	{
 		if (!MeshNode)
 		{
@@ -1174,9 +1192,13 @@ namespace UE_DATASMITHWIRETRANSLATOR_NAMESPACE
 			return TOptional<FMeshDescription>();
 		}
 
-		if (AlMeshInvGlobalMatrix != nullptr)
+		OutMeshParameters = OpenModelUtils::GetMeshParameters(*MeshNode);
+
+		if (OutMeshParameters.bIsSymmetric)
 		{
-			Mesh->transform(*AlMeshInvGlobalMatrix);
+			AlMatrix4x4 AlGlobalMatrix;
+			MeshNode->globalTransformationMatrix(AlGlobalMatrix);
+			Mesh->transform(AlGlobalMatrix);
 		}
 
 		FMeshDescription MeshDescription;
@@ -1185,8 +1207,6 @@ namespace UE_DATASMITHWIRETRANSLATOR_NAMESPACE
 		TSharedPtr<IDatasmithMaterialIDElement> MaterialIDElement = MeshElement->GetMaterialSlotAt(0);
 		const FString SlotMaterialName = MaterialIDElement ? MaterialIDElement->GetName() : FString();
 
-		OutMeshParameters = OpenModelUtils::GetMeshParameters(*MeshNode);
-
 		const bool bMerge = false;
 		OpenModelUtils::TransferAlMeshToMeshDescription(*Mesh, *SlotMaterialName, MeshDescription, OutMeshParameters, bMerge);
 
diff --git a/Engine/Plugins/Enterprise/DatasmithCADImporter/Source/WireInterface/Private/WireInterfaceImpl.h b/Engine/Plugins/Enterprise/DatasmithCADImporter/Source/WireInterface/Private/WireInterfaceImpl.h
--- a/Engine/Plugins/Enterprise/DatasmithCADImporter/Source/WireInterface/Private/WireInterfaceImpl.h
+++ b/Engine/Plugins/Enterprise/DatasmithCADImporter/Source/WireInterface/Private/WireInterfaceImpl.h
@@ -89,7 +89,7 @@ namespace UE_DATASMITHWIRETRANSLATOR_NAMESPACE
 		TOptional<FMeshDescription> GetMeshDescriptionFromBodyNode(TSharedPtr<FBodyNode>& BodyNode, TSharedPtr<IDatasmithMeshElement> MeshElement, CADLibrary::FMeshParameters& MeshParameters);
 		TOptional<FMeshDescription> GetMeshDescriptionFromPatchMesh(TSharedPtr<FPatchMesh>& PatchMesh, TSharedPtr<IDatasmithMeshElement> MeshElement, CADLibrary::FMeshParameters& MeshParameters);
 		TOptional<FMeshDescription> GetMeshDescriptionFromParametricNode(const TAlDagNodePtr<AlDagNode>& DagNode, TSharedPtr<IDatasmithMeshElement> MeshElement, CADLibrary::FMeshParameters& MeshParameters);
-		TOptional<FMeshDescription> GetMeshDescriptionFromMeshNode(const TAlDagNodePtr<AlMeshNode>& MeshNode, TSharedPtr<IDatasmithMeshElement> MeshElement, CADLibrary::FMeshParameters& MeshParameters, AlMatrix4x4* AlMeshInvGlobalMatrix = nullptr);
+		TOptional<FMeshDescription> GetMeshDescriptionFromMeshNode(const TAlDagNodePtr<AlMeshNode>& MeshNode, TSharedPtr<IDatasmithMeshElement> MeshElement, CADLibrary::FMeshParameters& MeshParameters);
 
 	private:
 		TSharedPtr<IDatasmithScene> DatasmithScene;
