Exposed Geometric and Stitching tolerance to users

#rnx

diff --git a/Engine/Plugins/Enterprise/DatasmithCADImporter/Source/CADLibrary/Private/CADKernelTools.cpp b/Engine/Plugins/Enterprise/DatasmithCADImporter/Source/CADLibrary/Private/CADKernelTools.cpp
--- a/Engine/Plugins/Enterprise/DatasmithCADImporter/Source/CADLibrary/Private/CADKernelTools.cpp
+++ b/Engine/Plugins/Enterprise/DatasmithCADImporter/Source/CADLibrary/Private/CADKernelTools.cpp
@@ -270,7 +270,7 @@ bool FCADKernelTools::Tessellate(UE::CADKernel::FTopologicalShapeEntity& CADTopo
 	// Tessellate the model
 	TSharedRef<FModelMesh> CADKernelModelMesh = FEntity::MakeShared<FModelMesh>();
 
-	const double GeometricTolerance = FImportParameters::GStitchingTolerance * 10; // cm to mm
+	const double GeometricTolerance = TessellationContext.GeometricTolerance * 10; // cm to mm
 	const bool bActivateThinZoneMeshing = FImportParameters::bGActivateThinZoneMeshing;
 	FMesher Mesher(*CADKernelModelMesh, GeometricTolerance, bActivateThinZoneMeshing);
 
diff --git a/Engine/Plugins/Enterprise/DatasmithCADImporter/Source/CADLibrary/Private/MeshDescriptionHelper.cpp b/Engine/Plugins/Enterprise/DatasmithCADImporter/Source/CADLibrary/Private/MeshDescriptionHelper.cpp
--- a/Engine/Plugins/Enterprise/DatasmithCADImporter/Source/CADLibrary/Private/MeshDescriptionHelper.cpp
+++ b/Engine/Plugins/Enterprise/DatasmithCADImporter/Source/CADLibrary/Private/MeshDescriptionHelper.cpp
@@ -139,8 +139,7 @@ void FillVertexPosition(const FMeshConversionContext& MeshConversionContext, FBo
 		VertexIdSet[VertexIndex] = VertexID;
 	}
 
-	const float GeometricTolerance = 0.001f; // cm
-	MergeCoincidents(VertexArray, FVector3f::OneVector, GeometricTolerance, VertexIdSet);
+	MergeCoincidents(VertexArray, FVector3f::OneVector, (float)MeshConversionContext.GeometricTolerance, VertexIdSet);
 
 	// if Symmetric mesh, the symmetric side of the mesh have to be generated
 	if (MeshConversionContext.MeshParameters.bIsSymmetric)
diff --git a/Engine/Plugins/Enterprise/DatasmithCADImporter/Source/CADLibrary/Public/MeshDescriptionHelper.h b/Engine/Plugins/Enterprise/DatasmithCADImporter/Source/CADLibrary/Public/MeshDescriptionHelper.h
--- a/Engine/Plugins/Enterprise/DatasmithCADImporter/Source/CADLibrary/Public/MeshDescriptionHelper.h
+++ b/Engine/Plugins/Enterprise/DatasmithCADImporter/Source/CADLibrary/Public/MeshDescriptionHelper.h
@@ -28,17 +28,19 @@ struct FMeshConversionContext
 	TArray<int32> VertexIds;
 	TArray<int32> SymmetricVertexIds;
 	TSet<int32> PatchesToMesh;
+	double GeometricTolerance = 0.001;
 
-	FMeshConversionContext(const FImportParameters& InImportParams, const FMeshParameters& InMeshParameters)
+	FMeshConversionContext(const FImportParameters& InImportParams, const FMeshParameters& InMeshParameters, double InGeometricTolerance = 0.001)
 		: ImportParameters(InImportParams)
 		, MeshParameters(InMeshParameters)
+		, GeometricTolerance(InGeometricTolerance)
 	{
 	}
 };
 
 /**
  * Macro data of MeshDescription.
- * Needed to updata a MeshDesction after a convertion of a DynamicMesh To a MeshDescription.
+ * Needed to update a MeshDesction after a conversion of a DynamicMesh To a MeshDescription.
  * Indeed, PolygonGroupId can be modify during the process of DynamicMesh and material slot names are not saved in a DynamicMesh.
  * The only persistent data is the patchGroupId
  * 
diff --git a/Engine/Plugins/Enterprise/DatasmithCADImporter/Source/WireInterface/Private/AliasModelToCADKernelConverter.cpp b/Engine/Plugins/Enterprise/DatasmithCADImporter/Source/WireInterface/Private/AliasModelToCADKernelConverter.cpp
--- a/Engine/Plugins/Enterprise/DatasmithCADImporter/Source/WireInterface/Private/AliasModelToCADKernelConverter.cpp
+++ b/Engine/Plugins/Enterprise/DatasmithCADImporter/Source/WireInterface/Private/AliasModelToCADKernelConverter.cpp
@@ -113,6 +113,15 @@ namespace AliasToCADKernelUtils
 	}
 }
 
+FAliasModelToCADKernelConverter::FAliasModelToCADKernelConverter(const FDatasmithTessellationOptions& Options, CADLibrary::FImportParameters InImportParameters)
+	: FCADModelToCADKernelConverterBase(InImportParameters)
+{
+	GeometricTolerance = Options.GeometricTolerance;
+	SquareTolerance = GeometricTolerance * GeometricTolerance;
+	CADKernelSession.SetGeometricTolerance(GeometricTolerance);
+	StitchingTolerance = Options.StitchingTolerance;
+}
+
 TSharedPtr<UE::CADKernel::FTopologicalEdge> FAliasModelToCADKernelConverter::AddEdge(const AlTrimCurve& AliasTrimCurve, TSharedPtr<UE::CADKernel::FSurface>& CarrierSurface)
 {
 	UE::CADKernel::FNurbsCurveData NurbsCurveData;
@@ -367,7 +376,7 @@ bool FAliasModelToCADKernelConverter::Tessellate(const CADLibrary::FMeshParamete
 {
 	UE::CADKernel::FModel& Model = CADKernelSession.GetModel();
 
-	CADLibrary::FMeshConversionContext Context(ImportParameters, InMeshParameters);
+	CADLibrary::FMeshConversionContext Context(ImportParameters, InMeshParameters, GeometricTolerance);
 
 	return CADLibrary::FCADKernelTools::Tessellate(Model, Context, OutMeshDescription);
 }
@@ -378,15 +387,11 @@ bool FAliasModelToCADKernelConverter::RepairTopology()
 	// Apply stitching if applicable
 	if (ImportParameters.GetStitchingTechnique() != StitchingNone)
 	{
-		const double StitchingTolerance = FImportParameters::GStitchingTolerance * 10.; //CM to MM
 		UE::CADKernel::ESewOption SewOptionValue = (UE::CADKernel::ESewOption)SewOption::GetFromImportParameters();
 
-		const IConsoleVariable* ConsoleVariable = IConsoleManager::Get().FindConsoleVariable(TEXT("ds.WireTranslator.SkipThinZoneMeshing"));
-		bool bSkipThinZoneMeshing = ConsoleVariable ? ConsoleVariable->GetBool() : false;
-		if (bSkipThinZoneMeshing)
-		{
-			SewOptionValue = UE::CADKernel::ESewOption((uint8)SewOptionValue & (uint8)~ESewOption::RemoveThinFaces);
-		}
+#if !THINFACE_ENABLED
+		SewOptionValue = UE::CADKernel::ESewOption((uint8)SewOptionValue & (uint8)~ESewOption::RemoveThinFaces);
+#endif
 
 		UE::CADKernel::FTopomakerOptions TopomakerOptions(SewOptionValue, StitchingTolerance, FImportParameters::GStitchingForceFactor);
 
diff --git a/Engine/Plugins/Enterprise/DatasmithCADImporter/Source/WireInterface/Private/AliasModelToCADKernelConverter.h b/Engine/Plugins/Enterprise/DatasmithCADImporter/Source/WireInterface/Private/AliasModelToCADKernelConverter.h
--- a/Engine/Plugins/Enterprise/DatasmithCADImporter/Source/WireInterface/Private/AliasModelToCADKernelConverter.h
+++ b/Engine/Plugins/Enterprise/DatasmithCADImporter/Source/WireInterface/Private/AliasModelToCADKernelConverter.h
@@ -35,10 +35,7 @@ class FAliasModelToCADKernelConverter : public FCADModelToCADKernelConverterBase
 {
 
 public:
-	FAliasModelToCADKernelConverter(CADLibrary::FImportParameters InImportParameters)
-		: FCADModelToCADKernelConverterBase(InImportParameters)
-	{
-	}
+	FAliasModelToCADKernelConverter(const FDatasmithTessellationOptions& Options, CADLibrary::FImportParameters InImportParameters);
 
 	// Begin FCADModelToCADKernelConverterBase overrides
 	virtual bool Tessellate(const CADLibrary::FMeshParameters& InMeshParameters, FMeshDescription& OutMeshDescription) override;
@@ -68,6 +65,7 @@ protected:
 protected:
 	int32 LastFaceId = 1;
 	TMap<void*, TSharedPtr<UE::CADKernel::FTopologicalEdge>>  AlEdge2CADKernelEdge;
+	double StitchingTolerance;
 };
 
 }
diff --git a/Engine/Plugins/Enterprise/DatasmithCADImporter/Source/WireInterface/Private/WireInterfaceImpl.cpp b/Engine/Plugins/Enterprise/DatasmithCADImporter/Source/WireInterface/Private/WireInterfaceImpl.cpp
--- a/Engine/Plugins/Enterprise/DatasmithCADImporter/Source/WireInterface/Private/WireInterfaceImpl.cpp
+++ b/Engine/Plugins/Enterprise/DatasmithCADImporter/Source/WireInterface/Private/WireInterfaceImpl.cpp
@@ -71,6 +71,12 @@ DEFINE_LOG_CATEGORY_STATIC(LogWireInterface, Log, All);
 
 #define WRONG_VERSION_TEXT "Unsupported version of Alias detected. Please upgrade to Alias 2021.3 (or later version)."
 
+#define TRACK_MESHELEMENT 0
+
+#if TRACK_MESHELEMENT
+#include "CompGeom/FitOrientedBox3.h"
+#endif
+
 #ifdef OPEN_MODEL_2023_0
 
 static bool bGAliasSewByMaterial = false;
@@ -85,6 +91,9 @@ ECVF_Default);
 
 namespace UE_DATASMITHWIRETRANSLATOR_NAMESPACE
 {
+#if TRACK_MESHELEMENT
+	void MakeMeshVisible(FMeshDescription& MeshDescription);
+#endif
 
 	static bool bGSewByMaterial = false;
 
@@ -238,7 +247,7 @@ namespace UE_DATASMITHWIRETRANSLATOR_NAMESPACE
 			}
 			else
 			{
-				CADModelConverter = MakeShared<FAliasModelToCADKernelConverter>(ImportParameters);
+				CADModelConverter = MakeShared<FAliasModelToCADKernelConverter>(WireSettings, ImportParameters);
 			}
 		}
 		else
@@ -759,7 +768,16 @@ namespace UE_DATASMITHWIRETRANSLATOR_NAMESPACE
 
 		MeshElement->SetLabel(*GeomNode.GetName());
 		MeshElement->SetLightmapSourceUV(-1);
-
+#if TRACK_MESHELEMENT
+		{
+			//static const FString ToTrack(TEXT("shell_30181"));
+			static const FString ToTrack(TEXT("Blend_srf_1395"));
+			if (!ToTrack.Equals(MeshElement->GetLabel()))
+			{
+				return TSharedPtr<IDatasmithMeshElement>();
+			}
+		}
+#endif
 		auto ApplyMaterial = [this, &MeshElement](const TAlObjectPtr<AlShader>& Shader, int32 SlotIndex)
 			{
 				if (TSharedPtr<IDatasmithMaterialIDElement> MaterialIDElement = this->FindOrAddMaterial(Shader))
@@ -834,6 +852,15 @@ namespace UE_DATASMITHWIRETRANSLATOR_NAMESPACE
 
 		MeshElement->SetLabel(*BodyNode->GetName());
 		MeshElement->SetLightmapSourceUV(-1);
+#if TRACK_MESHELEMENT
+		{
+			static const FString ToTrack(TEXT("Blend_srf_1395"));
+			if (!ToTrack.Equals(MeshElement->GetLabel()))
+			{
+				return TSharedPtr<IDatasmithMeshElement>();
+			}
+		}
+#endif
 
 		auto ApplyMaterial = [this, &MeshElement](int32 SlotIndex, const TAlObjectPtr<AlShader>& Shader)
 			{
@@ -946,29 +973,18 @@ namespace UE_DATASMITHWIRETRANSLATOR_NAMESPACE
 
 	// Geometry retrieval
 
-#define TRACK_MESH 0
-
 	bool FWireTranslatorImpl::LoadStaticMesh(const TSharedPtr<IDatasmithMeshElement> MeshElement, FDatasmithMeshElementPayload& OutMeshPayload, const FDatasmithTessellationOptions& InTessellationOptions)
 	{
-#if TRACK_MESH
-		static const FString MeshName(TEXT("shell_30181"));
-		bool bSkip = true;
-		if (MeshName.Equals(MeshElement->GetLabel()))
-		{
-			bSkip = false;
-		}
-		if (bSkip)
-		{
-			return false;
-		}
-
-		uint64 StartTime = FPlatformTime::Cycles64();
-#endif
 		CADLibrary::FMeshParameters MeshParameters;
 
 		if (TOptional<FMeshDescription> Mesh = GetMeshDescription(MeshElement, MeshParameters))
 		{
+#if TRACK_MESHELEMENT
+			FMeshDescription& MeshDescription = OutMeshPayload.LodMeshes.Add_GetRef(MoveTemp(Mesh.GetValue()));
+			MakeMeshVisible(MeshDescription);
+#else
 			OutMeshPayload.LodMeshes.Add(MoveTemp(Mesh.GetValue()));
+#endif
 			const TCHAR* MeshFilename = MeshElement->GetFile();
 			if (!WireSettings.bAliasUseNative && FPaths::FileExists(MeshFilename))
 			{
@@ -980,15 +996,11 @@ namespace UE_DATASMITHWIRETRANSLATOR_NAMESPACE
 					IFileManager::Get().Delete(MeshFilename);
 				}
 			}
-		}
 
-#if TRACK_MESH
-		double ElapsedSeconds = FPlatformTime::ToSeconds64(FPlatformTime::Cycles64() - StartTime);
-
-		UE_LOG(LogWireInterface, Display, TEXT("Mesh description of %s retrieved in %.3f s"), MeshElement->GetLabel(), ElapsedSeconds);
-#endif
+			return true;
+		}
 
-		return OutMeshPayload.LodMeshes.Num() > 0;
+		return false;
 	}
 
 	TOptional<FMeshDescription> FWireTranslatorImpl::GetMeshDescription(TSharedPtr<IDatasmithMeshElement> MeshElement, CADLibrary::FMeshParameters& OutMeshParameters)
@@ -1227,7 +1239,7 @@ namespace UE_DATASMITHWIRETRANSLATOR_NAMESPACE
 		}
 		else
 		{
-			ModelConverter = MakeShared<FAliasModelToCADKernelConverter>(ImportParameters);
+			ModelConverter = MakeShared<FAliasModelToCADKernelConverter>(WireSettings, ImportParameters);
 		}
 
 		if (ModelConverter)
@@ -2175,6 +2187,39 @@ namespace UE_DATASMITHWIRETRANSLATOR_NAMESPACE
 		}
 	};
 
+#if TRACK_MESHELEMENT
+#pragma optimize("", off)
+	void MakeMeshVisible(FMeshDescription& MeshDescription)
+	{
+		using namespace UE::Geometry;
+
+		TVertexAttributesRef<FVector3f> VertexPositions = MeshDescription.GetVertexPositions();
+		
+		TArrayView<FVector3f> Positions = VertexPositions.GetRawArray();
+		TOrientedBox3<float> OBox = FitOrientedBox3Points<float>(Positions);
+		
+		FVertexArray& Vertices = MeshDescription.Vertices();
+		const TMatrix3<float> Matrix = OBox.Frame.Rotation.ToRotationMatrix();
+
+		constexpr float MinSize = 1.f;
+		constexpr float MaxSize = 20.f;
+
+		const float ScaleX = OBox.Extents.X < MinSize ? 2.f / OBox.Extents.X : OBox.Extents.X > MaxSize ? MaxSize / OBox.Extents.X : 1.f;
+		const float ScaleY = OBox.Extents.Y < MinSize ? 2.f / OBox.Extents.Y : OBox.Extents.Y > MaxSize ? MaxSize / OBox.Extents.Y : 1.f;
+		const float ScaleZ = OBox.Extents.Z < MinSize ? 2.f / OBox.Extents.Z : OBox.Extents.Z > MaxSize ? MaxSize / OBox.Extents.Z : 1.f;
+		UE_LOG(LogWireInterface, Warning, TEXT("Scaling factor: %.3f %.3f %.3f"), ScaleX, ScaleY, ScaleZ);
+		const FVector3f AxisX = OBox.AxisX();
+		const FVector3f AxisY = OBox.AxisY();
+		const FVector3f AxisZ = OBox.AxisZ();
+
+		for (const FVertexID& VertexID : MeshDescription.Vertices().GetElementIDs())
+		{
+			FVector3f P = VertexPositions[VertexID] - OBox.Frame.Origin;
+			VertexPositions[VertexID] = ((P | AxisX) * ScaleX) * AxisX + ((P | AxisY) * ScaleY) * AxisY + ((P | AxisZ) * ScaleZ) * AxisZ;
+		}
+	}
+#endif
+
 } // namespace
 
 #undef LOCTEXT_NAMESPACE // "DatasmithWireTranslator"
diff --git a/Engine/Plugins/Enterprise/DatasmithCADImporter/Source/WireInterface/WireInterface2020.Build.cs b/Engine/Plugins/Enterprise/DatasmithCADImporter/Source/WireInterface/WireInterface2020.Build.cs
--- a/Engine/Plugins/Enterprise/DatasmithCADImporter/Source/WireInterface/WireInterface2020.Build.cs
+++ b/Engine/Plugins/Enterprise/DatasmithCADImporter/Source/WireInterface/WireInterface2020.Build.cs
@@ -34,6 +34,7 @@ public abstract class WireInterfaceBase : ModuleRules
 				"DatasmithWireTranslator",
 				"Engine",
 				"MeshDescription",
+				"GeometryCore",
 				"ParametricSurface",
 				"StaticMeshDescription",
 			}
@@ -55,6 +56,7 @@ public abstract class WireInterfaceBase : ModuleRules
 		PublicDefinitions.Add($"UE_DATASMITHWIRETRANSLATOR_MODULE_NAME={GetType().Name}");
 		PublicDefinitions.Add($"UE_OPENMODEL_MAJOR_VERSION={GetMajorVersion()}");
 		PublicDefinitions.Add($"UE_OPENMODEL_MINOR_VERSION={GetMinorVersion()}");
+		PublicDefinitions.Add("THINFACE_ENABLED=0");
 
 		if (System.Type.GetType(GetAliasVersion()) != null)
 		{
diff --git a/Engine/Plugins/Enterprise/DatasmithContent/Source/DatasmithContent/Public/DatasmithImportOptions.h b/Engine/Plugins/Enterprise/DatasmithContent/Source/DatasmithContent/Public/DatasmithImportOptions.h
--- a/Engine/Plugins/Enterprise/DatasmithContent/Source/DatasmithContent/Public/DatasmithImportOptions.h
+++ b/Engine/Plugins/Enterprise/DatasmithContent/Source/DatasmithContent/Public/DatasmithImportOptions.h
@@ -263,11 +263,24 @@ struct DATASMITHCONTENT_API FDatasmithTessellationOptions
 	 * Sewing : Connects surfaces which physically share a boundary but not topologically within a set of objects.
 	 *          This technique can modify the structure of the model by removing and adding objects.
 	 * Healing : Connects surfaces which physically share a boundary but not topologically within an object.
-	 * The techniques are using the chord tolerance to determine if two surfaces should be stitched.
 	 */
 	UPROPERTY(config, EditAnywhere, BlueprintReadWrite, AdvancedDisplay, Category = "Geometry & Tessellation Options", meta = (ToolTip = "Stitching technique applied on model before tessellation. Sewing could impact number of objects."))
 	EDatasmithCADStitchingTechnique StitchingTechnique;
 
+	/**
+	 * Tolerance used to determine if a surface should be tessellate or not.
+	 * Any surface which is narrower than the geometric tolerance
+	 * in one of the iso direction will not be tessellated
+	 */
+	UPROPERTY(config, EditAnywhere, BlueprintReadWrite, AdvancedDisplay, Category = "Geometry & Tessellation Options", meta = (Units = cm, ToolTip = "Tolerance used to determine if a surface should be tessellated or not."))
+	double GeometricTolerance = 0.001;
+
+	/**
+	 * Tolerance used to determine if two surfaces should be stitched.
+	 */
+	UPROPERTY(config, EditAnywhere, BlueprintReadWrite, AdvancedDisplay, Category = "Geometry & Tessellation Options", meta = (Units = cm, ToolTip = "Tolerance used to determine if two surfaces should be stitched.", editCondition = "StitchingTechnique!=EDatasmithCADStitchingTechnique::StitchingNone"))
+	double StitchingTolerance = 0.001;
+
 	bool bUseCADKernel = false;
 
 public:
