Enabled support for dgn extension for the CAD translator. Only available using theEnableUnsupportedCADFormat CVar.

#rnx

diff --git a/Engine/Plugins/Enterprise/DatasmithCADImporter/Source/DatasmithCADTranslator/Private/DatasmithCADTranslator.cpp b/Engine/Plugins/Enterprise/DatasmithCADImporter/Source/DatasmithCADTranslator/Private/DatasmithCADTranslator.cpp
--- a/Engine/Plugins/Enterprise/DatasmithCADImporter/Source/DatasmithCADTranslator/Private/DatasmithCADTranslator.cpp
+++ b/Engine/Plugins/Enterprise/DatasmithCADImporter/Source/DatasmithCADTranslator/Private/DatasmithCADTranslator.cpp
@@ -124,6 +124,7 @@ void FDatasmithCADTranslator::Initialize(FDatasmithTranslatorCapabilities& OutCa
 	{
 		OutCapabilities.SupportedFileFormats.Add(FFileFormatInfo{ TEXT("3mf"), TEXT("3D Manufacturing Format") });
 		OutCapabilities.SupportedFileFormats.Add(FFileFormatInfo{ TEXT("3ds"), TEXT("Autodesk 3DS") });
+		OutCapabilities.SupportedFileFormats.Add(FFileFormatInfo{ TEXT("dae"), TEXT("Collada") });
 		OutCapabilities.SupportedFileFormats.Add(FFileFormatInfo{ TEXT("dwf"), TEXT("Autodesk DWF") });
 		OutCapabilities.SupportedFileFormats.Add(FFileFormatInfo{ TEXT("dwfx"), TEXT("Autodesk DWF") });
 		OutCapabilities.SupportedFileFormats.Add(FFileFormatInfo{ TEXT("nwd"), TEXT("Autodesk Navisworks") });
@@ -131,7 +132,7 @@ void FDatasmithCADTranslator::Initialize(FDatasmithTranslatorCapabilities& OutCa
 		OutCapabilities.SupportedFileFormats.Add(FFileFormatInfo{ TEXT("arc"), TEXT("I-Deas") });
 		OutCapabilities.SupportedFileFormats.Add(FFileFormatInfo{ TEXT("unv"), TEXT("I-Deas") });
 		OutCapabilities.SupportedFileFormats.Add(FFileFormatInfo{ TEXT("pkg"), TEXT("I-Deas") });
-		//OutCapabilities.SupportedFileFormats.Add(FFileFormatInfo{ TEXT("dgn"), TEXT("Microstation") });  // available with Hoops Exchange 2023
+		OutCapabilities.SupportedFileFormats.Add(FFileFormatInfo{ TEXT("dgn"), TEXT("Microstation") });  // available with Hoops Exchange 2023
 		OutCapabilities.SupportedFileFormats.Add(FFileFormatInfo{ TEXT("stl"), TEXT("Stereo Lithography (STL)") });
 		OutCapabilities.SupportedFileFormats.Add(FFileFormatInfo{ TEXT("u3d"), TEXT("U3D (ECMA-363)") });
 		OutCapabilities.SupportedFileFormats.Add(FFileFormatInfo{ TEXT("vda"), TEXT("VDA-FS") });
