Fixed wrong unit conversion

#rnx

diff --git a/Engine/Plugins/Enterprise/DatasmithCADImporter/Source/CADLibrary/Private/CADKernelTools.cpp b/Engine/Plugins/Enterprise/DatasmithCADImporter/Source/CADLibrary/Private/CADKernelTools.cpp
--- a/Engine/Plugins/Enterprise/DatasmithCADImporter/Source/CADLibrary/Private/CADKernelTools.cpp
+++ b/Engine/Plugins/Enterprise/DatasmithCADImporter/Source/CADLibrary/Private/CADKernelTools.cpp
@@ -270,11 +270,10 @@ bool FCADKernelTools::Tessellate(UE::CADKernel::FTopologicalShapeEntity& CADTopo
 	// Tessellate the model
 	TSharedRef<FModelMesh> CADKernelModelMesh = FEntity::MakeShared<FModelMesh>();
 
-	const double GeometricTolerance = TessellationContext.GeometricTolerance * 10; // cm to mm
 	const bool bActivateThinZoneMeshing = FImportParameters::bGActivateThinZoneMeshing;
-	FMesher Mesher(*CADKernelModelMesh, GeometricTolerance, bActivateThinZoneMeshing);
+	FMesher Mesher(*CADKernelModelMesh, TessellationContext.GeometricTolerance, bActivateThinZoneMeshing);
 
-	DefineMeshCriteria(*CADKernelModelMesh, TessellationContext.ImportParameters, GeometricTolerance);
+	DefineMeshCriteria(*CADKernelModelMesh, TessellationContext.ImportParameters, TessellationContext.GeometricTolerance);
 	Mesher.MeshEntity(CADTopologicalEntity);
 
 	return ConvertModelMeshToMeshDescription(TessellationContext, *CADKernelModelMesh, OutMeshDescription);
