Fixed ensure when adding bodies to a CADKernel model

#rnx

diff --git a/Engine/Plugins/Enterprise/DatasmithCADImporter/Source/WireInterface/Private/AliasModelToCADKernelConverter.cpp b/Engine/Plugins/Enterprise/DatasmithCADImporter/Source/WireInterface/Private/AliasModelToCADKernelConverter.cpp
--- a/Engine/Plugins/Enterprise/DatasmithCADImporter/Source/WireInterface/Private/AliasModelToCADKernelConverter.cpp
+++ b/Engine/Plugins/Enterprise/DatasmithCADImporter/Source/WireInterface/Private/AliasModelToCADKernelConverter.cpp
@@ -316,12 +316,6 @@ bool FAliasModelToCADKernelConverter::AddBRep(AlDagNode& DagNode, uint32 SlotID,
 {
 	AlEdge2CADKernelEdge.Empty();
 
-	TSharedRef<UE::CADKernel::FBody> CADKernelBody = UE::CADKernel::FEntity::MakeShared<UE::CADKernel::FBody>();
-	CADKernelBody->SetColorId(SlotID);
-
-	TSharedRef<UE::CADKernel::FShell> CADKernelShell = UE::CADKernel::FEntity::MakeShared<UE::CADKernel::FShell>();
-	CADKernelBody->AddShell(CADKernelShell);
-
 	boolean bAlOrientation;
 	DagNode.getSurfaceOrientation(bAlOrientation);
 	bool bOrientation = !(bool)bAlOrientation;
@@ -332,6 +326,8 @@ bool FAliasModelToCADKernelConverter::AddBRep(AlDagNode& DagNode, uint32 SlotID,
 		DagNode.localTransformationMatrix(AlMatrix);
 	}
 
+	TSharedRef<UE::CADKernel::FShell> CADKernelShell = UE::CADKernel::FEntity::MakeShared<UE::CADKernel::FShell>();
+
 	AlObjectType objectType = DagNode.type();
 	switch (objectType)
 	{
@@ -367,10 +363,21 @@ bool FAliasModelToCADKernelConverter::AddBRep(AlDagNode& DagNode, uint32 SlotID,
 		break;
 	}
 
-	CADKernelBody->CompleteMetaData();
-	CADKernelSession.GetModel().Add(CADKernelBody);
+	if (CADKernelShell->FaceCount() > 0)
+	{
+		TSharedRef<UE::CADKernel::FBody> CADKernelBody = UE::CADKernel::FEntity::MakeShared<UE::CADKernel::FBody>();
+
+		CADKernelBody->SetColorId(SlotID);
+		CADKernelBody->AddShell(CADKernelShell);
+		CADKernelBody->CompleteMetaData();
+
+		CADKernelSession.GetModel().Add(CADKernelBody);
 
-	return CADKernelShell->FaceCount() > 0;
+		return true;
+	}
+
+	// #wire_import: Log that no face was added to the model
+	return false;
 }
 
 bool FAliasModelToCADKernelConverter::Tessellate(const CADLibrary::FMeshParameters& InMeshParameters, FMeshDescription& OutMeshDescription)
@@ -417,19 +424,29 @@ bool FAliasModelToCADKernelConverter::AddGeometry(const CADLibrary::FCADModelGeo
 	{
 		const FBodyNodeGeometry& BodyNodeGeometry = static_cast<const FBodyNodeGeometry&>(Geometry);
 
-		bool bBRepAdded = true;
-		BodyNodeGeometry.BodyNode->IterateOnSurfaceNodes([&](const TAlDagNodePtr<AlSurfaceNode>& SurfaceNode)
+		bool bBodyAdded = true;
+		if (BodyNodeGeometry.BodyNode->HasSurfaceNodes())
 		{
-			bBRepAdded &= AddBRep(*SurfaceNode, BodyNodeGeometry.BodyNode->GetSlotIndex(SurfaceNode.Get()), BodyNodeGeometry.Reference);
-		});
+			bool bBRepAdded = false; // Track that at least one surface was properly added.
+			BodyNodeGeometry.BodyNode->IterateOnSurfaceNodes([&](const TAlDagNodePtr<AlSurfaceNode>& SurfaceNode)
+				{
+					bBRepAdded |= AddBRep(*SurfaceNode, BodyNodeGeometry.BodyNode->GetSlotIndex(SurfaceNode.Get()), BodyNodeGeometry.Reference);
+				});
+			bBodyAdded &= bBRepAdded;
+		}
 
-		BodyNodeGeometry.BodyNode->IterateOnShellNodes([&](const TAlDagNodePtr<AlShellNode>& ShellNode)
+		if (BodyNodeGeometry.BodyNode->HasShellNodes())
 		{
-			bBRepAdded &= AddBRep(*ShellNode, BodyNodeGeometry.BodyNode->GetSlotIndex(ShellNode.Get()), BodyNodeGeometry.Reference);
-		});
+			bool bBRepAdded = false; // Track that at least one shell was properly added.
+			BodyNodeGeometry.BodyNode->IterateOnShellNodes([&](const TAlDagNodePtr<AlShellNode>& ShellNode)
+				{
+					bBRepAdded |= AddBRep(*ShellNode, BodyNodeGeometry.BodyNode->GetSlotIndex(ShellNode.Get()), BodyNodeGeometry.Reference);
+				});
+			bBodyAdded &= bBRepAdded;
+		}
 
-		ensure(bBRepAdded);
-		return bBRepAdded;
+		ensure(bBodyAdded);
+		return bBodyAdded;
 	}
 
 	return false;
diff --git a/Engine/Plugins/Enterprise/DatasmithCADImporter/Source/WireInterface/Private/OpenModelUtils.h b/Engine/Plugins/Enterprise/DatasmithCADImporter/Source/WireInterface/Private/OpenModelUtils.h
--- a/Engine/Plugins/Enterprise/DatasmithCADImporter/Source/WireInterface/Private/OpenModelUtils.h
+++ b/Engine/Plugins/Enterprise/DatasmithCADImporter/Source/WireInterface/Private/OpenModelUtils.h
@@ -328,6 +328,9 @@ namespace UE_DATASMITHWIRETRANSLATOR_NAMESPACE
 
 		int32 GetSlotIndex(const TAlDagNodePtr<AlDagNode>& DagNode);
 
+		bool HasSurfaceNodes() const { return SurfaceNodes.Num() > 0; }
+		bool HasShellNodes() const { return ShellNodes.Num() > 0; }
+
 	private:
 		FString Name;
 		TArray<TAlDagNodePtr<AlSurfaceNode>> SurfaceNodes;
