Addressed miscellaneous crashes:
- Clean up the logic in FTopologicalFace::UpdateBBox and remove potential crashes
- Modified logic in UE::CADKernel::FindLoopIntersectionsWithIso to only return unique intersection values
- Fixed TCurveSamplerAbstract::RunSampling to immediately return if there is no segment to sample
- Fixed issue with import of multiple files in Datasmith: The UnloadScene of the translator was not called until all the sources were imported. This was not correct for translators which use SDK with global settings, i.e. OpenModel.
- Fixed FAliasModelToCADKernelConverter::AddFace to only add a Face if it has at least one loop.
- Made sure debug 'ensure' were not hit in regular builds

#rnx

diff --git a/Engine/Plugins/Enterprise/DatasmithCADImporter/Source/DatasmithWireTranslator/Private/DatasmithWireTranslator.cpp b/Engine/Plugins/Enterprise/DatasmithCADImporter/Source/DatasmithWireTranslator/Private/DatasmithWireTranslator.cpp
--- a/Engine/Plugins/Enterprise/DatasmithCADImporter/Source/DatasmithWireTranslator/Private/DatasmithWireTranslator.cpp
+++ b/Engine/Plugins/Enterprise/DatasmithCADImporter/Source/DatasmithWireTranslator/Private/DatasmithWireTranslator.cpp
@@ -54,6 +54,22 @@ namespace WireTranslator
 
 		return TSharedPtr<IWireInterface>();
 	}
+
+	bool IsFileSupported(const TCHAR* Filename)
+	{
+		TSharedPtr<IWireInterface> WireInterface;
+
+		for (const TPair<uint32, FInterfaceMaker>& Entry : RegisteredInterfaces)
+		{
+			WireInterface = Entry.Value();
+			if (WireInterface->Initialize(Filename))
+			{
+				return true;
+			}
+		}
+
+		return false;
+	}
 }
 
 uint64 IWireInterface::GetRequiredAliasVersion()
@@ -93,12 +109,7 @@ bool FDatasmithWireTranslator::CanTranslate()
 				return false;
 			}
 
-			if (!WireTranslator::GetInterfaceFromFile(nullptr).IsValid())
-			{
-				return false;
-			}
-
-			return true;
+			return RegisteredInterfaces.Num() > 0 ? true : false;
 		}();
 
 		return bCanTranslate;
@@ -127,19 +138,19 @@ void FDatasmithWireTranslator::Initialize(FDatasmithTranslatorCapabilities& OutC
 
 bool FDatasmithWireTranslator::IsSourceSupported(const FDatasmithSceneSource& Source)
 {
-	WireInterface = WireTranslator::GetInterfaceFromFile(*Source.GetSourceFile());
-	return WireInterface.IsValid();
+	return WireTranslator::IsFileSupported(*Source.GetSourceFile());
 }
 
 bool FDatasmithWireTranslator::LoadScene(TSharedRef<IDatasmithScene> OutScene)
 {
-	//WireInterface = WireTranslator::GetInterfaceFromFile(*Filename);
+	ensure(WireImportOptions);
+	WireInterface = WireTranslator::GetInterfaceFromFile(*GetSource().GetSourceFile());
 	if (!WireInterface.IsValid())
 	{
 		return false;
 	}
 
-	//const FString& Filename = GetSource().GetSourceFile();
+	WireInterface->SetImportSettings(WireImportOptions->Settings);
 
 	UE_LOG(LogDatasmithWireTranslator, Display, TEXT("CAD translation [%s]."), *GetSource().GetSourceFile());
 	UE_LOG(LogDatasmithWireTranslator, Display, TEXT(" - Parsing Library:      %s"), TEXT("Alias"));
@@ -205,16 +216,15 @@ void FDatasmithWireTranslator::SetSceneImportOptions(const TArray<TObjectPtr<UDa
 {
 	FParametricSurfaceTranslator::SetSceneImportOptions(Options);
 
-	if (ensure(WireInterface.IsValid()))
+	WireImportOptions = nullptr;
+
+	for (const TObjectPtr<UDatasmithOptionsBase>& OptionPtr : Options)
 	{
-		for (const TObjectPtr<UDatasmithOptionsBase>& OptionPtr : Options)
+		if (UDatasmithWireOptions* ImportOptions = Cast<UDatasmithWireOptions>(OptionPtr))
 		{
-			if (UDatasmithWireOptions* ImportOptions = Cast<UDatasmithWireOptions>(OptionPtr))
-			{
-				ImportOptions->SaveConfig(CPF_Config);
-				WireInterface->SetImportSettings(ImportOptions->Settings);
-				CommonTessellationOptions = ImportOptions->Settings;
-			}
+			ImportOptions->SaveConfig(CPF_Config);
+			WireImportOptions = ImportOptions;
+			CommonTessellationOptions = ImportOptions->Settings;
 		}
 	}
 }
diff --git a/Engine/Plugins/Enterprise/DatasmithCADImporter/Source/DatasmithWireTranslator/Private/DatasmithWireTranslator.h b/Engine/Plugins/Enterprise/DatasmithCADImporter/Source/DatasmithWireTranslator/Private/DatasmithWireTranslator.h
--- a/Engine/Plugins/Enterprise/DatasmithCADImporter/Source/DatasmithWireTranslator/Private/DatasmithWireTranslator.h
+++ b/Engine/Plugins/Enterprise/DatasmithCADImporter/Source/DatasmithWireTranslator/Private/DatasmithWireTranslator.h
@@ -55,4 +55,5 @@ private:
 
 private:
 	TSharedPtr<IWireInterface> WireInterface;
+	TObjectPtr<UDatasmithWireOptions> WireImportOptions;
 };
diff --git a/Engine/Plugins/Enterprise/DatasmithCADImporter/Source/WireInterface/Private/AliasModelToCADKernelConverter.cpp b/Engine/Plugins/Enterprise/DatasmithCADImporter/Source/WireInterface/Private/AliasModelToCADKernelConverter.cpp
--- a/Engine/Plugins/Enterprise/DatasmithCADImporter/Source/WireInterface/Private/AliasModelToCADKernelConverter.cpp
+++ b/Engine/Plugins/Enterprise/DatasmithCADImporter/Source/WireInterface/Private/AliasModelToCADKernelConverter.cpp
@@ -207,7 +207,7 @@ void FAliasModelToCADKernelConverter::LinkEdgesLoop(const AlTrimBoundary& TrimBo
 			continue;
 		}
 
-		ensure(&Loop == (*Edge)->GetLoop());
+		ensureWire(&Loop == (*Edge)->GetLoop());
 
 		// Link edges
 		TAlObjectPtr<AlTrimCurve> TwinCurve(TrimCurve->getTwinCurve());
@@ -290,7 +290,13 @@ void FAliasModelToCADKernelConverter::AddFace(const AlSurface& Surface, EAliasOb
 	{
 		TSharedRef<UE::CADKernel::FTopologicalFace> Face = UE::CADKernel::FEntity::MakeShared<UE::CADKernel::FTopologicalFace>(CADKernelSurface);
 		Face->ApplyNaturalLoops();
-		Shell->Add(Face, bInOrientation ? UE::CADKernel::EOrientation::Front : UE::CADKernel::EOrientation::Back);
+		// Surface can be too thin based on tolerances. Skip it
+		// #cadkernel_check: Warn user a surface was too thin
+		if (Face->GetLoops().Num() > 0)
+		{
+			Shell->Add(Face, bInOrientation ? UE::CADKernel::EOrientation::Front : UE::CADKernel::EOrientation::Back);
+		}
+		// #wire_import: Log that this face was not added
 	}
 }
 
@@ -445,7 +451,7 @@ bool FAliasModelToCADKernelConverter::AddGeometry(const CADLibrary::FCADModelGeo
 			bBodyAdded &= bBRepAdded;
 		}
 
-		ensure(bBodyAdded);
+		ensureWire(bBodyAdded);
 		return bBodyAdded;
 	}
 
diff --git a/Engine/Plugins/Enterprise/DatasmithCADImporter/Source/WireInterface/Private/OpenModelUtils.cpp b/Engine/Plugins/Enterprise/DatasmithCADImporter/Source/WireInterface/Private/OpenModelUtils.cpp
--- a/Engine/Plugins/Enterprise/DatasmithCADImporter/Source/WireInterface/Private/OpenModelUtils.cpp
+++ b/Engine/Plugins/Enterprise/DatasmithCADImporter/Source/WireInterface/Private/OpenModelUtils.cpp
@@ -80,7 +80,7 @@ namespace UE_DATASMITHWIRETRANSLATOR_NAMESPACE
 			}
 		}
 
-		ensure(false);
+		ensureWire(false);
 		return 0;
 	}
 
@@ -144,7 +144,7 @@ namespace UE_DATASMITHWIRETRANSLATOR_NAMESPACE
 				Hash = HashCombine(Hash, ShellNode.GetHash());
 			}
 
-			ensure(LayerSet.Num() == 1);
+			ensureWire(LayerSet.Num() == 1);
 			TAlObjectPtr<AlLayer> GeomLayer = LayerSet.Array()[0];
 			if (GeomLayer != Layer)
 			{
diff --git a/Engine/Plugins/Enterprise/DatasmithCADImporter/Source/WireInterface/Private/OpenModelUtils.h b/Engine/Plugins/Enterprise/DatasmithCADImporter/Source/WireInterface/Private/OpenModelUtils.h
--- a/Engine/Plugins/Enterprise/DatasmithCADImporter/Source/WireInterface/Private/OpenModelUtils.h
+++ b/Engine/Plugins/Enterprise/DatasmithCADImporter/Source/WireInterface/Private/OpenModelUtils.h
@@ -43,6 +43,13 @@ struct FMeshDescription;
 #define UNIT_CONVERSION_CM_TO_MM		10.
 #define UE_TO_CADKERNEL(Distance)		(Distance * 10.)
 
+#define WIRE_ENSURE_ENABLED 0
+#if WIRE_ENSURE_ENABLED
+#define ensureWire(InExpression) ensure(InExpression)
+#else
+#define ensureWire(InExpression)
+#endif
+
 namespace UE_DATASMITHWIRETRANSLATOR_NAMESPACE
 {
 	typedef double AlMatrix4x4[4][4];
diff --git a/Engine/Plugins/Enterprise/DatasmithImporter/Source/DatasmithImporter/Private/DatasmithImportFactory.cpp b/Engine/Plugins/Enterprise/DatasmithImporter/Source/DatasmithImporter/Private/DatasmithImportFactory.cpp
--- a/Engine/Plugins/Enterprise/DatasmithImporter/Source/DatasmithImporter/Private/DatasmithImportFactory.cpp
+++ b/Engine/Plugins/Enterprise/DatasmithImporter/Source/DatasmithImporter/Private/DatasmithImportFactory.cpp
@@ -533,12 +533,26 @@ UObject* UDatasmithImportFactory::CreateFromExternalSource(UClass* InClass, UObj
 
 	FDatasmithImportContext ImportContext(InExternalSource, !IsAutomatedImport(), GetLoggerName(), GetDisplayName());
 
+	TFunction< UObject* ()> ExitReturn = [&InExternalSource, &ImportContext]() -> UObject*
+		{
+			if (InExternalSource->GetDatasmithScene())
+			{
+				if (TSharedPtr<IDatasmithTranslator> Translator = InExternalSource->GetAssetTranslator())
+				{
+					// Translators may be using global content. Make sure it is flushed before completing the import
+					Translator->UnloadScene();
+				}
+			}
+
+			return ImportContext.SceneAsset;
+		};
+
 	const bool bIsSilent = IsAutomatedImport() || !bShowOptions;
 	if (!ImportContext.Init(ImportPath, InFlags, InWarn, ImportSettingsJson, bIsSilent))
 	{
 		bOperationCanceled = true;
 		bOutOperationCanceled = true;
-		return nullptr;
+		return ExitReturn();
 	}
 
 	// Collect start time to log amount of time spent to import incoming file
@@ -551,20 +565,21 @@ UObject* UDatasmithImportFactory::CreateFromExternalSource(UClass* InClass, UObj
 	else
 	{
 		UE_LOG(LogDatasmithImport, Warning, TEXT("Datasmith import error: Scene translation failure. Abort import."));
-		return nullptr;
+		return ExitReturn();
 	}
 
 	if (!Import( ImportContext ))
 	{
 		bOperationCanceled = true;
 		bOutOperationCanceled = true;
+		ImportContext.SceneAsset = nullptr;
 		UE_LOG(LogDatasmithImport, Warning, TEXT("Datasmith import error. Abort import."));
-		return nullptr;
+		return ExitReturn();
 	}
 
 	DatasmithImportFactoryImpl::ReportImportStats(ImportContext, StartTime);
 
-	return ImportContext.SceneAsset;
+	return ExitReturn();
 }
 
 bool UDatasmithImportFactory::Import( FDatasmithImportContext& ImportContext )
diff --git a/Engine/Source/Runtime/Datasmith/CADKernel/Private/CADKernel/Geo/Sampler/SamplerAbstract.h b/Engine/Source/Runtime/Datasmith/CADKernel/Private/CADKernel/Geo/Sampler/SamplerAbstract.h
--- a/Engine/Source/Runtime/Datasmith/CADKernel/Private/CADKernel/Geo/Sampler/SamplerAbstract.h
+++ b/Engine/Source/Runtime/Datasmith/CADKernel/Private/CADKernel/Geo/Sampler/SamplerAbstract.h
@@ -149,6 +149,13 @@ protected:
 	 */
 	void RunSampling()
 	{
+		if (IsOptimalSegments.Num() == 0)
+		{
+			// #cadkernel_check: What does this mean that IsOptimalSegments is empty.
+			//                   Why only one CandidatePoints?
+			return;
+		}
+
 		// < 100000 : check to avoid that the process loops endlessly
 		int32 CandidatePointsCount = CandidatePoints.Size();
 		while (CandidatePointsCount && Sampling.Coordinates.Num() < 100000)
@@ -187,7 +194,7 @@ protected:
 
 					// if the next candidate point is biggest than the next sampling point, no chord check can be done. 
 					// an intermediate point is added. 
-					// the next segment is processed 
+					// the next segment is processed
 					if ((FirstCandidateIndex >= CandidatePointsCount) || (FirstCandidateIndex < CandidatePointsCount && TmpPolylineCoordinates[FirstCandidateIndex] > EndStudySegment.Polyline->Coordinates[EndStudySegment.Index]))
 					{
 						AddIntermediateCoordinates(Sampling.Coordinates[StartSamplingSegmentIndex], Sampling.Coordinates[EndSamplingSegmentIndex], 1);
diff --git a/Engine/Source/Runtime/Datasmith/CADKernel/Private/CADKernel/Geo/Sampler/SamplerOnParam.h b/Engine/Source/Runtime/Datasmith/CADKernel/Private/CADKernel/Geo/Sampler/SamplerOnParam.h
--- a/Engine/Source/Runtime/Datasmith/CADKernel/Private/CADKernel/Geo/Sampler/SamplerOnParam.h
+++ b/Engine/Source/Runtime/Datasmith/CADKernel/Private/CADKernel/Geo/Sampler/SamplerOnParam.h
@@ -93,34 +93,6 @@ protected:
 			}
 		}
 
-#ifdef DEBUG_CHECKCHORDERROR
-		if (bDisplay)
-		{
-			{
-				int32 Index = 0;
-				F3DDebugSession G(FString::Printf(TEXT("CheckChord %f %d"), MaxChord, FirstIndex));
-
-				UE::CADKernel::Display(APoint, EVisuProperty::BluePoint);
-				UE::CADKernel::Display(BPoint, EVisuProperty::BluePoint);
-
-				for (int32 Index = FirstIndex; Index < EndIndex; ++Index)
-				{
-					const FPoint& Middle = CandidatePoints.GetPoints()[Index];
-					UE::CADKernel::Display(Middle, EVisuProperty::YellowPoint);
-				}
-
-				DisplaySegment(APoint, BPoint, 0, EVisuProperty::BlueCurve);
-				DisplaySegment(APoint, CandidatePoints.GetPoints()[FirstIndex], 0, EVisuProperty::YellowCurve);
-				DisplaySegment(BPoint, CandidatePoints.GetPoints()[EndIndex - 1], 0, EVisuProperty::YellowCurve);
-
-				for (int32 Index = FirstIndex; Index < EndIndex - 1; ++Index)
-				{
-					DisplaySegment(CandidatePoints.GetPoints()[Index], CandidatePoints.GetPoints()[Index + 1], 0, EVisuProperty::YellowCurve);
-				}
-			}
-		}
-#endif
-
 		int32 SegmentNumToRespectParametrizationError = (int32)(MaxParametrizationError / DesiredParameterizationError);
 		if (SegmentNumToRespectParametrizationError > 2)
 		{
diff --git a/Engine/Source/Runtime/Datasmith/CADKernel/Private/CADKernel/Math/Geometry.cpp b/Engine/Source/Runtime/Datasmith/CADKernel/Private/CADKernel/Math/Geometry.cpp
--- a/Engine/Source/Runtime/Datasmith/CADKernel/Private/CADKernel/Math/Geometry.cpp
+++ b/Engine/Source/Runtime/Datasmith/CADKernel/Private/CADKernel/Math/Geometry.cpp
@@ -108,7 +108,8 @@ double ComputeCurvature(const FPoint& Normal, const FPoint& Gradient, const FPoi
 
 void FindLoopIntersectionsWithIso(const EIso Iso, const double IsoParameter, const TArray<TArray<FPoint2D>>& Loops, TArray<double>& OutIntersections)
 {
-	OutIntersections.Empty(8);
+	TArray<double> LocalIntersections;
+	LocalIntersections.Reserve(8);
 
 	const int32 UIndex = Iso == EIso::IsoU ? 0 : 1;
 	const int32 VIndex = Iso == EIso::IsoU ? 1 : 0;
@@ -117,8 +118,8 @@ void FindLoopIntersectionsWithIso(const EIso Iso, const double IsoParameter, con
 	{
 		if (IsoParameter > Point1[UIndex] && IsoParameter <= Point2[UIndex])
 		{
-			double Intersection = (IsoParameter - Point1[UIndex]) / (Point2[UIndex] - Point1[UIndex]) * (Point2[VIndex] - Point1[VIndex]) + Point1[VIndex];
-			OutIntersections.Add((IsoParameter - Point1[UIndex]) / (Point2[UIndex] - Point1[UIndex]) * (Point2[VIndex] - Point1[VIndex]) + Point1[VIndex]);
+			const double Intersection = (IsoParameter - Point1[UIndex]) / (Point2[UIndex] - Point1[UIndex]) * (Point2[VIndex] - Point1[VIndex]) + Point1[VIndex];
+			LocalIntersections.Add(Intersection);
 		}
 	};
 
@@ -141,7 +142,25 @@ void FindLoopIntersectionsWithIso(const EIso Iso, const double IsoParameter, con
 			Point1 = &Point2;
 		}
 	}
-	Algo::Sort(OutIntersections);
+
+	if (LocalIntersections.Num() == 0)
+	{
+		return;
+	}
+
+	Algo::Sort(LocalIntersections);
+
+	// Remove any duplicates
+	OutIntersections.Empty(LocalIntersections.Num());
+	OutIntersections.Add(LocalIntersections[0]);
+	for (int32 Index = 1; Index < LocalIntersections.Num(); ++Index)
+	{
+		if (FMath::IsNearlyEqual(LocalIntersections[Index], OutIntersections.Last(), UE_DOUBLE_SMALL_NUMBER))
+		{
+			continue;
+		}
+		OutIntersections.Add(LocalIntersections[Index]);
+	}
 }
 
 bool DoIntersect(const FSegment2D& SegmentAB, const FSegment2D& SegmentCD, TFunction<bool(double, double, double, double)> DoCoincidentSegmentsIntersect, const double Min, const double Max)
diff --git a/Engine/Source/Runtime/Datasmith/CADKernel/Private/CADKernel/Topo/TopologicalFace.cpp b/Engine/Source/Runtime/Datasmith/CADKernel/Private/CADKernel/Topo/TopologicalFace.cpp
--- a/Engine/Source/Runtime/Datasmith/CADKernel/Private/CADKernel/Topo/TopologicalFace.cpp
+++ b/Engine/Source/Runtime/Datasmith/CADKernel/Private/CADKernel/Topo/TopologicalFace.cpp
@@ -79,9 +79,7 @@ void FTopologicalFace::UpdateBBox(int32 IsoCount, const double ApproximationFact
 			TArray<double> Intersections;
 			FindLoopIntersectionsWithIso(IsoType, Coordinate, BoundaryApproximation, Intersections);
 			int32 IntersectionCount = Intersections.Num();
-			// #cadkernel_check: Why IntersectionCount can be less than 2 or Intersections has same extremities
-			//if (IntersectionCount < 2 || FMath::IsNearlyEqual(Intersections[0], Intersections.Last(), UE_DOUBLE_SMALL_NUMBER))
-			if (IntersectionCount == 0)
+			if (IntersectionCount < 2)
 			{
 				continue;
 			}
@@ -97,12 +95,10 @@ void FTopologicalFace::UpdateBBox(int32 IsoCount, const double ApproximationFact
 				TArray<FPoint> SubPolyline;
 				FLinearBoundary IntersectionBoundary(Intersections[IntersectionCount - 1], CurveBounds.GetMax());
 
-#ifdef DEBUG_GET_BBOX2
-				Polyline.GetSubPolyline(Boundary, EOrientation::Front, SubPolyline);
-				Draw(SubPolyline, EVisuProperty::Iso);
-#endif
-
-				Polyline.UpdateSubPolylineBBox(IntersectionBoundary, IsoBBox);
+				if (Polyline.Size() > 1)
+				{
+					Polyline.UpdateSubPolylineBBox(IntersectionBoundary, IsoBBox);
+				}
 
 				Intersections.Pop();
 				IntersectionCount--;
@@ -113,25 +109,16 @@ void FTopologicalFace::UpdateBBox(int32 IsoCount, const double ApproximationFact
 				continue;
 			}
 
-			for (int32 ISection = 0; ISection < IntersectionCount; ISection += 2)
+			if (IntersectionCount > 1)
 			{
-				TArray<FPoint> SubPolyline;
-				FLinearBoundary IntersectionBoundary(Intersections[ISection], Intersections[ISection + 1]);
-
-#ifdef DEBUG_GET_BBOX2
-				Polyline.GetSubPolyline(Boundary, EOrientation::Front, SubPolyline);
-				Draw(SubPolyline, EVisuProperty::Iso);
-#endif
-				Polyline.UpdateSubPolylineBBox(IntersectionBoundary, IsoBBox);
-			}
+				for (int32 ISection = 0; ISection < IntersectionCount; ISection += 2)
+				{
+					TArray<FPoint> SubPolyline;
+					FLinearBoundary IntersectionBoundary(Intersections[ISection], Intersections[ISection + 1]);
 
-#ifdef DEBUG_GET_BBOX2
-			for (int32 Index = 0; Index < 3; ++Index)
-			{
-				UE::CADKernel::DisplayPoint(IsoBBox.MaxPoints[Index], EVisuProperty::YellowPoint);
-				UE::CADKernel::DisplayPoint(IsoBBox.MinPoints[Index], EVisuProperty::YellowPoint);
+					Polyline.UpdateSubPolylineBBox(IntersectionBoundary, IsoBBox);
+				}
 			}
-#endif
 
 			BBox.Update(IsoBBox, IsoType, Coordinate);
 		}
@@ -161,6 +148,7 @@ void FTopologicalFace::UpdateBBox(int32 IsoCount, const double ApproximationFact
 #endif
 
 }
+#pragma optimize("",on)
 
 
 void FTopologicalFace::ApplyNaturalLoops()
