Fixed build issue on Linux

#rnx

diff --git a/Engine/Plugins/Enterprise/DatasmithCADImporter/Source/CADKernelSurface/CADKernelSurface.Build.cs b/Engine/Plugins/Enterprise/DatasmithCADImporter/Source/CADKernelSurface/CADKernelSurface.Build.cs
--- a/Engine/Plugins/Enterprise/DatasmithCADImporter/Source/CADKernelSurface/CADKernelSurface.Build.cs
+++ b/Engine/Plugins/Enterprise/DatasmithCADImporter/Source/CADKernelSurface/CADKernelSurface.Build.cs
@@ -8,8 +8,6 @@ namespace UnrealBuildTool.Rules
 	{
 		public CADKernelSurface(ReadOnlyTargetRules Target) : base(Target)
 		{
-			OptimizeCode = CodeOptimization.Never;
-
 			PublicDependencyModuleNames.AddRange(
 				new string[]
 				{
