Fixed static analysis issue

#rnx

diff --git a/Engine/Plugins/Enterprise/DatasmithCADImporter/Source/CADKernelSurface/Public/CADModelToCADKernelConverterBase.h b/Engine/Plugins/Enterprise/DatasmithCADImporter/Source/CADKernelSurface/Public/CADModelToCADKernelConverterBase.h
--- a/Engine/Plugins/Enterprise/DatasmithCADImporter/Source/CADKernelSurface/Public/CADModelToCADKernelConverterBase.h
+++ b/Engine/Plugins/Enterprise/DatasmithCADImporter/Source/CADKernelSurface/Public/CADModelToCADKernelConverterBase.h
@@ -24,7 +24,7 @@ class FCADModelToCADKernelConverterBase : public CADLibrary::ICADModelConverter
 public:
 
 	FCADModelToCADKernelConverterBase(const CADLibrary::FImportParameters& InImportParameters)
-		: CADKernelSession(GeometricTolerance)
+		: CADKernelSession(0.01)
 		, ImportParameters(InImportParameters)
 	{
 	}
