Added support for users having AliasStudio 2024.1 and/or 2025.0 installed
Removed CVar ds.CADTranslator.Alias.LayersAsActors. A new option has been added to allow users to select that feature from the translator's import dialog
Removed CVar ds.CADTranslator.Alias.SewByMaterial.
Revisited management of memory allocated by the OpenModel SDK - WIP

#rnx

diff --git a/Engine/Plugins/Enterprise/DatasmithCADImporter/DatasmithCADImporter.uplugin b/Engine/Plugins/Enterprise/DatasmithCADImporter/DatasmithCADImporter.uplugin
--- a/Engine/Plugins/Enterprise/DatasmithCADImporter/DatasmithCADImporter.uplugin
+++ b/Engine/Plugins/Enterprise/DatasmithCADImporter/DatasmithCADImporter.uplugin
@@ -134,6 +134,22 @@
 				"Win64"
 			]
 		},
+		{
+			"Name": "WireInterface2024_1",
+			"Type": "Editor",
+			"LoadingPhase": "Default",
+			"PlatformAllowList": [
+				"Win64"
+			]
+		},
+		{
+			"Name": "WireInterface2025_0",
+			"Type": "Editor",
+			"LoadingPhase": "Default",
+			"PlatformAllowList": [
+				"Win64"
+			]
+		},
 		{
 			"Name": "ParametricSurfaceExtension",
 			"Type": "Editor",
diff --git a/Engine/Plugins/Enterprise/DatasmithCADImporter/Source/DatasmithWireTranslator/Private/DatasmithWireTranslator.cpp b/Engine/Plugins/Enterprise/DatasmithCADImporter/Source/DatasmithWireTranslator/Private/DatasmithWireTranslator.cpp
--- a/Engine/Plugins/Enterprise/DatasmithCADImporter/Source/DatasmithWireTranslator/Private/DatasmithWireTranslator.cpp
+++ b/Engine/Plugins/Enterprise/DatasmithCADImporter/Source/DatasmithWireTranslator/Private/DatasmithWireTranslator.cpp
@@ -28,12 +28,6 @@ static TAutoConsoleVariable<bool> CVarAliasUseNative(
 	TEXT("If true, the AlaisStudio tessellator will be used. Default false.\n"),
 	ECVF_Default);
 
-TAutoConsoleVariable<bool> GAliasLayerAsActor(
-	TEXT("ds.CADTranslator.Alias.LayersAsActors"),
-	false,
-	TEXT("If true, the first level of actors in the outliner is the layers. Default is false\n"),
-	ECVF_Default);
-
 static TMap<uint32, FInterfaceMaker> RegisteredInterfaces;
 static uint64 AliasVersion = 0;
 
@@ -152,39 +146,6 @@ bool FDatasmithWireTranslator::LoadScene(TSharedRef<IDatasmithScene> OutScene)
 
 	WireInterface->SetImportSettings(WireImportOptions->Settings);
 
-	UE_LOG(LogDatasmithWireTranslator, Display, TEXT("CAD translation [%s]."), *GetSource().GetSourceFile());
-	UE_LOG(LogDatasmithWireTranslator, Display, TEXT(" - Parsing Library:      %s"), TEXT("Alias"));
-	UE_LOG(LogDatasmithWireTranslator, Display, TEXT(" - Tessellation Library: %s"), CommonTessellationOptions.bUseCADKernel ? TEXT("CADKernel") : TEXT("TechSoft"));
-	UE_LOG(LogDatasmithWireTranslator, Display, TEXT(" - Import parameters:"));
-
-	UE_LOG(LogDatasmithWireTranslator, Display, TEXT("     - ChordTolerance:     %lf"), CommonTessellationOptions.ChordTolerance);
-	UE_LOG(LogDatasmithWireTranslator, Display, TEXT("     - MaxEdgeLength:      %lf"), CommonTessellationOptions.MaxEdgeLength);
-	UE_LOG(LogDatasmithWireTranslator, Display, TEXT("     - MaxNormalAngle:     %lf"), CommonTessellationOptions.NormalTolerance);
-	
-	FString StitchingTechnique;
-	switch (CommonTessellationOptions.StitchingTechnique)
-	{
-	case EDatasmithCADStitchingTechnique::StitchingHeal:
-		StitchingTechnique = TEXT("Heal");
-		break;
-	case EDatasmithCADStitchingTechnique::StitchingSew:
-		StitchingTechnique = TEXT("Sew");
-		break;
-	default:
-		StitchingTechnique = TEXT("None");
-		break;
-	}
-	UE_LOG(LogDatasmithWireTranslator, Display, TEXT("     - StitchingTechnique: %s"), *StitchingTechnique);
-
-	if (CommonTessellationOptions.bUseCADKernel)
-	{
-		UE_LOG(LogDatasmithWireTranslator, Display, TEXT("     - Stitching Options:"));
-		//UE_LOG(LogDatasmithWireTranslator, Display, TEXT("         - ForceSew:              %s"), CADLibrary::FImportParameters::bGStitchingForceSew ? TEXT("True") : TEXT("False"));
-		//UE_LOG(LogDatasmithWireTranslator, Display, TEXT("         - RemoveThinFaces:       %s"), CADLibrary::FImportParameters::bGStitchingRemoveThinFaces ? TEXT("True") : TEXT("False"));
-		//UE_LOG(LogDatasmithWireTranslator, Display, TEXT("         - RemoveDuplicatedFaces: %s"), CADLibrary::FImportParameters::bGStitchingRemoveDuplicatedFaces ? TEXT("True") : TEXT("False"));
-		//UE_LOG(LogDatasmithWireTranslator, Display, TEXT("         - ForceFactor:           %f"), CADLibrary::FImportParameters::GStitchingForceFactor);
-	}
-
 	static const FString CacheRootDir = FPaths::ConvertRelativePathToFull(FPaths::Combine(FPaths::ProjectSavedDir(), TEXT("WireTranslator"), TEXT("Cache")));
 	
 	const FString OutputPath = FPaths::Combine(CacheRootDir, GetSource().GetSceneName());
diff --git a/Engine/Plugins/Enterprise/DatasmithCADImporter/Source/WireInterface/Private/AliasModelToCADKernelConverter.cpp b/Engine/Plugins/Enterprise/DatasmithCADImporter/Source/WireInterface/Private/AliasModelToCADKernelConverter.cpp
--- a/Engine/Plugins/Enterprise/DatasmithCADImporter/Source/WireInterface/Private/AliasModelToCADKernelConverter.cpp
+++ b/Engine/Plugins/Enterprise/DatasmithCADImporter/Source/WireInterface/Private/AliasModelToCADKernelConverter.cpp
@@ -50,16 +50,19 @@
 #include "Windows/HideWindowsPlatformTypes.h"
 #endif
 
+#define LOCTEXT_NAMESPACE "WireInterface"
+
 namespace UE_DATASMITHWIRETRANSLATOR_NAMESPACE
 {
+using namespace UE::CADKernel;
+using ECADKernelSewOption = UE::CADKernel::ESewOption;
 
 namespace AliasToCADKernelUtils
 {
-
 	template<typename Surface_T>
-	TSharedPtr<UE::CADKernel::FSurface> AddNURBSSurface(double GeometricTolerance, Surface_T& AliasSurface, EAliasObjectReference InObjectReference, const AlMatrix4x4& InAlMatrix)
+	TSharedPtr<FSurface> AddNURBSSurface(double GeometricTolerance, Surface_T& AliasSurface, EAliasObjectReference InObjectReference, const AlMatrix4x4& InAlMatrix)
 	{
-		UE::CADKernel::FNurbsSurfaceHomogeneousData NURBSData;
+		FNurbsSurfaceHomogeneousData NURBSData;
 		NURBSData.bSwapUV= true;
 		NURBSData.bIsRational = true;
 
@@ -109,7 +112,7 @@ namespace AliasToCADKernelUtils
 			Poles[Index + 2] *= UNIT_CONVERSION_CM_TO_MM;
 		}
 
-		return UE::CADKernel::FSurface::MakeNurbsSurface(GeometricTolerance, NURBSData);
+		return FSurface::MakeNurbsSurface(GeometricTolerance, NURBSData);
 	}
 }
 
@@ -119,9 +122,9 @@ FAliasModelToCADKernelConverter::FAliasModelToCADKernelConverter(const FDatasmit
 	SetTolerances(Options.GetGeometricTolerance(true), Options.GetStitchingTolerance(true));
 }
 
-TSharedPtr<UE::CADKernel::FTopologicalEdge> FAliasModelToCADKernelConverter::AddEdge(const AlTrimCurve& AliasTrimCurve, TSharedPtr<UE::CADKernel::FSurface>& CarrierSurface)
+TSharedPtr<FTopologicalEdge> FAliasModelToCADKernelConverter::AddEdge(const AlTrimCurve& AliasTrimCurve, TSharedPtr<FSurface>& CarrierSurface)
 {
-	UE::CADKernel::FNurbsCurveData NurbsCurveData;
+	FNurbsCurveData NurbsCurveData;
 
 	NurbsCurveData.Degree = AliasTrimCurve.degree();
 	int32 ControlPointCount = AliasTrimCurve.numberOfCVs();
@@ -150,13 +153,13 @@ TSharedPtr<UE::CADKernel::FTopologicalEdge> FAliasModelToCADKernelConverter::Add
 		NurbsCurveData.Poles[Index].Z = 0;
 	}
 
-	TSharedPtr<UE::CADKernel::FCurve> Nurbs = UE::CADKernel::FCurve::MakeNurbsCurve(NurbsCurveData);
+	TSharedPtr<FCurve> Nurbs = FCurve::MakeNurbsCurve(NurbsCurveData);
 
-	TSharedRef<UE::CADKernel::FRestrictionCurve> RestrictionCurve = UE::CADKernel::FCurve::MakeShared<UE::CADKernel::FRestrictionCurve>(CarrierSurface.ToSharedRef(), Nurbs.ToSharedRef());
-	TSharedPtr<UE::CADKernel::FTopologicalEdge> Edge = UE::CADKernel::FTopologicalEdge::Make(RestrictionCurve);
+	TSharedRef<FRestrictionCurve> RestrictionCurve = FCurve::MakeShared<FRestrictionCurve>(CarrierSurface.ToSharedRef(), Nurbs.ToSharedRef());
+	TSharedPtr<FTopologicalEdge> Edge = FTopologicalEdge::Make(RestrictionCurve);
 	if (!Edge.IsValid())
 	{
-		return TSharedPtr<UE::CADKernel::FTopologicalEdge>();
+		return TSharedPtr<FTopologicalEdge>();
 	}
 
 	// Only TrimCurve with twin need to be in the map used in LinkEdgesLoop 
@@ -169,37 +172,44 @@ TSharedPtr<UE::CADKernel::FTopologicalEdge> FAliasModelToCADKernelConverter::Add
 	return Edge;
 }
 
-TSharedPtr<UE::CADKernel::FTopologicalLoop> FAliasModelToCADKernelConverter::AddLoop(const AlTrimBoundary& TrimBoundary, TSharedPtr<UE::CADKernel::FSurface>& CarrierSurface, const bool bIsExternal)
+TSharedPtr<FTopologicalLoop> FAliasModelToCADKernelConverter::AddLoop(const AlTrimBoundary& TrimBoundary, TSharedPtr<FSurface>& CarrierSurface, const bool bIsExternal)
 {
-	TArray<TSharedPtr<UE::CADKernel::FTopologicalEdge>> Edges;
+	TArray<TSharedPtr<FTopologicalEdge>> Edges;
 	TArray<UE::CADKernel::EOrientation> Directions;
 
-	for (TAlObjectPtr<AlTrimCurve> TrimCurve(TrimBoundary.firstCurve()); TrimCurve.IsValid(); TrimCurve = TAlObjectPtr<AlTrimCurve>(TrimCurve->nextCurve()))
+	TAlObjectPtr<AlTrimCurve> TrimCurve(TrimBoundary.firstCurve());
+	statusCode Status = TrimCurve ? sSuccess : sObjectNotFound;
+	while (Status == sSuccess)
 	{
-		TSharedPtr<UE::CADKernel::FTopologicalEdge> Edge = AddEdge(*TrimCurve, CarrierSurface);
+		TSharedPtr<FTopologicalEdge> Edge = AddEdge(*TrimCurve, CarrierSurface);
 		if (Edge.IsValid())
 		{
 			Edges.Add(Edge);
 			Directions.Emplace(UE::CADKernel::EOrientation::Front);
 		}
+
+		Status = TrimCurve->nextCurveD();
 	}
 
 	if (Edges.Num() == 0)
 	{
-		return TSharedPtr<UE::CADKernel::FTopologicalLoop>();
+		return TSharedPtr<FTopologicalLoop>();
 	}
 
-	TSharedPtr<UE::CADKernel::FTopologicalLoop> Loop = UE::CADKernel::FTopologicalLoop::Make(Edges, Directions, bIsExternal, GeometricTolerance);
+	TSharedPtr<FTopologicalLoop> Loop = FTopologicalLoop::Make(Edges, Directions, bIsExternal, GeometricTolerance);
 	return Loop;
 }
 
-void FAliasModelToCADKernelConverter::LinkEdgesLoop(const AlTrimBoundary& TrimBoundary, UE::CADKernel::FTopologicalLoop& Loop)
+void FAliasModelToCADKernelConverter::LinkEdgesLoop(const AlTrimBoundary& TrimBoundary, FTopologicalLoop& Loop)
 {
-	for (TAlObjectPtr<AlTrimCurve> TrimCurve(TrimBoundary.firstCurve()); TrimCurve.IsValid(); TrimCurve = TAlObjectPtr<AlTrimCurve>(TrimCurve->nextCurve()))
+	TAlObjectPtr<AlTrimCurve> TrimCurve(TrimBoundary.firstCurve());
+	statusCode Status = TrimCurve ? sSuccess : sObjectNotFound;
+	while (Status == sSuccess)
 	{
-		TSharedPtr<UE::CADKernel::FTopologicalEdge>* Edge = AlEdge2CADKernelEdge.Find(TrimCurve->fSpline);
+		TSharedPtr<FTopologicalEdge>* Edge = AlEdge2CADKernelEdge.Find(TrimCurve->fSpline);
 		if (!Edge || !Edge->IsValid() || (*Edge)->IsDeleted() || (*Edge)->IsDegenerated())
 		{
+			Status = TrimCurve->nextCurveD();
 			continue;
 		}
 
@@ -209,7 +219,7 @@ void FAliasModelToCADKernelConverter::LinkEdgesLoop(const AlTrimBoundary& TrimBo
 		TAlObjectPtr<AlTrimCurve> TwinCurve(TrimCurve->getTwinCurve());
 		if (TwinCurve.IsValid())
 		{
-			if (TSharedPtr<UE::CADKernel::FTopologicalEdge>* TwinEdge = AlEdge2CADKernelEdge.Find(TwinCurve->fSpline))
+			if (TSharedPtr<FTopologicalEdge>* TwinEdge = AlEdge2CADKernelEdge.Find(TwinCurve->fSpline))
 			{
 				if (TwinEdge->IsValid() && !(*TwinEdge)->IsDeleted() && !(*TwinEdge)->IsDegenerated())
 				{
@@ -217,37 +227,44 @@ void FAliasModelToCADKernelConverter::LinkEdgesLoop(const AlTrimBoundary& TrimBo
 				}
 			}
 		}
+
+		Status = TrimCurve->nextCurveD();
 	}
 }
 
-TSharedPtr<UE::CADKernel::FTopologicalFace> FAliasModelToCADKernelConverter::AddTrimRegion(const AlTrimRegion& TrimRegion, EAliasObjectReference InObjectReference, const AlMatrix4x4& InAlMatrix, bool bInOrientation)
+TSharedPtr<FTopologicalFace> FAliasModelToCADKernelConverter::AddTrimRegion(const AlTrimRegion& TrimRegion, EAliasObjectReference InObjectReference, const AlMatrix4x4& InAlMatrix, bool bInOrientation)
 {
-	TSharedPtr<UE::CADKernel::FSurface> Surface = AliasToCADKernelUtils::AddNURBSSurface(GeometricTolerance, TrimRegion, InObjectReference, InAlMatrix);
+	TSharedPtr<FSurface> Surface = AliasToCADKernelUtils::AddNURBSSurface(GeometricTolerance, TrimRegion, InObjectReference, InAlMatrix);
 	if (!Surface.IsValid())
 	{
-		return TSharedPtr<UE::CADKernel::FTopologicalFace>();
+		return TSharedPtr<FTopologicalFace>();
 	}
 
 	bool bIsExternal = true;
-	TArray<TSharedPtr<UE::CADKernel::FTopologicalLoop>> Loops;
-	for (TAlObjectPtr<AlTrimBoundary> TrimBoundary(TrimRegion.firstBoundary()); TrimBoundary.IsValid(); TrimBoundary = TAlObjectPtr<AlTrimBoundary>(TrimBoundary->nextBoundary()))
+	TArray<TSharedPtr<FTopologicalLoop>> Loops;
+	TAlObjectPtr<AlTrimBoundary> TrimBoundary(TrimRegion.firstBoundary());
+
+	statusCode Status = TrimBoundary ? sSuccess : sObjectNotFound;
+	while (Status == sSuccess)
 	{
-		TSharedPtr<UE::CADKernel::FTopologicalLoop> Loop = AddLoop(*TrimBoundary, Surface, bIsExternal);
+		TSharedPtr<FTopologicalLoop> Loop = AddLoop(*TrimBoundary, Surface, bIsExternal);
 		if (Loop.IsValid())
 		{
 			LinkEdgesLoop(*TrimBoundary, *Loop);
 			Loops.Add(Loop);
 			bIsExternal = false;
 		}
+
+		Status = TrimBoundary->nextBoundaryD();
 	}
 
 	if (Loops.Num() == 0)
 	{
-		UE::CADKernel::FMessage::Printf(UE::CADKernel::EVerboseLevel::Log, TEXT("The Face %s is degenerate, this face is ignored\n"), TrimRegion.name());
-		return TSharedPtr<UE::CADKernel::FTopologicalFace>();
+		FMessage::Printf(EVerboseLevel::Log, TEXT("The Face %s is degenerate, this face is ignored\n"), TrimRegion.name());
+		return TSharedPtr<FTopologicalFace>();
 	}
 
-	TSharedRef<UE::CADKernel::FTopologicalFace> Face = UE::CADKernel::FEntity::MakeShared<UE::CADKernel::FTopologicalFace>(Surface);
+	TSharedRef<FTopologicalFace> Face = FEntity::MakeShared<FTopologicalFace>(Surface);
 	Face->SetPatchId(LastFaceId++);
 
 	int32 DoubtfulLoopOrientationCount = 0;
@@ -258,116 +275,108 @@ TSharedPtr<UE::CADKernel::FTopologicalFace> FAliasModelToCADKernelConverter::Add
 		Face->SetAsDegenerated();
 		Face->Delete();
 
-		UE::CADKernel::FMessage::Printf(UE::CADKernel::EVerboseLevel::Log, TEXT("The Face %s is degenerate, this face is ignored\n"), TrimRegion.name());
-		return TSharedPtr<UE::CADKernel::FTopologicalFace>();
+		FMessage::Printf(EVerboseLevel::Log, TEXT("The Face %s is degenerate, this face is ignored\n"), TrimRegion.name());
+		return TSharedPtr<FTopologicalFace>();
 	}
 
 	return Face;
 }
 
-void FAliasModelToCADKernelConverter::AddFace(const AlSurface& Surface, EAliasObjectReference InObjectReference, const AlMatrix4x4& InAlMatrix, bool bInOrientation, TSharedRef<UE::CADKernel::FShell>& Shell)
+void FAliasModelToCADKernelConverter::AddFace(const AlSurface& Surface, EAliasObjectReference InObjectReference, const AlMatrix4x4& InAlMatrix, bool bInOrientation, TSharedRef<FShell>& Shell)
 {
 	TAlObjectPtr<AlTrimRegion> TrimRegion(Surface.firstTrimRegion());
-	if (TrimRegion.IsValid())
+	statusCode Status = TrimRegion ? sSuccess : sObjectNotFound;
+	if (Status == sSuccess)
 	{
-		for (; TrimRegion.IsValid(); TrimRegion = TAlObjectPtr<AlTrimRegion>(TrimRegion->nextRegion()))
+		while(Status == sSuccess)
 		{
-			TSharedPtr<UE::CADKernel::FTopologicalFace> Face = AddTrimRegion(*TrimRegion, InObjectReference, InAlMatrix, bInOrientation);
+			TSharedPtr<FTopologicalFace> Face = AddTrimRegion(*TrimRegion, InObjectReference, InAlMatrix, bInOrientation);
 			if (Face.IsValid())
 			{
 				Shell->Add(Face.ToSharedRef(), bInOrientation ? UE::CADKernel::EOrientation::Front : UE::CADKernel::EOrientation::Back);
 			}
+
+			Status = TrimRegion->nextRegionD();
 		}
-		return;
 	}
-
-	TSharedPtr<UE::CADKernel::FSurface> CADKernelSurface = AliasToCADKernelUtils::AddNURBSSurface(GeometricTolerance, Surface, InObjectReference, InAlMatrix);
-	if (CADKernelSurface.IsValid())
+	else
 	{
-		TSharedRef<UE::CADKernel::FTopologicalFace> Face = UE::CADKernel::FEntity::MakeShared<UE::CADKernel::FTopologicalFace>(CADKernelSurface);
-		Face->ApplyNaturalLoops();
-		// Surface can be too thin based on tolerances. Skip it
-		// #cadkernel_check: Warn user a surface was too thin
-		if (Face->GetLoops().Num() > 0)
+		TSharedPtr<FSurface> CADKernelSurface = AliasToCADKernelUtils::AddNURBSSurface(GeometricTolerance, Surface, InObjectReference, InAlMatrix);
+		if (CADKernelSurface.IsValid())
 		{
-			Shell->Add(Face, bInOrientation ? UE::CADKernel::EOrientation::Front : UE::CADKernel::EOrientation::Back);
+			TSharedRef<FTopologicalFace> Face = FEntity::MakeShared<FTopologicalFace>(CADKernelSurface);
+			Face->ApplyNaturalLoops();
+			// Surface can be too thin based on tolerances. Skip it
+			// #cadkernel_check: Warn user a surface was too thin
+			if (Face->GetLoops().Num() > 0)
+			{
+				Shell->Add(Face, bInOrientation ? UE::CADKernel::EOrientation::Front : UE::CADKernel::EOrientation::Back);
+			}
+			// #wire_import: Log that this face was not added
 		}
-		// #wire_import: Log that this face was not added
 	}
 }
 
-void FAliasModelToCADKernelConverter::AddShell(const AlShell& InShell, EAliasObjectReference InObjectReference, const AlMatrix4x4& InAlMatrix, bool bInOrientation, TSharedRef<UE::CADKernel::FShell>& CADKernelShell)
+void FAliasModelToCADKernelConverter::AddShell(const AlShell& InShell, EAliasObjectReference InObjectReference, const AlMatrix4x4& InAlMatrix, bool bInOrientation, TSharedRef<FShell>& CADKernelShell)
 {
-	for(TAlObjectPtr<AlTrimRegion> TrimRegion(InShell.firstTrimRegion()); TrimRegion.IsValid(); TrimRegion = TAlObjectPtr<AlTrimRegion>(TrimRegion->nextRegion()))
+	TAlObjectPtr<AlTrimRegion> TrimRegion(InShell.firstTrimRegion());
+	statusCode Status = TrimRegion ? sSuccess : sObjectNotFound;
+	while (Status == sSuccess)
 	{
-		TSharedPtr<UE::CADKernel::FTopologicalFace> Face = AddTrimRegion(*TrimRegion, InObjectReference, InAlMatrix, bInOrientation);
+		TSharedPtr<FTopologicalFace> Face = AddTrimRegion(*TrimRegion, InObjectReference, InAlMatrix, bInOrientation);
 		if (Face.IsValid())
 		{
 			CADKernelShell->Add(Face.ToSharedRef(), bInOrientation ? UE::CADKernel::EOrientation::Front : UE::CADKernel::EOrientation::Back);
 		}
+
+		Status = TrimRegion->nextRegionD();
 	}
 }
 
-bool FAliasModelToCADKernelConverter::AddBRep(AlDagNode& DagNode, const FColor& Color, EAliasObjectReference InObjectReference)
+bool FAliasModelToCADKernelConverter::AddBRep(const FAlDagNodePtr& DagNode, const FColor& Color, EAliasObjectReference InObjectReference)
 {
 	uint32 ColorId = (uint32)CADLibrary::BuildColorUId(Color);
 	return AddBRep(DagNode, ColorId, InObjectReference);
 }
 
-bool FAliasModelToCADKernelConverter::AddBRep(AlDagNode& DagNode, uint32 SlotID, EAliasObjectReference InObjectReference)
+bool FAliasModelToCADKernelConverter::AddBRep(const FAlDagNodePtr& DagNode, uint32 SlotID, EAliasObjectReference InObjectReference)
 {
+	if (!DagNode.IsValid())
+	{
+		return false;
+	}
+
 	AlEdge2CADKernelEdge.Empty();
 
 	boolean bAlOrientation;
-	DagNode.getSurfaceOrientation(bAlOrientation);
+	DagNode->getSurfaceOrientation(bAlOrientation);
 	bool bOrientation = !(bool)bAlOrientation;
 
 	AlMatrix4x4 AlMatrix;
 	if (InObjectReference == EAliasObjectReference::ParentReference)
 	{
-		DagNode.localTransformationMatrix(AlMatrix);
+		DagNode->localTransformationMatrix(AlMatrix);
 	}
 
-	TSharedRef<UE::CADKernel::FShell> CADKernelShell = UE::CADKernel::FEntity::MakeShared<UE::CADKernel::FShell>();
+	TSharedRef<FShell> CADKernelShell = FEntity::MakeShared<FShell>();
 
-	AlObjectType objectType = DagNode.type();
-	switch (objectType)
-	{
-		// Push all leaf nodes into 'leaves'
-	case kShellNodeType:
+	TAlObjectPtr<AlShell> Shell;
+	if (DagNode.GetShell(Shell))
 	{
-		if (AlShellNode* ShellNode = DagNode.asShellNodePtr())
-		{
-			AlShell* ShellPtr = ShellNode->shell();
-			if(AlIsValid(ShellPtr))
-			{
-				TAlObjectPtr<AlShell> AliasShell(ShellPtr);
-				AddShell(*ShellPtr, InObjectReference, AlMatrix, bOrientation, CADKernelShell);
-			}
-		}
-		break;
+		AddShell(*Shell, InObjectReference, AlMatrix, bOrientation, CADKernelShell);
 	}
-
-	case kSurfaceNodeType:
+	else
 	{
-		if (AlSurfaceNode* SurfaceNode = DagNode.asSurfaceNodePtr())
+		TAlObjectPtr<AlSurface> Surface;
+		if (DagNode.GetSurface(Surface))
 		{
-			AlSurface* SurfacePtr = SurfaceNode->surface();
-			if (AlIsValid(SurfacePtr))
-			{
-				TAlObjectPtr<AlSurface> AliasSurface(SurfacePtr);
-				AddFace(*SurfacePtr, InObjectReference, AlMatrix, bOrientation, CADKernelShell);
-			}
+			AddFace(*Surface, InObjectReference, AlMatrix, bOrientation, CADKernelShell);
 		}
-		break;
-	}
-	default:
-		break;
 	}
 
 	if (CADKernelShell->FaceCount() > 0)
 	{
-		TSharedRef<UE::CADKernel::FBody> CADKernelBody = UE::CADKernel::FEntity::MakeShared<UE::CADKernel::FBody>();
+		TSharedRef<FBody> CADKernelBody = FEntity::MakeShared<FBody>();
 
 		CADKernelBody->SetColorId(SlotID);
 		CADKernelBody->AddShell(CADKernelShell);
@@ -384,7 +393,7 @@ bool FAliasModelToCADKernelConverter::AddBRep(AlDagNode& DagNode, uint32 SlotID,
 
 bool FAliasModelToCADKernelConverter::Tessellate(const CADLibrary::FMeshParameters& InMeshParameters, FMeshDescription& OutMeshDescription)
 {
-	UE::CADKernel::FModel& Model = CADKernelSession.GetModel();
+	FModel& Model = CADKernelSession.GetModel();
 
 	CADLibrary::FMeshConversionContext Context(ImportParameters, InMeshParameters, GeometricTolerance);
 
@@ -397,15 +406,15 @@ bool FAliasModelToCADKernelConverter::RepairTopology()
 	// Apply stitching if applicable
 	if (ImportParameters.GetStitchingTechnique() != StitchingNone)
 	{
-		UE::CADKernel::ESewOption SewOptionValue = (UE::CADKernel::ESewOption)SewOption::GetFromImportParameters();
+		ECADKernelSewOption SewOptionValue = (ECADKernelSewOption)CADLibrary::SewOption::GetFromImportParameters();
 
-#if !THINFACE_ENABLED
-		SewOptionValue = UE::CADKernel::ESewOption((uint8)SewOptionValue & (uint8)~ESewOption::RemoveThinFaces);
+#if !WIRE_THINFACE_ENABLED
+		SewOptionValue = ECADKernelSewOption((uint8)SewOptionValue & (uint8)~ECADKernelSewOption::RemoveThinFaces);
 #endif
 
-		UE::CADKernel::FTopomakerOptions TopomakerOptions(SewOptionValue, StitchingTolerance, FImportParameters::GStitchingForceFactor);
+		FTopomakerOptions TopomakerOptions(SewOptionValue, StitchingTolerance, FImportParameters::GStitchingForceFactor);
 
-		UE::CADKernel::FTopomaker Topomaker(CADKernelSession, TopomakerOptions);
+		FTopomaker Topomaker(CADKernelSession, TopomakerOptions);
 		Topomaker.Sew();
 		Topomaker.SplitIntoConnectedShells();
 		Topomaker.OrientShells();
@@ -420,32 +429,22 @@ bool FAliasModelToCADKernelConverter::AddGeometry(const CADLibrary::FCADModelGeo
 	{
 		const FDagNodeGeometry& DagNodeGeometry = static_cast<const FDagNodeGeometry&>(Geometry);
 
-		return AddBRep(*DagNodeGeometry.DagNode, 0, DagNodeGeometry.Reference);
+		return AddBRep(DagNodeGeometry.DagNode, 0, DagNodeGeometry.Reference);
 	}
 	else if (Geometry.Type == (int32)ECADModelGeometryType::BodyNode)
 	{
 		const FBodyNodeGeometry& BodyNodeGeometry = static_cast<const FBodyNodeGeometry&>(Geometry);
 
-		bool bBodyAdded = true;
-		if (BodyNodeGeometry.BodyNode->HasSurfaceNodes())
-		{
-			bool bBRepAdded = false; // Track that at least one surface was properly added.
-			BodyNodeGeometry.BodyNode->IterateOnSurfaceNodes([&](const TAlDagNodePtr<AlSurfaceNode>& SurfaceNode)
-				{
-					bBRepAdded |= AddBRep(*SurfaceNode, BodyNodeGeometry.BodyNode->GetSlotIndex(SurfaceNode.Get()), BodyNodeGeometry.Reference);
-				});
-			bBodyAdded &= bBRepAdded;
-		}
-
-		if (BodyNodeGeometry.BodyNode->HasShellNodes())
-		{
-			bool bBRepAdded = false; // Track that at least one shell was properly added.
-			BodyNodeGeometry.BodyNode->IterateOnShellNodes([&](const TAlDagNodePtr<AlShellNode>& ShellNode)
+		bool bBodyAdded = false;
+		BodyNodeGeometry.BodyNode->IterateOnDagNodes([&](const FAlDagNodePtr& DagNode)
+			{
+				const bool bBRepAdded = AddBRep(DagNode, BodyNodeGeometry.BodyNode->GetSlotIndex(DagNode), BodyNodeGeometry.Reference);
+				if (!bBRepAdded)
 				{
-					bBRepAdded |= AddBRep(*ShellNode, BodyNodeGeometry.BodyNode->GetSlotIndex(ShellNode.Get()), BodyNodeGeometry.Reference);
-				});
-			bBodyAdded &= bBRepAdded;
-		}
+					UE_LOG(LogWireInterface, Warning, TEXT("Failed to add DagNode %s to StaticMesh."), *DagNode.GetName());
+				}
+				bBodyAdded |= bBRepAdded;
+			});
 
 		ensureWire(bBodyAdded);
 		return bBodyAdded;
@@ -456,4 +455,6 @@ bool FAliasModelToCADKernelConverter::AddGeometry(const CADLibrary::FCADModelGeo
 
 }
 
-#endif
\ No newline at end of file
+#endif
+
+#undef LOCTEXT_NAMESPACE // "WireInterface"
diff --git a/Engine/Plugins/Enterprise/DatasmithCADImporter/Source/WireInterface/Private/AliasModelToCADKernelConverter.h b/Engine/Plugins/Enterprise/DatasmithCADImporter/Source/WireInterface/Private/AliasModelToCADKernelConverter.h
--- a/Engine/Plugins/Enterprise/DatasmithCADImporter/Source/WireInterface/Private/AliasModelToCADKernelConverter.h
+++ b/Engine/Plugins/Enterprise/DatasmithCADImporter/Source/WireInterface/Private/AliasModelToCADKernelConverter.h
@@ -44,8 +44,8 @@ public:
 	virtual bool AddGeometry(const CADLibrary::FCADModelGeometry& Geometry) override;
 	// End FCADModelToCADKernelConverterBase overrides
 
-	bool AddBRep(AlDagNode& DagNode, const FColor& Color, EAliasObjectReference ObjectReference);
-	bool AddBRep(AlDagNode& DagNode, uint32 SlotID, EAliasObjectReference ObjectReference);
+	bool AddBRep(const FAlDagNodePtr& DagNode, const FColor& Color, EAliasObjectReference ObjectReference);
+	bool AddBRep(const FAlDagNodePtr& DagNode, uint32 SlotID, EAliasObjectReference ObjectReference);
 
 protected:
 	TSharedPtr<UE::CADKernel::FTopologicalEdge> AddEdge(const AlTrimCurve& TrimCurve, TSharedPtr<UE::CADKernel::FSurface>& CarrierSurface);
diff --git a/Engine/Plugins/Enterprise/DatasmithCADImporter/Source/WireInterface/Private/AliasModelToTechSoftConverter.cpp b/Engine/Plugins/Enterprise/DatasmithCADImporter/Source/WireInterface/Private/AliasModelToTechSoftConverter.cpp
--- a/Engine/Plugins/Enterprise/DatasmithCADImporter/Source/WireInterface/Private/AliasModelToTechSoftConverter.cpp
+++ b/Engine/Plugins/Enterprise/DatasmithCADImporter/Source/WireInterface/Private/AliasModelToTechSoftConverter.cpp
@@ -36,6 +36,8 @@
 #include "Windows/HideWindowsPlatformTypes.h"
 #endif
 
+#define LOCTEXT_NAMESPACE "WireInterface"
+
 namespace UE_DATASMITHWIRETRANSLATOR_NAMESPACE
 {
 
@@ -263,13 +265,18 @@ A3DTopoLoop* FAliasModelToTechSoftConverter::CreateTopoLoop(const AlTrimBoundary
 	TArray<A3DTopoCoEdge*> Edges;
 	Edges.Reserve(20);
 
-	for (TAlObjectPtr<AlTrimCurve> TrimCurve(TrimBoundary.firstCurve()); TrimCurve.IsValid(); TrimCurve = TAlObjectPtr<AlTrimCurve>(TrimCurve->nextCurve()))
+	TAlObjectPtr<AlTrimCurve> TrimCurve(TrimBoundary.firstCurve());
+	statusCode Status = TrimCurve ? sSuccess : sObjectNotFound;
+
+	while(Status == sSuccess)
 	{
 		A3DTopoCoEdge* Edge = CreateEdge(*TrimCurve);
 		if (Edge != nullptr)
 		{
 			Edges.Add(Edge);
 		}
+
+		Status = TrimCurve->nextCurveD();
 	}
 
 	if (Edges.Num() == 0)
@@ -288,27 +295,25 @@ A3DTopoLoop* FAliasModelToTechSoftConverter::CreateTopoLoop(const AlTrimBoundary
 
 void FAliasModelToTechSoftConverter::LinkEdgesLoop(const AlTrimBoundary& TrimBoundary)
 {
-	for (TAlObjectPtr<AlTrimCurve> TrimCurve(TrimBoundary.firstCurve()); TrimCurve.IsValid(); TrimCurve = TAlObjectPtr<AlTrimCurve>(TrimCurve->nextCurve()))
-	{
-		A3DTopoCoEdge** Edge = AlEdgeToTSCoEdge.Find(TrimCurve->fSpline);
-		if (Edge == nullptr)
-		{
-			continue;
-		}
+	TAlObjectPtr<AlTrimCurve> TrimCurve(TrimBoundary.firstCurve());
+	statusCode Status = TrimCurve ? sSuccess : sObjectNotFound;
 
-		// Link edges
-		TAlObjectPtr<AlTrimCurve> TwinCurve(TrimCurve->getTwinCurve());
-		if (TwinCurve.IsValid())
+	while (Status == sSuccess)
+	{
+		if (A3DTopoCoEdge** Edge = AlEdgeToTSCoEdge.Find(TrimCurve->fSpline))
 		{
-			if (A3DTopoCoEdge** TwinEdge = AlEdgeToTSCoEdge.Find(TwinCurve->fSpline))
+			TAlObjectPtr<AlTrimCurve> TwinCurve(TrimCurve->getTwinCurve());
+			if (TwinCurve.IsValid())
 			{
-				if (TwinEdge != nullptr)
+				if (A3DTopoCoEdge** TwinEdge = AlEdgeToTSCoEdge.Find(TwinCurve->fSpline))
 				{
 					CADLibrary::TechSoftInterface::LinkCoEdges(*Edge, *TwinEdge);
 					break;
 				}
 			}
 		}
+
+		Status = TrimCurve->nextCurveD();
 	}
 }
 
@@ -323,7 +328,10 @@ A3DTopoFace* FAliasModelToTechSoftConverter::AddTrimRegion(const AlTrimRegion& I
 	TArray<A3DTopoLoop*> Loops;
 	Loops.Reserve(5);
 
-	for (TAlObjectPtr<AlTrimBoundary> TrimBoundary(InTrimRegion.firstBoundary()); TrimBoundary.IsValid(); TrimBoundary = TAlObjectPtr<AlTrimBoundary>(TrimBoundary->nextBoundary()))
+	TAlObjectPtr<AlTrimBoundary> TrimBoundary(InTrimRegion.firstBoundary());
+	statusCode Status = TrimBoundary ? sSuccess : sObjectNotFound;
+
+	while(Status == sSuccess)
 	{
 		A3DTopoLoop* Loop = CreateTopoLoop(*TrimBoundary);
 		if (Loop != nullptr)
@@ -331,6 +339,8 @@ A3DTopoFace* FAliasModelToTechSoftConverter::AddTrimRegion(const AlTrimRegion& I
 			Loops.Add(Loop);
 			LinkEdgesLoop(*TrimBoundary);
 		}
+
+		TrimBoundary->nextBoundaryD();
 	}
 
 	if (Loops.Num() == 0)
@@ -354,14 +364,14 @@ A3DTopoFace* FAliasModelToTechSoftConverter::AddTrimRegion(const AlTrimRegion& I
 }
 #endif
 
-bool FAliasModelToTechSoftConverter::AddBRep(AlDagNode& DagNode, uint32 SlotID, EAliasObjectReference InObjectReference)
+bool FAliasModelToTechSoftConverter::AddBRep(const FAlDagNodePtr& DagNode, uint32 SlotID, EAliasObjectReference InObjectReference)
 {
 	FColor Color(SlotID);
 
 	return AddBRep(DagNode, Color, InObjectReference);
 }
 
-bool FAliasModelToTechSoftConverter::AddBRep(AlDagNode& DagNode, const FColor& Color, EAliasObjectReference InObjectReference)
+bool FAliasModelToTechSoftConverter::AddBRep(const FAlDagNodePtr& DagNode, const FColor& Color, EAliasObjectReference InObjectReference)
 {
 #ifdef USE_TECHSOFT_SDK
 	AlEdgeToTSCoEdge.Empty();
@@ -369,62 +379,55 @@ bool FAliasModelToTechSoftConverter::AddBRep(AlDagNode& DagNode, const FColor& C
 	AlMatrix4x4 AlMatrix;
 	if (InObjectReference == EAliasObjectReference::ParentReference)
 	{
-		DagNode.localTransformationMatrix(AlMatrix);
+		DagNode->localTransformationMatrix(AlMatrix);
 	}
 
 	boolean bAlOrientation;
-	DagNode.getSurfaceOrientation(bAlOrientation);
+	DagNode->getSurfaceOrientation(bAlOrientation);
 	bool bOrientation = !(bool)bAlOrientation;
 
 	TArray<A3DTopoFace*> TSFaces;
 	TSFaces.Reserve(100);
 
-	AlObjectType objectType = DagNode.type();
-	switch (objectType)
+	TAlObjectPtr<AlShell> Shell;
+	if (DagNode.GetShell(Shell))
 	{
+		TAlObjectPtr<AlTrimRegion> TrimRegion(Shell->firstTrimRegion());
+		statusCode Status = TrimRegion ? sSuccess : sObjectNotFound;
+		while(Status == sSuccess)
+		{
+			A3DTopoFace* TSFace = AddTrimRegion(*TrimRegion, Color, InObjectReference, AlMatrix);
+			if (TSFace != nullptr)
+			{
+				TSFaces.Add(TSFace);
+			}
 
-	case kShellNodeType:
+			Status = TrimRegion->nextRegionD();
+		}
+	}
+	else
 	{
-		if (AlShellNode* ShellNode = DagNode.asShellNodePtr())
+		TAlObjectPtr<AlSurface> Surface;
+		if (DagNode.GetSurface(Surface))
 		{
-			TAlObjectPtr<AlShell> AliasShell(ShellNode->shell());
-			if(AlIsValid(AliasShell.Get()))
+			TAlObjectPtr<AlTrimRegion> TrimRegion(Surface->firstTrimRegion());
+			statusCode Status = TrimRegion ? sSuccess : sObjectNotFound;
+			if (Status == sSuccess)
 			{
-				for (TAlObjectPtr<AlTrimRegion> TrimRegion(AliasShell->firstTrimRegion()); TrimRegion.IsValid(); TrimRegion = TAlObjectPtr<AlTrimRegion>(TrimRegion->nextRegion()))
+				while (Status == sSuccess)
 				{
 					A3DTopoFace* TSFace = AddTrimRegion(*TrimRegion, Color, InObjectReference, AlMatrix);
 					if (TSFace != nullptr)
 					{
 						TSFaces.Add(TSFace);
 					}
+
+					Status = TrimRegion->nextRegionD();
 				}
 			}
-		}
-		break;
-	}
-
-	case kSurfaceNodeType:
-	{
-		if (AlSurfaceNode* SurfaceNode = DagNode.asSurfaceNodePtr())
-		{
-			TAlObjectPtr<AlSurface> AliasSurface(SurfaceNode->surface());
-			if (AlIsValid(AliasSurface.Get()))
+			else
 			{
-				TAlObjectPtr<AlTrimRegion> TrimRegion(AliasSurface->firstTrimRegion());
-				if (TrimRegion.IsValid())
-				{
-					for (; TrimRegion.IsValid(); TrimRegion = TAlObjectPtr<AlTrimRegion>(TrimRegion->nextRegion()))
-					{
-						A3DTopoFace* TSFace = AddTrimRegion(*TrimRegion, Color, InObjectReference, AlMatrix);
-						if (TSFace != nullptr)
-						{
-							TSFaces.Add(TSFace);
-						}
-					}
-					break;
-				}
-
-				A3DSurfBase* TSSurface = AliasToTechSoftUtils::AddNURBSSurface(*AliasSurface, InObjectReference, AlMatrix);
+				A3DSurfBase* TSSurface = AliasToTechSoftUtils::AddNURBSSurface(*Surface, InObjectReference, AlMatrix);
 				if (TSSurface != nullptr)
 				{
 					A3DTopoFace* TSFace = CADLibrary::TechSoftUtils::CreateTopoFaceWithNaturalLoop(TSSurface);
@@ -437,10 +440,6 @@ bool FAliasModelToTechSoftConverter::AddBRep(AlDagNode& DagNode, const FColor& C
 				}
 			}
 		}
-		break;
-	}
-	default:
-		break;
 	}
 
 	if (TSFaces.IsEmpty())
@@ -484,25 +483,25 @@ bool FAliasModelToTechSoftConverter::AddGeometry(const CADLibrary::FCADModelGeom
 	{
 		const FDagNodeGeometry& DagNodeGeometry = static_cast<const FDagNodeGeometry&>(Geometry);
 
-		return AddBRep(*DagNodeGeometry.DagNode, 0, DagNodeGeometry.Reference);
+		return AddBRep(DagNodeGeometry.DagNode, 0, DagNodeGeometry.Reference);
 	}
 	else if (Geometry.Type == (int32)ECADModelGeometryType::BodyNode)
 	{
 		const FBodyNodeGeometry& BodyNodeGeometry = static_cast<const FBodyNodeGeometry&>(Geometry);
 
-		bool bBRepAdded = true;
-		BodyNodeGeometry.BodyNode->IterateOnSurfaceNodes([&](const TAlDagNodePtr<AlSurfaceNode>& SurfaceNode)
-			{
-				bBRepAdded &= AddBRep(*SurfaceNode, BodyNodeGeometry.BodyNode->GetSlotIndex(SurfaceNode.Get()), BodyNodeGeometry.Reference);
-			});
-
-		BodyNodeGeometry.BodyNode->IterateOnShellNodes([&](const TAlDagNodePtr<AlShellNode>& ShellNode)
+		bool bBodyAdded = false;
+		BodyNodeGeometry.BodyNode->IterateOnDagNodes([&](const FAlDagNodePtr& DagNode)
 			{
-				bBRepAdded &= AddBRep(*ShellNode, BodyNodeGeometry.BodyNode->GetSlotIndex(ShellNode.Get()), BodyNodeGeometry.Reference);
+				// #wire_import: TODO: Inform user one DagNode was not imported
+				const bool bBRepAdded = AddBRep(DagNode, BodyNodeGeometry.BodyNode->GetSlotIndex(DagNode), BodyNodeGeometry.Reference);
+				if (!bBRepAdded)
+				{
+					UE_LOG(LogWireInterface, Warning, TEXT("Failed to add DagNode %s to StaticMesh."), *DagNode.GetName());
+				}
+				bBodyAdded |= bBRepAdded;
 			});
 
-		ensure(bBRepAdded);
-		return bBRepAdded;
+		return bBodyAdded;
 	}
 
 	return false;
@@ -510,3 +509,5 @@ bool FAliasModelToTechSoftConverter::AddGeometry(const CADLibrary::FCADModelGeom
 
 }
 #endif
+
+#undef LOCTEXT_NAMESPACE // "WireInterface"
diff --git a/Engine/Plugins/Enterprise/DatasmithCADImporter/Source/WireInterface/Private/AliasModelToTechSoftConverter.h b/Engine/Plugins/Enterprise/DatasmithCADImporter/Source/WireInterface/Private/AliasModelToTechSoftConverter.h
--- a/Engine/Plugins/Enterprise/DatasmithCADImporter/Source/WireInterface/Private/AliasModelToTechSoftConverter.h
+++ b/Engine/Plugins/Enterprise/DatasmithCADImporter/Source/WireInterface/Private/AliasModelToTechSoftConverter.h
@@ -39,8 +39,8 @@ public:
 
 	virtual bool AddGeometry(const CADLibrary::FCADModelGeometry& Geometry) override;
 
-	bool AddBRep(AlDagNode& DagNode, const FColor& Color, EAliasObjectReference ObjectReference);
-	bool AddBRep(AlDagNode& DagNode, uint32 SlotID, EAliasObjectReference ObjectReference);
+	bool AddBRep(const FAlDagNodePtr& DagNode, const FColor& Color, EAliasObjectReference ObjectReference);
+	bool AddBRep(const FAlDagNodePtr& DagNode, uint32 SlotID, EAliasObjectReference ObjectReference);
 
 protected:
 	TMap<void*, A3DTopoCoEdge*> AlEdgeToTSCoEdge;
diff --git a/Engine/Plugins/Enterprise/DatasmithCADImporter/Source/WireInterface/Private/OpenModelUtils.cpp b/Engine/Plugins/Enterprise/DatasmithCADImporter/Source/WireInterface/Private/OpenModelUtils.cpp
--- a/Engine/Plugins/Enterprise/DatasmithCADImporter/Source/WireInterface/Private/OpenModelUtils.cpp
+++ b/Engine/Plugins/Enterprise/DatasmithCADImporter/Source/WireInterface/Private/OpenModelUtils.cpp
@@ -5,13 +5,14 @@
 #include "CADOptions.h"
 #include "DatasmithUtils.h"
 #include "DatasmithTranslator.h"
-#include "IDatasmithSceneElements.h"
 
 #include "MeshAttributes.h"
 #include "MeshDescription.h"
 #include "StaticMeshAttributes.h"
 #include "StaticMeshOperations.h"
 
+DEFINE_LOG_CATEGORY(LogWireInterface);
+
 #if PLATFORM_WINDOWS
 #include "Windows/AllowWindowsPlatformTypes.h"
 #endif
@@ -55,7 +56,7 @@ namespace UE_DATASMITHWIRETRANSLATOR_NAMESPACE
 		{
 			Hash = GetTypeHash(Name);
 
-			for (const TAlDagNodePtr<AlMeshNode>& MeshNode : MeshNodes)
+			for (const FAlDagNodePtr& MeshNode : MeshNodes)
 			{
 				Hash = HashCombine(Hash, MeshNode.GetHash());
 			}
@@ -68,12 +69,13 @@ namespace UE_DATASMITHWIRETRANSLATOR_NAMESPACE
 		return bInitialized;
 	}
 
-	int32 FBodyNode::GetSlotIndex(const TAlDagNodePtr<AlDagNode>& DagNode)
+	int32 FBodyNode::GetSlotIndex(const FAlDagNodePtr& DagNode)
 	{
 		// #wire_import: Add support for AlShell
-		if (TAlDagNodePtr<AlSurfaceNode> SurfaceNode = DagNode->asSurfaceNodePtr())
+		if (DagNode.IsASurface())
 		{
-			if (TAlObjectPtr<AlSurface> Surface = SurfaceNode->surface())
+			TAlObjectPtr<AlSurface> Surface;
+			if (DagNode.GetSurface(Surface))
 			{
 				TAlObjectPtr<AlShader> Shader(Surface->firstShader());
 				return Shader ? ShaderNameToSlotIndex[Shader.GetName()] : 0;
@@ -86,7 +88,7 @@ namespace UE_DATASMITHWIRETRANSLATOR_NAMESPACE
 
 	bool FBodyNode::Initialize()
 	{
-		if (SurfaceNodes.IsEmpty() && ShellNodes.IsEmpty())
+		if (DagNodes.IsEmpty())
 		{
 			return false;
 		}
@@ -96,60 +98,74 @@ namespace UE_DATASMITHWIRETRANSLATOR_NAMESPACE
 			Hash = GetTypeHash(Name);
 
 			int32 SlotIndex = 0;
-			TSet<TAlObjectPtr<AlLayer>> LayerSet;
+			TMap<FString, AlLayer*> LayerSet;
 
-			for (const TAlDagNodePtr<AlSurfaceNode>& SurfaceNode : SurfaceNodes)
-			{
-				TAlObjectPtr<AlSurface> Surface(SurfaceNode->surface());
-				if (!Surface)
+			TFunction<void(const TAlObjectPtr<AlShader>&)> RegisterShader = [this, &SlotIndex](const TAlObjectPtr<AlShader>& Shader)
 				{
-					continue;
-				}
+					if (Shader && !ShaderNameToSlotIndex.Contains(Shader.GetName()))
+					{
+						ShaderNameToSlotIndex.Add(Shader.GetName(), SlotIndex);
+						SlotIndexToShader.Add(SlotIndex++, Shader);
+					}
+				};
 
-				TAlObjectPtr<AlShader> Shader(Surface->firstShader());
-				if (Shader && !ShaderNameToSlotIndex.Contains(Shader.GetName()))
+			for (const FAlDagNodePtr& DagNode : DagNodes)
+			{
+				TAlObjectPtr<AlSurface> Surface;
+				if (DagNode.GetSurface(Surface))
 				{
-					ShaderNameToSlotIndex.Add(Shader.GetName(), SlotIndex);
-					SlotIndexToShader.Add(SlotIndex++, Shader);
+					TAlObjectPtr<AlShader> Shader(Surface->firstShader());
+					ensureWire(Shader);
+					RegisterShader(Shader);
 				}
-
-				if (TAlObjectPtr<AlLayer> SLayer = SurfaceNode->layer())
+				else
 				{
-					LayerSet.Add(TAlObjectPtr<AlLayer >(SLayer));
+					TAlObjectPtr<AlShell> Shell;
+					if (DagNode.GetShell(Shell))
+					{
+						TAlObjectPtr<AlShader> Shader(Shell->firstShader());
+						RegisterShader(Shader);
+						// #wire_import: Do we have as many shaders than trim regions
+#if WIRE_ENSURE_ENABLED
+						int32 ShaderCount = 0;
+						{
+							statusCode Status = Shader.IsValid() ? sSuccess : sFailure;
+							while (Status == sSuccess)
+							{
+								++ShaderCount;
+								Status = Shell->nextShaderD(Shader.Get());
+							};
+						}
+						TAlObjectPtr<AlTrimRegion> TrimRegion(Shell->firstTrimRegion());
+						int32 TrimCount = 0;
+						{
+							statusCode Status = TrimRegion.IsValid() ? sSuccess : sFailure;
+							while (Status == sSuccess)
+							{
+								++TrimCount;
+								Status = TrimRegion->nextRegionD();
+							};
+						}
+						ensureWire(ShaderCount == TrimCount);
+#endif
+					}
+					else
+					{
+						ensure(false);
+					}
 				}
 
-				Hash = HashCombine(Hash, SurfaceNode.GetHash());
-			}
-
-			for (const TAlDagNodePtr<AlShellNode>& ShellNode : ShellNodes)
-			{
-				// #wire_import: Extract all shaders
-				TAlObjectPtr<AlShell> Shell(ShellNode->shell());
-				if (!Shell)
-				{
-					continue;
-				}
 
-				TAlObjectPtr<AlShader> Shader(Shell->firstShader());
-				if (Shader && !ShaderNameToSlotIndex.Contains(Shader.GetName()))
+				if (DagNode.GetLayer())
 				{
-					ShaderNameToSlotIndex.Add(Shader.GetName(), SlotIndex);
-					SlotIndexToShader.Add(SlotIndex++, Shader);
+					LayerSet.Add(DagNode.GetLayerName(), DagNode.GetLayer().Get());
 				}
 
-				if (TAlObjectPtr<AlLayer> SLayer = ShellNode->layer())
-				{
-					LayerSet.Add(TAlObjectPtr<AlLayer >(SLayer));
-				}
-				Hash = HashCombine(Hash, ShellNode.GetHash());
+				Hash = HashCombine(Hash, DagNode.GetHash());
 			}
 
 			ensureWire(LayerSet.Num() == 1);
-			TAlObjectPtr<AlLayer> GeomLayer = LayerSet.Array()[0];
-			if (GeomLayer != Layer)
-			{
-				Layer = GeomLayer;
-			}
+			// #wire_import: TODO - Make sure Body's layer is the same as those of the added geometries
 
 			UniqueID = TEXT("BodyNode") + FString::FromInt(Hash);
 
@@ -159,30 +175,31 @@ namespace UE_DATASMITHWIRETRANSLATOR_NAMESPACE
 		return bInitialized;
 	}
 
-	void FBodyNode::AddSurfaceNode(const TAlDagNodePtr<AlSurfaceNode>& SurfaceNode)
+	bool FBodyNode::AddNode(FAlDagNodePtr& DagNode)
 	{
-		if (OpenModelUtils::IsGeometryValid(SurfaceNode.Get()))
+		if (DagNode.IsASurface() || DagNode.IsAShell())
 		{
-			SurfaceNodes.Add(SurfaceNode);
+			DagNodes.Add(DagNode);
+			return true;
 		}
-	}
+#if WIRE_ENSURE_ENABLED
+		ensureWire(false);
+#endif
 
-	void FBodyNode::AddShellNode(const TAlDagNodePtr<AlShellNode>& ShellNode)
-	{
-		if (OpenModelUtils::IsGeometryValid(ShellNode.Get()))
-		{
-			ShellNodes.Add(ShellNode);
-		}
+		return false;
 	}
 
 	template<>
 	uint32 TAlObjectPtr<AlLayer>::GetHash() const
 	{
 		uint32 Hash = GetTypeHash(GetName());
-		Hash = HashCombine(Hash, GetTypeHash(Ptr->number()));
-		Hash = HashCombine(Hash, GetTypeHash(Ptr->color()));
-		Hash = HashCombine(Hash, (bool)Ptr->invisible() ? 1 : 0);
-		Hash = HashCombine(Hash, (bool)Ptr->isSymmetric() ? 1 : 0);
+		if (IsValid())
+		{
+			Hash = HashCombine(Hash, GetTypeHash(Super::Get()->number()));
+			Hash = HashCombine(Hash, GetTypeHash(Super::Get()->color()));
+			Hash = HashCombine(Hash, (bool)Super::Get()->invisible() ? 1 : 0);
+			Hash = HashCombine(Hash, (bool)Super::Get()->isSymmetric() ? 1 : 0);
+		}
 
 		return Hash;
 	}
@@ -196,7 +213,7 @@ namespace UE_DATASMITHWIRETRANSLATOR_NAMESPACE
 
 		CsvString = Layer.GetName();
 
-		TAlObjectPtr<AlLayer> ParentLayer(Layer->parentLayer());
+		TAlObjectPtr<AlLayer> ParentLayer = FLayerContainer::FindOrAdd(Layer->parentLayer());
 		while (ParentLayer)
 		{
 			FString ParentLayerName = ParentLayer.GetName();
@@ -205,7 +222,7 @@ namespace UE_DATASMITHWIRETRANSLATOR_NAMESPACE
 				CsvString += TEXT(",") + ParentLayerName;
 			}
 
-			ParentLayer = ParentLayer->parentLayer();
+			ParentLayer = FLayerContainer::FindOrAdd(ParentLayer->parentLayer());
 		}
 
 		return !CsvString.IsEmpty();
@@ -221,34 +238,6 @@ namespace UE_DATASMITHWIRETRANSLATOR_NAMESPACE
 		return ActorElement->IsA(EDatasmithElementType::StaticMeshActor) || ActorElement->GetChildrenCount() > 0;
 	}
 
-	void OpenModelUtils::SetActorTransform(IDatasmithActorElement& ActorElement, const TAlDagNodePtr<AlDagNode>& DagNode)
-	{
-		// #wire_import: Find why no transform applied if layer is symmetrical???
-		if (DagNode.HasSymmetry())
-		{
-			return;
-		}
-
-		AlMatrix4x4 AlGlobalMatrix;
-		DagNode->globalTransformationMatrix(AlGlobalMatrix);
-
-		FMatrix GlobalMatrix;
-		double* MatrixFloats = (double*)GlobalMatrix.M;
-		for (int32 IndexI = 0; IndexI < 4; ++IndexI)
-		{
-			for (int32 IndexJ = 0; IndexJ < 4; ++IndexJ)
-			{
-				MatrixFloats[IndexI * 4 + IndexJ] = AlGlobalMatrix[IndexI][IndexJ];
-			}
-		}
-
-		FTransform GlobalTransform = FDatasmithUtils::ConvertTransform(FDatasmithUtils::EModelCoordSystem::ZUp_RightHanded, FTransform(GlobalMatrix));
-
-		ActorElement.SetTranslation(GlobalTransform.GetTranslation());
-		ActorElement.SetScale(GlobalTransform.GetScale3D());
-		ActorElement.SetRotation(GlobalTransform.GetRotation());
-	}
-
 	bool OpenModelUtils::IsValidActor(const TSharedPtr<IDatasmithActorElement>& ActorElement)
 	{
 		if (ActorElement != nullptr)
@@ -462,7 +451,7 @@ namespace UE_DATASMITHWIRETRANSLATOR_NAMESPACE
 		return true;
 	}
 
-	TAlDagNodePtr<AlMeshNode> OpenModelUtils::TesselateDagLeaf(const AlDagNode& DagLeaf, ETesselatorType TessType, double Tolerance)
+	FAlDagNodePtr OpenModelUtils::TesselateDagLeaf(const AlDagNode& DagLeaf, ETesselatorType TessType, double Tolerance)
 	{
 		AlDagNode* TesselatedNode = nullptr;
 		statusCode TessStatus;
@@ -478,69 +467,7 @@ namespace UE_DATASMITHWIRETRANSLATOR_NAMESPACE
 			break;
 		}
 
-		TAlDagNodePtr<AlMeshNode> MeshNode(TesselatedNode ? TesselatedNode->asMeshNodePtr() : nullptr);
-		if ((TessStatus == sSuccess) && MeshNode)
-		{
-			return MeshNode;
-		}
-		else
-		{
-			return TAlDagNodePtr<AlMeshNode>();
-		}
-	}
-
-	bool OpenModelUtils::IsGeometryValid(const TAlDagNodePtr<AlDagNode>& GeomNode)
-	{
-		TAlObjectPtr<AlLayer> Layer(GeomNode->layer());
-		if (Layer.IsValid() && (bool)Layer->invisible())
-		{
-			return false;
-		}
-
-		bool bValidGeometry = false;
-
-		{
-			TAlDagNodePtr<AlMeshNode> MeshNode(GeomNode->asMeshNodePtr());
-			if (MeshNode)
-			{
-				TAlObjectPtr<AlMesh> Mesh(MeshNode->mesh());
-				bValidGeometry |= Mesh.IsValid();
-			}
-		}
-		{
-			TAlDagNodePtr<AlShellNode> ShellNode(GeomNode->asShellNodePtr());
-			if (ShellNode)
-			{
-				TAlObjectPtr<AlShell> Shell(ShellNode->shell());
-				bValidGeometry |= Shell.IsValid();
-			}
-		}
-		{
-			TAlDagNodePtr<AlSurfaceNode> SurfaceNode(GeomNode->asSurfaceNodePtr());
-			if (SurfaceNode)
-			{
-				TAlObjectPtr<AlSurface> Surface(SurfaceNode->surface());
-				bValidGeometry |= Surface.IsValid();
-			}
-		}
-		return bValidGeometry;
-	}
-
-	bool OpenModelUtils::IsDagNodeValid(const TAlDagNodePtr<AlDagNode>& DagNode)
-	{
-		TAlObjectPtr<AlLayer> Layer(DagNode->layer());
-		if (Layer.IsValid() && (bool)Layer->invisible())
-		{
-			return false;
-		}
-
-		TAlDagNodePtr<AlGroupNode> GroupNode(DagNode->asGroupNodePtr());
-		if (GroupNode)
-		{
-			return true;
-		}
-
-		return IsGeometryValid(DagNode);
+		return TessStatus == sSuccess ? TesselatedNode : FAlDagNodePtr();
 	}
 
 	CADLibrary::FMeshParameters OpenModelUtils::GetMeshParameters(const TAlObjectPtr<AlLayer>& Layer)
@@ -568,18 +495,42 @@ namespace UE_DATASMITHWIRETRANSLATOR_NAMESPACE
 		return MeshParameters;
 	}
 
-	CADLibrary::FMeshParameters OpenModelUtils::GetMeshParameters(AlDagNode& DagNode)
+	CADLibrary::FMeshParameters FAlDagNodePtr::GetMeshParameters() const
 	{
-		CADLibrary::FMeshParameters MeshParameters = GetMeshParameters(DagNode.layer());
-		
+		if (!IsValid())
+		{
+			return {};
+		}
+
+		CADLibrary::FMeshParameters MeshParameters = OpenModelUtils::GetMeshParameters(GetLayer());
+
 		boolean bAlOrientation;
-		DagNode.getSurfaceOrientation(bAlOrientation);
+		AsADagNode()->getSurfaceOrientation(bAlOrientation);
 
-		MeshParameters.bNeedSwapOrientation = (DagNode.type() == kMeshNodeType) ? (bool)bAlOrientation : false;
+		MeshParameters.bNeedSwapOrientation = IsAMesh() ? (bool)bAlOrientation : false;
 
 		return MeshParameters;
 	}
 
+	TMap<AlLayer*, TAlObjectPtr<AlLayer>> FLayerContainer::LayerMap;
+
+	TAlObjectPtr<AlLayer> FLayerContainer::FindOrAdd(AlLayer* Layer)
+	{
+		if (const TAlObjectPtr<AlLayer>* ValuePtr = LayerMap.Find(Layer))
+		{
+			return *ValuePtr;
+		}
+
+		TAlObjectPtr<AlLayer>& LayerPtr = LayerMap.Add(Layer);
+		LayerPtr = Layer;
+
+		return LayerPtr;
+	}
+
+	void FLayerContainer::Reset()
+	{
+		LayerMap.Reset();
+	}
 }
 
 #endif
diff --git a/Engine/Plugins/Enterprise/DatasmithCADImporter/Source/WireInterface/Private/OpenModelUtils.h b/Engine/Plugins/Enterprise/DatasmithCADImporter/Source/WireInterface/Private/OpenModelUtils.h
--- a/Engine/Plugins/Enterprise/DatasmithCADImporter/Source/WireInterface/Private/OpenModelUtils.h
+++ b/Engine/Plugins/Enterprise/DatasmithCADImporter/Source/WireInterface/Private/OpenModelUtils.h
@@ -5,10 +5,20 @@
 #include "CoreMinimal.h"
 
 #ifdef USE_OPENMODEL
-//#include "CADModelToTechSoftConverterBase.h"
 #include "CADModelConverter.h"
 #include "CADOptions.h"
+#include "IDatasmithSceneElements.h"
 
+#include "Misc/Optional.h"
+
+#if WITH_EDITOR
+#include "Editor.h"
+#include "IMessageLogListing.h"
+#include "Logging/TokenizedMessage.h"
+#include "MessageLogModule.h"
+#endif
+
+DECLARE_LOG_CATEGORY_EXTERN(LogWireInterface, Log, All);
 
 #if PLATFORM_WINDOWS
 #include "Windows/AllowWindowsPlatformTypes.h"
@@ -16,8 +26,13 @@
 
 #include "AlShadingFields.h"
 #include "AlDagNode.h"
+#include "AlLayer.h"
+#include "AlMesh.h"
 #include "AlMeshNode.h"
+#include "AlShader.h"
+#include "AlShell.h"
 #include "AlShellNode.h"
+#include "AlSurface.h"
 #include "AlSurfaceNode.h"
 #include "AlPersistentID.h"
 
@@ -43,6 +58,7 @@ struct FMeshDescription;
 #define UNIT_CONVERSION_CM_TO_MM		10.
 #define UE_TO_CADKERNEL(Distance)		(Distance * 10.)
 
+#define WIRE_MEMORY_CHECK 0
 #define WIRE_ENSURE_ENABLED 0
 #if WIRE_ENSURE_ENABLED
 #define ensureWire(InExpression) ensure(InExpression)
@@ -52,6 +68,13 @@ struct FMeshDescription;
 
 namespace UE_DATASMITHWIRETRANSLATOR_NAMESPACE
 {
+#if WIRE_MEMORY_CHECK
+	extern TSet<AlDagNode*> DagNodeSet;
+	extern TSet<AlObject*> ObjectSet;
+	extern int32 AllocatedObjects;
+	extern int32 MaxAllocatedObjects;
+#endif
+
 	typedef double AlMatrix4x4[4][4];
 
 	enum class ETesselatorType : uint8
@@ -69,36 +92,55 @@ namespace UE_DATASMITHWIRETRANSLATOR_NAMESPACE
 	};
 
 	template<typename T>
-	class TAlObjectPtr
+	class TAlObjectPtr : public TSharedPtr<T>
 	{
-	public:
-		TAlObjectPtr(T* AlObject = nullptr) : Ptr(AlObject) {}
-
-		bool IsValid() const { return Ptr && AlIsValid(Ptr); }
+		using Super = TSharedPtr<T>;
 
-		T* operator->() const
+	public:
+		TAlObjectPtr(T* Object = nullptr) : Super(Object)
 		{
-			return Ptr;
+			static_assert(std::is_base_of<AlObject, T>::value);
+#if WIRE_MEMORY_CHECK
+			ensure(!Object || !ObjectSet.Contains(Object));
+			if (IsValid() && Super::GetSharedReferenceCount() == 1)
+			{
+				++AllocatedObjects;
+				if (AllocatedObjects > MaxAllocatedObjects) { MaxAllocatedObjects = AllocatedObjects; }
+				ObjectSet.Add(Object);
+			}
+#endif
 		}
 
-		T& operator*() const
+		~TAlObjectPtr()
 		{
-			return *Ptr;
+#if WIRE_MEMORY_CHECK
+			if (Super::IsValid() && Super::GetSharedReferenceCount() == 1)
+			{
+				ensure(ObjectSet.Contains(Super::Get()));
+				--AllocatedObjects;
+				ObjectSet.Remove(Super::Get());
+			}
+#endif
+			Super::Reset();
 		}
 
-		T* Get() const
-		{
-			return Ptr;
-		}
+		bool IsValid() const { return Super::IsValid() && AlIsValid((const AlObject*)Super::Get()); }
 
 		FString GetName() const
 		{
-			return Ptr ? StringCast<TCHAR>(Ptr->name()).Get() : FString();
+			return IsValid() ? StringCast<TCHAR>(Super::Get()->name()).Get() : FString();
 		}
 
 		uint32 GetHash() const
 		{
-			return GetTypeHash(FString(StringCast<TCHAR>(Ptr->name()).Get()));
+			if (IsValid())
+			{
+				uint32 NameHash = GetTypeHash(GetName());
+				uint32 TypeHash = GetTypeHash(Super::Get()->type());
+				return HashCombine(NameHash, TypeHash);
+			}
+
+			return -1;
 		}
 
 		FString GetUniqueID(const TCHAR* TypeName = TEXT("Object")) const
@@ -106,12 +148,6 @@ namespace UE_DATASMITHWIRETRANSLATOR_NAMESPACE
 			return FString(TypeName) + FString::FromInt(GetHash());
 		}
 
-		TAlObjectPtr<T> operator=(const TAlObjectPtr<T>& Src)
-		{
-			Ptr = Src.IsValid() ? Src.Get() : nullptr;
-			return *this;
-		}
-
 		operator bool() const
 		{
 			return IsValid();
@@ -121,9 +157,6 @@ namespace UE_DATASMITHWIRETRANSLATOR_NAMESPACE
 		{
 			return Object.GetHash()/*GetTypeHash(DateTime.Ticks)*/;
 		}
-
-	protected:
-		T* Ptr = nullptr;
 	};
 
 	template<typename T>
@@ -141,59 +174,206 @@ namespace UE_DATASMITHWIRETRANSLATOR_NAMESPACE
 	template<>
 	uint32 TAlObjectPtr<AlLayer>::GetHash() const;
 
-	template<typename T>
-	class TAlDagNodePtr : public TAlObjectPtr<T>
+	class FLayerContainer
 	{
 	public:
-		TAlDagNodePtr(T* AlDagNode = nullptr) : TAlObjectPtr<T>(AlDagNode) {}
+		static void Reset();
+		static TAlObjectPtr<AlLayer> FindOrAdd(AlLayer* Layer);
+
+	private:
+		static TMap<AlLayer*, TAlObjectPtr<AlLayer>> LayerMap;
+	};
+
+	enum  class EDagNodeType : uint8
+	{
+		Unknown = 0x00,
+		MeshType = 0x01,
+		SurfaceType = 0x02,
+		ShellType = 0x04,
+		GroupType = 0x08,
+		GeometryType = MeshType | SurfaceType | ShellType,
+	};
+	ENUM_CLASS_FLAGS(EDagNodeType);
+
+	class FAlDagNodePtr : public TAlObjectPtr<AlDagNode>
+	{
+		using Super = TSharedPtr<AlDagNode>;
+
+	public:
+		FAlDagNodePtr(AlDagNode* DagNode = nullptr)
+			: TAlObjectPtr<AlDagNode>(DagNode)
+			, Type(EDagNodeType::Unknown)
+		{
+#if WIRE_MEMORY_CHECK
+			ensure(!DagNode || !DagNodeSet.Contains(DagNode));
+#endif
+
+			if (AlIsValid(DagNode))
+			{
+				if (DagNode->type() == AlObjectType::kMeshNodeType && DagNode->asMeshNodePtr() != nullptr)
+				{
+					EnumAddFlags(Type, EDagNodeType::MeshType);
+				}
+				else if (DagNode->type() == AlObjectType::kSurfaceNodeType && DagNode->asSurfaceNodePtr() != nullptr)
+				{
+					EnumAddFlags(Type, EDagNodeType::SurfaceType);
+				}
+				else if (DagNode->type() == AlObjectType::kShellNodeType && DagNode->asShellNodePtr() != nullptr)
+				{
+					EnumAddFlags(Type, EDagNodeType::ShellType);
+				}
+				else if (DagNode->type() == AlObjectType::kGroupNodeType && DagNode->asGroupNodePtr() != nullptr)
+				{
+					EnumAddFlags(Type, EDagNodeType::GroupType);
+				}
+
+				CachedLayer = FLayerContainer::FindOrAdd(DagNode->layer());
+				LayerName = CachedLayer ? FString(StringCast<TCHAR>(CachedLayer->name()).Get()) : FString();
+
+				bCanDeleteObject = DagNode->parentNode() == nullptr;
+			}
+#if WIRE_MEMORY_CHECK
+			if (DagNode)
+			{
+				DagNodeSet.Add(DagNode);
+			}
+#endif
+		}
+
+		~FAlDagNodePtr()
+		{
+			if (TAlObjectPtr<AlDagNode>::IsValid())
+			{
+				if (Super::GetSharedReferenceCount() == 1)
+				{
+					CachedLayer = nullptr;
+					if (bCanDeleteObject)
+					{
+						Super::Get()->deleteObject();
+					}
+
+#if WIRE_MEMORY_CHECK
+					ensure(AllocatedObjects > 0);
+					ensure(DagNodeSet.Contains(Super::Get()));
+					DagNodeSet.Remove(Super::Get());
+#endif
+				}
+			}
+			TAlObjectPtr<AlDagNode>::~TAlObjectPtr();
+		}
 
 		FString GetLayerName() const
 		{
-			static_assert(std::is_base_of<AlDagNode, T>::value);
-			T* LocalPtr = TAlObjectPtr<T>::Ptr;
-			return LocalPtr && LocalPtr->layer() ? StringCast<FString,const char*>(LocalPtr->layer()->name()) : FString();
+			return LayerName.IsSet() ? LayerName.GetValue() : FString();
 		}
 
-		TAlObjectPtr<AlLayer> GetLayer() const
+		const TAlObjectPtr<AlLayer>& GetLayer() const
 		{
-			static_assert(std::is_base_of<AlDagNode, T>::value);
-			T* LocalPtr = TAlObjectPtr<T>::Ptr;
-			return TAlObjectPtr<AlLayer>(LocalPtr && LocalPtr->layer() ? LocalPtr->layer() : nullptr);
+			return CachedLayer;
 		}
 
 		bool HasSymmetry() const
 		{
-			static_assert(std::is_base_of<AlDagNode, T>::value);
-			T* LocalPtr = TAlObjectPtr<T>::Ptr;
-			return (AlIsValid(LocalPtr) && LocalPtr->layer()) ? (bool)LocalPtr->layer()->isSymmetric() : false;
+			return CachedLayer ? (bool)CachedLayer->isSymmetric() : false;
+		}
+
+		bool IsVisible() const
+		{
+			return CachedLayer.IsValid() ? !((bool)CachedLayer->invisible()) : false;
+		}
+
+		AlDagNode* AsADagNode() const
+		{
+			return static_cast<AlDagNode*>(Super::Get());
+		}
+
+		bool HasGeometry() const
+		{ 
+			return EnumHasAnyFlags(Type, EDagNodeType::GeometryType);
+		}
+
+		bool IsAGroup() const
+		{
+			return EnumHasAnyFlags(Type, EDagNodeType::GroupType);
+		}
+
+		bool IsAMesh() const
+		{
+			return EnumHasAnyFlags(Type, EDagNodeType::MeshType);
+		}
+
+		bool IsASurface() const
+		{
+			return EnumHasAnyFlags(Type, EDagNodeType::SurfaceType);
 		}
 
-		TAlObjectPtr<AlMesh> GetMesh() const
+		bool IsAShell() const
 		{
-			static_assert(std::is_base_of<AlDagNode, T>::value);
-			AlMeshNode* MeshNode = ((AlDagNode*)TAlObjectPtr<T>::Ptr)->asMeshNodePtr();
-			return MeshNode ? TAlObjectPtr<AlMesh>(MeshNode->mesh()) : TAlObjectPtr<AlMesh>();
+			return EnumHasAnyFlags(Type, EDagNodeType::ShellType);
 		}
 
-		TAlObjectPtr<AlSurface> GetSurface() const
+		bool GetMesh(TAlObjectPtr<AlMesh>& OutMesh) const
 		{
-			static_assert(std::is_base_of<AlDagNode, T>::value);
-			AlSurfaceNode* SurfaceNode = ((AlDagNode*)TAlObjectPtr<T>::Ptr)->asSurfaceNodePtr();
-			return SurfaceNode ? TAlObjectPtr<AlSurface>(SurfaceNode->surface()) : TAlObjectPtr<AlSurface>();
+			OutMesh = IsAMesh() ? AsADagNode()->asMeshNodePtr()->mesh() : TAlObjectPtr<AlMesh>();
+
+			return OutMesh.IsValid();
 		}
 
-		TAlObjectPtr<AlShell> GetShell() const
+		bool GetSurface(TAlObjectPtr<AlSurface>& OutSurface) const
 		{
-			static_assert(std::is_base_of<AlDagNode, T>::value);
-			AlShellNode* ShellNode = ((AlDagNode*)TAlObjectPtr<T>::Ptr)->asShellNodePtr();
-			return ShellNode ? TAlObjectPtr<AlShell>(ShellNode->shell()) : TAlObjectPtr<AlShell>();
+			OutSurface = IsASurface() ? AsADagNode()->asSurfaceNodePtr()->surface() : TAlObjectPtr<AlSurface>();
+
+			return OutSurface.IsValid();
 		}
+
+		bool GetShell(TAlObjectPtr<AlShell>& OutShell) const
+		{
+			OutShell = IsAShell() ? AsADagNode()->asShellNodePtr()->shell() : TAlObjectPtr<AlShell>();
+
+			return OutShell.IsValid();
+		}
+
+		CADLibrary::FMeshParameters GetMeshParameters() const;
+
+		void SetActorTransform(IDatasmithActorElement& ActorElement) const
+		{
+			// Node with symmetry cannot be baked with the global transform because the symmetry is done in the parent referential
+			if (HasSymmetry())
+			{
+				return;
+			}
+
+			AlMatrix4x4 AlGlobalMatrix;
+			AsADagNode()->globalTransformationMatrix(AlGlobalMatrix);
+
+			FMatrix GlobalMatrix;
+			double* MatrixFloats = (double*)GlobalMatrix.M;
+			for (int32 IndexI = 0; IndexI < 4; ++IndexI)
+			{
+				for (int32 IndexJ = 0; IndexJ < 4; ++IndexJ)
+				{
+					MatrixFloats[IndexI * 4 + IndexJ] = AlGlobalMatrix[IndexI][IndexJ];
+				}
+			}
+
+			FTransform GlobalTransform = FDatasmithUtils::ConvertTransform(FDatasmithUtils::EModelCoordSystem::ZUp_RightHanded, FTransform(GlobalMatrix));
+
+			ActorElement.SetTranslation(GlobalTransform.GetTranslation());
+			ActorElement.SetScale(GlobalTransform.GetScale3D());
+			ActorElement.SetRotation(GlobalTransform.GetRotation());
+		}
+
+	private:
+		mutable TOptional<FString> LayerName;
+		mutable TAlObjectPtr<AlLayer> CachedLayer;
+		bool bCanDeleteObject = false;
+		EDagNodeType Type;
 	};
 
 	class FPatchMesh
 	{
 	public:
-		FPatchMesh(const FString& InName, const TAlObjectPtr<AlLayer>& InLayer, int32 Count)
+		FPatchMesh(const FString& InName, TAlObjectPtr<AlLayer>& InLayer, int32 Count)
 			: Name(InName)
 			, Layer(InLayer)
 		{
@@ -210,14 +390,17 @@ namespace UE_DATASMITHWIRETRANSLATOR_NAMESPACE
 			return MeshNodes.Num() == 1;
 		}
 
-		TAlDagNodePtr<AlMeshNode> GetSingleContent() const
+		bool GetSingleContent(FAlDagNodePtr& OutMeshNode)
 		{
 			if (MeshNodes.Num() == 1)
 			{
-				return MeshNodes[0];
+				FAlDagNodePtr& MeshNode = MeshNodes.Last();
+				OutMeshNode = MoveTemp(MeshNode);
+				MeshNodes.SetNum(0, EAllowShrinking::No);
+				return true;
 			}
 
-			return TAlDagNodePtr<AlMeshNode>();
+			return false;
 		}
 
 		const FString& GetName() const { return Name; }
@@ -226,18 +409,17 @@ namespace UE_DATASMITHWIRETRANSLATOR_NAMESPACE
 
 		const FString& GetUniqueID() const { return UniqueID; }
 
-		const TAlObjectPtr<AlLayer> GetLayer() const { return Layer; }
+		const TAlObjectPtr<AlLayer>& GetLayer() const { return Layer; }
 
-		void AddMeshNode(const TAlDagNodePtr<AlMeshNode>& MeshNode)
+		void AddMeshNode(FAlDagNodePtr& MeshNode)
 		{
 			ensure(Layer == MeshNode.GetLayer());
-			ensure(MeshNode->mesh());
 			MeshNodes.Add(MeshNode);
 		}
 
-		void IterateOnMeshNodes(const TFunction<void(const TAlDagNodePtr<AlMeshNode>& MeshNode)>& Callback) const
+		void IterateOnMeshNodes(const TFunction<void(const FAlDagNodePtr& MeshNode)>& Callback) const
 		{
-			for (const TAlDagNodePtr<AlMeshNode>& MeshNode : MeshNodes)
+			for (const FAlDagNodePtr& MeshNode : MeshNodes)
 			{
 				Callback(MeshNode);
 			}
@@ -247,7 +429,7 @@ namespace UE_DATASMITHWIRETRANSLATOR_NAMESPACE
 
 	private:
 		FString Name;
-		TArray<TAlDagNodePtr<AlMeshNode>> MeshNodes;
+		TArray<FAlDagNodePtr> MeshNodes;
 		TAlObjectPtr<AlLayer> Layer;
 		uint32 Hash;
 		FString UniqueID;
@@ -257,42 +439,34 @@ namespace UE_DATASMITHWIRETRANSLATOR_NAMESPACE
 	class FBodyNode
 	{
 	public:
-		FBodyNode(const FString& InName, const TAlObjectPtr<AlLayer>& InLayer, int32 Count)
+		FBodyNode(const FString& InName, TAlObjectPtr<AlLayer>& InLayer, int32 Count)
 			: Name(InName)
 			, Layer(InLayer)
 		{
-			SurfaceNodes.Reserve(Count);
-			ShellNodes.Reserve(Count);
+			DagNodes.Reserve(Count);
 		}
 
 		bool HasContent()
 		{
-			return Initialize() && (!SurfaceNodes.IsEmpty() || !ShellNodes.IsEmpty());
+			return Initialize() && !DagNodes.IsEmpty();
 		}
 		 
 		bool HasSingleContent() const
 		{
-			return (SurfaceNodes.Num() + ShellNodes.Num()) == 1;
+			return DagNodes.Num() == 1;
 		}
 
-		TAlDagNodePtr<AlDagNode> GetSingleContent() const
+		bool GetSingleContent(FAlDagNodePtr& OutDagNode) const
 		{
-			if (!HasSingleContent())
-			{
-				return TAlDagNodePtr<AlDagNode>();
-			}
-
-			if (SurfaceNodes.Num() == 1)
+			if (HasSingleContent())
 			{
-				return TAlDagNodePtr<AlDagNode>(SurfaceNodes[0].Get());
+				FAlDagNodePtr&  DagNode = DagNodes.Last();
+				OutDagNode = MoveTemp(DagNode);
+				DagNodes.SetNum(0, EAllowShrinking::No);
+				return true;
 			}
 
-			if (ShellNodes.Num() == 1)
-			{
-				return TAlDagNodePtr<AlDagNode>(ShellNodes[0].Get());
-			}
-
-			return TAlDagNodePtr<AlDagNode>();
+			return false;
 		}
 
 		const FString& GetName() const { return Name; }
@@ -301,25 +475,18 @@ namespace UE_DATASMITHWIRETRANSLATOR_NAMESPACE
 
 		const FString& GetUniqueID() const { return UniqueID; }
 
-		const TAlObjectPtr<AlLayer> GetLayer() const { return Layer; }
-
-		void AddSurfaceNode(const TAlDagNodePtr<AlSurfaceNode>& SurfaceNode);
+		const TAlObjectPtr<AlLayer>& GetLayer() const { return Layer; }
 
-		void AddShellNode(const TAlDagNodePtr<AlShellNode>& ShellNode);
-
-		void IterateOnSurfaceNodes(const TFunction<void(const TAlDagNodePtr<AlSurfaceNode>& SurfaceNode)>& Callback) const
-		{
-			for (const TAlDagNodePtr<AlSurfaceNode>& SurfaceNode : SurfaceNodes)
-			{
-				Callback(SurfaceNode);
-			}
-		}
+		bool AddNode(FAlDagNodePtr& DagNode);
 
-		void IterateOnShellNodes(const TFunction<void(const TAlDagNodePtr<AlShellNode>& ShellNode)>& Callback) const
+		void IterateOnDagNodes(const TFunction<void(const FAlDagNodePtr&)>& Callback) const
 		{
-			for (const TAlDagNodePtr<AlShellNode>& ShellNode : ShellNodes)
+			for (const FAlDagNodePtr& DagNode : DagNodes)
 			{
-				Callback(ShellNode);
+				if (DagNode)
+				{
+					Callback(DagNode);
+				}
 			}
 		}
 
@@ -333,15 +500,11 @@ namespace UE_DATASMITHWIRETRANSLATOR_NAMESPACE
 
 		bool Initialize();
 
-		int32 GetSlotIndex(const TAlDagNodePtr<AlDagNode>& DagNode);
-
-		bool HasSurfaceNodes() const { return SurfaceNodes.Num() > 0; }
-		bool HasShellNodes() const { return ShellNodes.Num() > 0; }
+		int32 GetSlotIndex(const FAlDagNodePtr& DagNode);
 
 	private:
 		FString Name;
-		TArray<TAlDagNodePtr<AlSurfaceNode>> SurfaceNodes;
-		TArray<TAlDagNodePtr<AlShellNode>> ShellNodes;
+		mutable TArray<FAlDagNodePtr> DagNodes;
 		TAlObjectPtr<AlLayer> Layer;
 		TMap<FString, int> ShaderNameToSlotIndex;
 		TMap<int, TAlObjectPtr<AlShader>> SlotIndexToShader;
@@ -372,13 +535,13 @@ namespace UE_DATASMITHWIRETRANSLATOR_NAMESPACE
 
 	struct FDagNodeGeometry : public FAliasGeometry
 	{
-		TAlDagNodePtr<AlDagNode> DagNode;
+		const FAlDagNodePtr& DagNode;
 
-		FDagNodeGeometry(int32 InType, EAliasObjectReference InReference, const TAlDagNodePtr<AlDagNode>& InDagNode)
+		FDagNodeGeometry(int32 InType, EAliasObjectReference InReference, const FAlDagNodePtr& InDagNode)
+			: DagNode(InDagNode)
 		{
 			Type = InType;
 			Reference = InReference;
-			DagNode = InDagNode;
 		}
 	};
 
@@ -401,8 +564,6 @@ namespace UE_DATASMITHWIRETRANSLATOR_NAMESPACE
 
 		bool ActorHasContent(const TSharedPtr<IDatasmithActorElement>& ActorElement);
 
-		void SetActorTransform(IDatasmithActorElement& OutActorElement, const TAlDagNodePtr<AlDagNode>& InDagNode);
-
 		bool IsValidActor(const TSharedPtr<IDatasmithActorElement>& ActorElement);
 
 		inline FString UuidToString(const uint32& Uuid)
@@ -431,15 +592,8 @@ namespace UE_DATASMITHWIRETRANSLATOR_NAMESPACE
 
 		bool TransferAlMeshToMeshDescription(const AlMesh& Mesh, const TCHAR* SlotMaterialName, FMeshDescription& MeshDescription, CADLibrary::FMeshParameters& SymmetricParameters, const bool bMerge = false);
 
-		TAlDagNodePtr<AlMeshNode> TesselateDagLeaf(const AlDagNode& DagLeaf, ETesselatorType TessType, double Tolerance);
-
-		/** Returns true if DagNode is supported and its content is valid */
-		bool IsDagNodeValid(const TAlDagNodePtr<AlDagNode>& DagNode);
-
-		/** Returns true if geometry node is supported and its content is valid */
-		bool IsGeometryValid(const TAlDagNodePtr<AlDagNode>& GeomNode);
+		FAlDagNodePtr TesselateDagLeaf(const AlDagNode& DagLeaf, ETesselatorType TessType, double Tolerance);
 
-		CADLibrary::FMeshParameters GetMeshParameters(AlDagNode& DagNode);
 		CADLibrary::FMeshParameters GetMeshParameters(const TAlObjectPtr<AlLayer>& Layer);
 	}
 
diff --git a/Engine/Plugins/Enterprise/DatasmithCADImporter/Source/WireInterface/Private/WireInterfaceImpl.cpp b/Engine/Plugins/Enterprise/DatasmithCADImporter/Source/WireInterface/Private/WireInterfaceImpl.cpp
--- a/Engine/Plugins/Enterprise/DatasmithCADImporter/Source/WireInterface/Private/WireInterfaceImpl.cpp
+++ b/Engine/Plugins/Enterprise/DatasmithCADImporter/Source/WireInterface/Private/WireInterfaceImpl.cpp
@@ -23,13 +23,6 @@
 #include "StaticMeshDescription.h"
 #include "StaticMeshOperations.h"
 
-#if WITH_EDITOR
-#include "Editor.h"
-#include "IMessageLogListing.h"
-#include "Logging/TokenizedMessage.h"
-#include "MessageLogModule.h"
-#endif
-
 #include "AliasModelToCADKernelConverter.h"
 #include "AliasModelToTechSoftConverter.h" // requires Techsoft as public dependency
 #include "CADInterfacesModule.h"
@@ -65,41 +58,41 @@
 #include "Windows/HideWindowsPlatformTypes.h"
 #endif
 
-DEFINE_LOG_CATEGORY_STATIC(LogWireInterface, Log, All);
-
 #define LOCTEXT_NAMESPACE "WireInterface"
 
 #define WRONG_VERSION_TEXT "Unsupported version of Alias detected. Please upgrade to Alias 2021.3 (or later version)."
 
 #define TRACK_MESHELEMENT 0
+#define MAKE_VISIBLE 0
 
-#if TRACK_MESHELEMENT
+#if MAKE_VISIBLE
 #include "CompGeom/FitOrientedBox3.h"
 #endif
 
-#ifdef OPEN_MODEL_2023_0
-
-static bool bGAliasSewByMaterial = false;
-FAutoConsoleVariableRef GAliasSewByMaterial(
-	TEXT("ds.CADTranslator.Alias.SewByMaterial"),
-	bGAliasSewByMaterial,
-	TEXT("Enable Sew action merges BReps according to their material i.e. only BReps associated with same material can be merged together.\
-Default is disable\n"),
-ECVF_Default);
-
-#endif
-
 namespace UE_DATASMITHWIRETRANSLATOR_NAMESPACE
 {
 #if TRACK_MESHELEMENT
+	//static const FString ToTrack(TEXT("shell_30181"));
+	//static const FString ToTrack(TEXT("Blend_srf_1395"));
+	//static const FString ToTrack(TEXT("M0047942"));
+	//static const FString ToTrack(TEXT("C5_31495_surf")); // G6
+	static const FString ToTrack(TEXT("Object__TRANSFORM_copy_2_surf")); // refrigerator
+
 	void MakeMeshVisible(FMeshDescription& MeshDescription);
 #endif
 
-	static bool bGSewByMaterial = false;
+#if WIRE_MEMORY_CHECK
+	int32 AllocatedObjects = 0;
+	int32 MaxAllocatedObjects = 0;
+	TSet<AlDagNode*> DagNodeSet;
+	TSet<AlObject*> ObjectSet;
+#endif
 
 	static const FColor DefaultColor = FColor(200, 200, 200);
 
 	const uint64 LibAliasNext_Version = 0xffffffffffffffffull;
+	const uint64 LibAlias2025_0_0_Version = 8725724278030572;
+	const uint64 LibAlias2024_1_0_Version = 8444253596292024;
 	const uint64 LibAlias2023_1_0_Version = 8162778619576619;
 	const uint64 LibAlias2023_0_0_Version = 8162774324609149;
 	const uint64 LibAlias2022_2_0_Version = 7881307937833405;
@@ -138,8 +131,16 @@ namespace UE_DATASMITHWIRETRANSLATOR_NAMESPACE
 	const FString AliasSdkVersion   = TEXT("2023.0");
 #elif defined(OPEN_MODEL_2023_1)      
 	const uint64 LibAliasVersionMin = LibAlias2023_1_0_Version;
+	const uint64 LibAliasVersionMax = LibAlias2024_1_0_Version;
+	const FString AliasSdkVersion = TEXT("2023.1");
+#elif defined(OPEN_MODEL_2024_1)      
+	const uint64 LibAliasVersionMin = LibAlias2024_1_0_Version;
+	const uint64 LibAliasVersionMax = LibAlias2025_0_0_Version;
+	const FString AliasSdkVersion = TEXT("2024.1");
+#elif defined(OPEN_MODEL_2025_0)      
+	const uint64 LibAliasVersionMin = LibAlias2025_0_0_Version;
 	const uint64 LibAliasVersionMax = LibAliasNext_Version;
-	const FString AliasSdkVersion   = TEXT("2023.1");
+	const FString AliasSdkVersion = TEXT("2025.0");
 #endif
 
 	// Alias material management (to allow sew of BReps of different materials):
@@ -196,14 +197,29 @@ namespace UE_DATASMITHWIRETRANSLATOR_NAMESPACE
 
 	bool FWireTranslatorImpl::Initialize(const TCHAR* InSceneFullName)
 	{
-		// #wire_import: Check file can be loaded with current SDK
-
 		if (InSceneFullName)
 		{
-			SceneFullPath = InSceneFullName;
+			statusCode Status = AlUniverse::initialize();
+			if (Status != sSuccess && Status != sAlreadyCreated)
+			{
+				UE_LOG(LogWireInterface, Error, TEXT("Cannot initialize OpenModel SDK. Import is aborted."));
+				return false;
+			}
+
+			const TStringConversion Result = StringCast<ANSICHAR>(InSceneFullName);
+			char OpenModelVersion[10];
+			if (bool(AlUniverse::isWireFile(Result.Get(), OpenModelVersion)))
+			{
+				SceneFullPath = InSceneFullName;
+				SceneVersion = StringCast<TCHAR>(OpenModelVersion).Get();
+			}
+			else
+			{
+				UE_LOG(LogWireInterface, Error, TEXT("Cannot load %s with the selected OpenModel SDK."), InSceneFullName);
+				return false;
+			}
 		}
 
-		// #wire_import: Check whether this implementation can load it
 		return true;
 	}
 
@@ -211,7 +227,29 @@ namespace UE_DATASMITHWIRETRANSLATOR_NAMESPACE
 	{
 		if (bSceneLoaded)
 		{
-			AlUniverse::deleteAll();
+			FLayerContainer::Reset();
+			MeshElementToParametricNode.Reset();
+			MeshElementToMeshNode.Reset();
+			MeshElementToBodyNode.Reset();
+			MeshElementToPatchMesh.Reset();
+			EncounteredNodes.Reset();
+#if WIRE_MEMORY_CHECK
+			if (!ensure(AllocatedObjects == 0))
+			{
+				for (AlObject* Object : ObjectSet)
+				{
+					if (AlIsValid(Object))
+					{
+						AlObjectType Type = Object->type();
+						printf(">>> %d\n", Type);
+					}
+					else
+					{
+						ensure(false);
+					}
+				}
+			}
+#endif
 		}
 		bSceneLoaded = false;
 	}
@@ -219,6 +257,7 @@ namespace UE_DATASMITHWIRETRANSLATOR_NAMESPACE
 	void FWireTranslatorImpl::SetImportSettings(const FWireSettings& Settings)
 	{
 		WireSettings = Settings;
+		WireSettings.bUseCADKernel = !CADLibrary::FImportParameters::bGDisableCADKernelTessellation;
 		if (CADModelConverter)
 		{
 			CADModelConverter->SetImportParameters(Settings.ChordTolerance, Settings.MaxEdgeLength, Settings.NormalTolerance, (CADLibrary::EStitchingTechnique)Settings.StitchingTechnique);
@@ -228,6 +267,47 @@ namespace UE_DATASMITHWIRETRANSLATOR_NAMESPACE
 	// Wire file parsing
 	bool FWireTranslatorImpl::Load(TSharedPtr<IDatasmithScene> InScene)
 	{
+		UE_LOG(LogWireInterface, Display, TEXT("CAD translation [%s]."), *SceneFullPath);
+		UE_LOG(LogWireInterface, Display, TEXT(" - File version:         Alias %s"), *SceneVersion);
+		UE_LOG(LogWireInterface, Display, TEXT(" - Parsing Library:      Alias %s"), *AliasSdkVersion);
+		UE_LOG(LogWireInterface, Display, TEXT(" - Tessellation Library: %s"), WireSettings.bUseCADKernel ? TEXT("CADKernel") : TEXT("TechSoft"));
+
+		UE_LOG(LogWireInterface, Display, TEXT(" - Translation parameters:"));
+		UE_LOG(LogWireInterface, Display, TEXT("     - Merge By Group:      %s"), WireSettings.bMergeGeometryByGroup ? TEXT("True") : TEXT("False"));
+		UE_LOG(LogWireInterface, Display, TEXT("     - Layer As actor:      %s"), WireSettings.bUseLayerAsActor ? TEXT("True") : TEXT("False"));
+
+		UE_LOG(LogWireInterface, Display, TEXT(" - Tessellation parameters:"));
+		UE_LOG(LogWireInterface, Display, TEXT("     - ChordTolerance:      %lf"), WireSettings.ChordTolerance);
+		UE_LOG(LogWireInterface, Display, TEXT("     - MaxEdgeLength:       %lf"), WireSettings.MaxEdgeLength);
+		UE_LOG(LogWireInterface, Display, TEXT("     - MaxNormalAngle:      %lf"), WireSettings.NormalTolerance);
+
+		FString StitchingTechnique;
+		switch (WireSettings.StitchingTechnique)
+		{
+		case EDatasmithCADStitchingTechnique::StitchingHeal:
+			StitchingTechnique = TEXT("Heal");
+			break;
+		case EDatasmithCADStitchingTechnique::StitchingSew:
+			StitchingTechnique = TEXT("Sew");
+			break;
+		default:
+			StitchingTechnique = TEXT("None");
+			break;
+		}
+		UE_LOG(LogWireInterface, Display, TEXT("     - StitchingTechnique:  %s"), *StitchingTechnique);
+		UE_LOG(LogWireInterface, Display, TEXT("     - GeometricTolerance:  %lf"), WireSettings.GetGeometricTolerance());
+		UE_LOG(LogWireInterface, Display, TEXT("     - Stitching Tolerance: %lf"), WireSettings.GetStitchingTolerance());
+
+		// #wire_import: TODO - Revisit stitching extension when using CADKernel
+		//if (WireSettings.bUseCADKernel)
+		//{
+		//	UE_LOG(LogWireInterface, Display, TEXT("     - Stitching Options:"));
+		//	UE_LOG(LogWireInterface, Display, TEXT("         - ForceSew:              %s"), CADLibrary::FImportParameters::bGStitchingForceSew ? TEXT("True") : TEXT("False"));
+		//	UE_LOG(LogWireInterface, Display, TEXT("         - RemoveThinFaces:       %s"), CADLibrary::FImportParameters::bGStitchingRemoveThinFaces ? TEXT("True") : TEXT("False"));
+		//	UE_LOG(LogWireInterface, Display, TEXT("         - RemoveDuplicatedFaces: %s"), CADLibrary::FImportParameters::bGStitchingRemoveDuplicatedFaces ? TEXT("True") : TEXT("False"));
+		//	UE_LOG(LogWireInterface, Display, TEXT("         - ForceFactor:           %f"), CADLibrary::FImportParameters::GStitchingForceFactor);
+		//}
+
 		DatasmithScene = InScene;
 
 		const FString AliasProductVersion = FString::Printf(TEXT("Alias %s"), *AliasSdkVersion);
@@ -256,10 +336,13 @@ namespace UE_DATASMITHWIRETRANSLATOR_NAMESPACE
 			WireSettings.bMergeGeometryByGroup = false;
 		}
 
-		bGSewByMaterial = GetConsoleBoolValue(TEXT("ds.CADTranslator.Alias.SewByMaterial"), false);
-
 		// Initialize Alias.
-		AlUniverse::initialize();
+		statusCode Status = AlUniverse::initialize();
+		if (Status != sSuccess && Status != sAlreadyCreated)
+		{
+			UE_LOG(LogWireInterface, Error, TEXT("Cannot initialize OpenModel SDK. Import is aborted."));
+			return false;
+		}
 
 		if (AlUniverse::retrieve(TCHAR_TO_UTF8(*SceneFullPath)) != sSuccess)
 		{
@@ -276,7 +359,13 @@ namespace UE_DATASMITHWIRETRANSLATOR_NAMESPACE
 
 	bool FWireTranslatorImpl::TraverseModel()
 	{
-		TAlDagNodePtr<AlDagNode> DagNode(AlUniverse::firstDagNode());
+#if WIRE_MEMORY_CHECK
+		AllocatedObjects = 0;
+		MaxAllocatedObjects = 0;
+		ObjectSet.Reset();
+		DagNodeSet.Reset();
+#endif
+		FAlDagNodePtr DagNode = FindOrAddDagNode(AlUniverse::firstDagNode());
 		while (DagNode)
 		{
 			TSharedPtr<IDatasmithActorElement> ActorElement = TraverseDag(DagNode);
@@ -285,17 +374,19 @@ namespace UE_DATASMITHWIRETRANSLATOR_NAMESPACE
 				DatasmithScene->AddActor(ActorElement);
 			}
 
-			DagNode = TAlDagNodePtr<AlDagNode>(DagNode->nextNode());
+			DagNode = FindOrAddDagNode(DagNode->nextNode());
 		}
 
 		TAlObjectPtr<AlSet> Set(AlUniverse::firstSet());
-		while (Set)
+		statusCode Status = Set ? sSuccess : sFailure;
+		while (Status == sSuccess)
 		{
 			// #wire_import: Add an actor to represent the set.
 			TAlObjectPtr<AlSetMember> SetMember(Set->firstMember());
-			while (SetMember)
+			statusCode MemberStatus = SetMember ? sSuccess : sFailure;
+			while (MemberStatus == sSuccess)
 			{
-				TAlDagNodePtr<AlDagNode> DagNodeInSet(SetMember->object() ? SetMember->object()->asDagNodePtr() : nullptr);
+				FAlDagNodePtr DagNodeInSet(SetMember->object() ? SetMember->object()->asDagNodePtr() : nullptr);
 				if (DagNodeInSet)
 				{
 					TSharedPtr<IDatasmithActorElement> ActorElement = TraverseDag(DagNodeInSet);
@@ -305,23 +396,22 @@ namespace UE_DATASMITHWIRETRANSLATOR_NAMESPACE
 					}
 				}
 
-				SetMember = TAlObjectPtr<AlSetMember>(SetMember->nextSetMember());
+				MemberStatus = SetMember->nextSetMemberD();
 			}
 
-			Set = TAlObjectPtr<AlSet>(Set->nextSet());
+			Status = Set->nextSetD();
 		}
 
 		return true;
 	}
 
-	TSharedPtr<IDatasmithActorElement> FWireTranslatorImpl::TraverseDag(const TAlDagNodePtr<AlDagNode>& RootNode)
+	TSharedPtr<IDatasmithActorElement> FWireTranslatorImpl::TraverseDag(const FAlDagNodePtr& RootNode)
 	{
-		TAlDagNodePtr<AlGroupNode> GroupNode(RootNode->asGroupNodePtr());
-		if (GroupNode)
+		if (RootNode.IsAGroup())
 		{
-			return WireSettings.bMergeGeometryByGroup ? ProcessGroupNode(GroupNode) : TraverseGroupNode(GroupNode);
+			return WireSettings.bMergeGeometryByGroup ? ProcessGroupNode(RootNode) : TraverseGroupNode(RootNode);
 		}
-		else if (OpenModelUtils::IsGeometryValid(RootNode))
+		else if (RootNode.HasGeometry())
 		{
 			return ProcessGeometryNode(RootNode);
 		}
@@ -329,16 +419,23 @@ namespace UE_DATASMITHWIRETRANSLATOR_NAMESPACE
 		return TSharedPtr<IDatasmithActorElement>();
 	}
 
-	TSharedPtr<IDatasmithActorElement> FWireTranslatorImpl::TraverseGroupNode(const TAlDagNodePtr<AlGroupNode>& GroupNode, const TAlObjectPtr<AlLayer>& ParentLayer)
+	TSharedPtr<IDatasmithActorElement> FWireTranslatorImpl::TraverseGroupNode(const FAlDagNodePtr& GroupNode, const TAlObjectPtr<AlLayer>& ParentLayer)
 	{
-		TAlDagNodePtr<AlDagNode> ChildNode(GroupNode->childNode());
-		int32 ChildrenCount = 0;
+		if (!GroupNode.IsAGroup())
+		{
+			return TSharedPtr<IDatasmithActorElement>();
+		}
+
+		TArray<FAlDagNodePtr> Children;
+		
+		FAlDagNodePtr ChildNode = FindOrAddDagNode(GroupNode->asGroupNodePtr()->childNode());
 		while (ChildNode)
 		{
-			ChildrenCount++;
-			ChildNode = TAlDagNodePtr<AlDagNode>(ChildNode->nextNode());
+			Children.Add(ChildNode);
+			ChildNode = FindOrAddDagNode(ChildNode->nextNode());
 		}
 
+		const int32 ChildrenCount = Children.Num();
 		if (ChildrenCount == 0)
 		{
 			return TSharedPtr<IDatasmithActorElement>();
@@ -347,27 +444,20 @@ namespace UE_DATASMITHWIRETRANSLATOR_NAMESPACE
 		TArray<TSharedPtr<IDatasmithActorElement>> ChildActors;
 		ChildActors.Reserve(ChildrenCount);
 
-		ChildNode = TAlDagNodePtr<AlDagNode>(GroupNode->childNode());
-		while (ChildNode)
+		TAlObjectPtr<AlLayer> GroupLayer = FLayerContainer::FindOrAdd(GroupNode->layer());
+		for (FAlDagNodePtr& Child : Children)
 		{
-			TAlDagNodePtr<AlGroupNode> ChildGroupNode(ChildNode->asGroupNodePtr());
-			if (ChildGroupNode)
+			if (TSharedPtr<IDatasmithActorElement> ActorElement = TraverseGroupNode(Child, GroupLayer))
 			{
-				if (TSharedPtr<IDatasmithActorElement> ActorElement = TraverseGroupNode(ChildGroupNode, GroupNode->layer()))
-				{
-					ChildActors.Emplace(MoveTemp(ActorElement));
-				}
-
+				ChildActors.Emplace(MoveTemp(ActorElement));
 			}
-			else if (OpenModelUtils::IsGeometryValid(ChildNode))
+			else if (Child.IsAMesh() || Child.IsASurface())
 			{
-				if (TSharedPtr<IDatasmithActorElement> ActorElement = ProcessGeometryNode(ChildNode, ParentLayer))
+				if (TSharedPtr<IDatasmithActorElement> ChildActorElement = ProcessGeometryNode(Child, ParentLayer))
 				{
-					ChildActors.Emplace(MoveTemp(ActorElement));
+					ChildActors.Emplace(MoveTemp(ChildActorElement));
 				}
 			}
-
-			ChildNode = TAlDagNodePtr<AlDagNode>(ChildNode->nextNode());
 		}
 
 		if (ChildActors.Num() == 0)
@@ -384,14 +474,13 @@ namespace UE_DATASMITHWIRETRANSLATOR_NAMESPACE
 		FString Label = GroupNode.GetName();
 		ActorElement->SetLabel(Label.Len() > 0 ? *Label : TEXT("UnnamedGroup"));
 
-		TAlObjectPtr<AlLayer> Layer(GroupNode->layer());
 		FString CsvLayerString;
-		if (OpenModelUtils::GetCsvLayerString(Layer, CsvLayerString))
+		if (OpenModelUtils::GetCsvLayerString(GroupLayer, CsvLayerString))
 		{
 			ActorElement->SetLayer(*CsvLayerString);
 		}
 
-		OpenModelUtils::SetActorTransform(*ActorElement, GroupNode.Get());
+		GroupNode.SetActorTransform(*ActorElement);
 
 		for (TSharedPtr<IDatasmithActorElement>& ChildActor : ChildActors)
 		{
@@ -401,11 +490,11 @@ namespace UE_DATASMITHWIRETRANSLATOR_NAMESPACE
 			}
 		}
 
-		if (WireSettings.bUseLayerAsActor && Layer != ParentLayer)
+		if (WireSettings.bUseLayerAsActor && GroupLayer != ParentLayer)
 		{
 			if (OpenModelUtils::ActorHasContent(ActorElement))
 			{
-				if (TSharedPtr<IDatasmithActorElement> LayerActor = FindOrAddLayerActor(Layer))
+				if (TSharedPtr<IDatasmithActorElement> LayerActor = FindOrAddLayerActor(GroupLayer))
 				{
 					LayerActor->AddChild(ActorElement);
 				}
@@ -417,16 +506,26 @@ namespace UE_DATASMITHWIRETRANSLATOR_NAMESPACE
 		return ActorElement;
 	}
 
-	TSharedPtr<IDatasmithActorElement> FWireTranslatorImpl::ProcessGroupNode(const TAlDagNodePtr<AlGroupNode>& GroupNode, const TAlObjectPtr<AlLayer>& ParentLayer)
+	TSharedPtr<IDatasmithActorElement> FWireTranslatorImpl::ProcessGroupNode(const FAlDagNodePtr& GroupNode, const TAlObjectPtr<AlLayer>& ParentLayer)
 	{
-		TAlDagNodePtr<AlDagNode> ChildNode(GroupNode->childNode());
-		int32 ChildrenCount = 0;
+		if (!GroupNode.IsAGroup())
+		{
+			return TSharedPtr<IDatasmithActorElement>();
+		}
+
+		TArray<FAlDagNodePtr> Children;
+		FAlDagNodePtr ChildNode = FindOrAddDagNode(GroupNode->asGroupNodePtr()->childNode());
 		while (ChildNode)
 		{
-			ChildrenCount++;
-			ChildNode = TAlDagNodePtr<AlDagNode>(ChildNode->nextNode());
+			Children.Add(ChildNode);
+			AlDagNode* NextNode = ChildNode->nextNode();
+#if WIRE_MEMORY_CHECK
+			ensure(NextNode != ChildNode.Get());
+#endif
+			ChildNode = FindOrAddDagNode(NextNode);
 		}
 
+		const int32 ChildrenCount = Children.Num();
 		if (ChildrenCount == 0)
 		{
 			return TSharedPtr<IDatasmithActorElement>();
@@ -434,45 +533,26 @@ namespace UE_DATASMITHWIRETRANSLATOR_NAMESPACE
 
 		TArray<TSharedPtr<IDatasmithActorElement>> ChildActors;
 		ChildActors.Reserve(ChildrenCount);
+		
+		TAlObjectPtr<AlLayer> GroupLayer = FLayerContainer::FindOrAdd(GroupNode->layer());
 
-		TSharedPtr<FBodyNode> BodyNode = MakeShared<FBodyNode>(GroupNode.GetName() + TEXT("_surf"), GroupNode->layer(), ChildrenCount);
-		TSharedPtr<FPatchMesh> PatchMesh = MakeShared<FPatchMesh>(GroupNode.GetName() + TEXT("_mesh"), GroupNode->layer(), ChildrenCount);
+		TSharedPtr<FBodyNode> BodyNode = MakeShared<FBodyNode>(GroupNode.GetName() + TEXT("_surf"), GroupLayer, ChildrenCount);
+		TSharedPtr<FPatchMesh> PatchMesh = MakeShared<FPatchMesh>(GroupNode.GetName() + TEXT("_mesh"), GroupLayer, ChildrenCount);
 
-		ChildNode = TAlDagNodePtr<AlDagNode>(GroupNode->childNode());
-		while (ChildNode)
+		for(FAlDagNodePtr& Child : Children)
 		{
-			TAlDagNodePtr<AlGroupNode> ChildGroupNode(ChildNode->asGroupNodePtr());
-			if (ChildGroupNode)
+			if (TSharedPtr<IDatasmithActorElement> ActorElement = ProcessGroupNode(Child, GroupLayer))
 			{
-				if (TSharedPtr<IDatasmithActorElement> ActorElement = ProcessGroupNode(ChildGroupNode, GroupNode->layer()))
-				{
-					ChildActors.Emplace(MoveTemp(ActorElement));
-				}
-
+				ChildActors.Emplace(MoveTemp(ActorElement));
 			}
-			else if (OpenModelUtils::IsGeometryValid(ChildNode))
+			else if (Child.IsAMesh())
 			{
-				TAlDagNodePtr<AlMeshNode> MeshNode(ChildNode->asMeshNodePtr());
-				if (MeshNode && MeshNode->mesh())
-				{
-					PatchMesh->AddMeshNode(MeshNode);
-				}
-				TAlDagNodePtr<AlShellNode> ShellNode(ChildNode->asShellNodePtr());
-				if (ShellNode && ShellNode->shell())
-				{
-					BodyNode->AddShellNode(ShellNode);
-				}
-				else
-				{
-					TAlDagNodePtr<AlSurfaceNode> SurfaceNode(ChildNode->asSurfaceNodePtr());
-					if (SurfaceNode && SurfaceNode->surface())
-					{
-						BodyNode->AddSurfaceNode(SurfaceNode);
-					}
-				}
+				PatchMesh->AddMeshNode(Child);
+			}
+			else if (Child.IsASurface() || Child.IsAShell())
+			{
+				BodyNode->AddNode(Child);
 			}
-
-			ChildNode = TAlDagNodePtr<AlDagNode>(ChildNode->nextNode());
 		}
 
 		if (ChildActors.IsEmpty() && !BodyNode->Initialize() && !PatchMesh->Initialize())
@@ -480,7 +560,7 @@ namespace UE_DATASMITHWIRETRANSLATOR_NAMESPACE
 			return TSharedPtr<IDatasmithActorElement>();
 		}
 
-		if (TSharedPtr<IDatasmithActorElement> ActorElement = ProcessBodyNode(BodyNode, GroupNode, GroupNode->layer()))
+		if (TSharedPtr<IDatasmithActorElement> ActorElement = ProcessBodyNode(BodyNode, GroupNode, GroupLayer))
 		{
 			if (OpenModelUtils::ActorHasContent(ActorElement))
 			{
@@ -488,7 +568,7 @@ namespace UE_DATASMITHWIRETRANSLATOR_NAMESPACE
 			}
 		}
 
-		if (TSharedPtr<IDatasmithActorElement> ActorElement = ProcessPatchMesh(PatchMesh, GroupNode, GroupNode->layer()))
+		if (TSharedPtr<IDatasmithActorElement> ActorElement = ProcessPatchMesh(PatchMesh, GroupNode, GroupLayer))
 		{
 			if (OpenModelUtils::ActorHasContent(ActorElement))
 			{
@@ -519,16 +599,15 @@ namespace UE_DATASMITHWIRETRANSLATOR_NAMESPACE
 		FString Label = GroupNode.GetName();
 		ActorElement->SetLabel(Label.Len() > 0 ? *Label : TEXT("UnnamedGroup"));
 
-		TAlObjectPtr<AlLayer> Layer(GroupNode->layer());
 		FString CsvLayerString;
-		if (OpenModelUtils::GetCsvLayerString(Layer, CsvLayerString))
+		if (OpenModelUtils::GetCsvLayerString(GroupLayer, CsvLayerString))
 		{
 			ActorElement->SetLayer(*CsvLayerString);
 		}
 
 		if (WireSettings.bMergeGeometryByGroup && ChildActors.Num() > 1)
 		{
-			OpenModelUtils::SetActorTransform(*ActorElement, GroupNode.Get());
+			GroupNode.SetActorTransform(*ActorElement);
 
 			for (TSharedPtr<IDatasmithActorElement>& ChildActor : ChildActors)
 			{
@@ -540,11 +619,11 @@ namespace UE_DATASMITHWIRETRANSLATOR_NAMESPACE
 
 		}
 
-		if (WireSettings.bUseLayerAsActor && Layer != ParentLayer)
+		if (WireSettings.bUseLayerAsActor && GroupLayer != ParentLayer)
 		{
 			if (OpenModelUtils::ActorHasContent(ActorElement))
 			{
-				if (TSharedPtr<IDatasmithActorElement> LayerActor = FindOrAddLayerActor(Layer))
+				if (TSharedPtr<IDatasmithActorElement> LayerActor = FindOrAddLayerActor(GroupLayer))
 				{
 					LayerActor->AddChild(ActorElement);
 				}
@@ -556,9 +635,9 @@ namespace UE_DATASMITHWIRETRANSLATOR_NAMESPACE
 		return ActorElement;
 	}
 
-	TSharedPtr<IDatasmithActorElement> FWireTranslatorImpl::ProcessGeometryNode(const TAlDagNodePtr<AlDagNode>& GeomNode, const TAlObjectPtr<AlLayer>& ParentLayer)
+	TSharedPtr<IDatasmithActorElement> FWireTranslatorImpl::ProcessGeometryNode(const FAlDagNodePtr& GeomNode, const TAlObjectPtr<AlLayer>& ParentLayer)
 	{
-		TAlObjectPtr<AlLayer> Layer(GeomNode->layer());
+		const TAlObjectPtr<AlLayer>& Layer = GeomNode.GetLayer();
 		if (Layer && Layer->invisible())
 		{
 			return TSharedPtr<IDatasmithActorElement>();
@@ -586,15 +665,14 @@ namespace UE_DATASMITHWIRETRANSLATOR_NAMESPACE
 			ActorElement->SetLayer(*CsvLayerString);
 		}
 
-		OpenModelUtils::SetActorTransform(*ActorElement, GeomNode);
+		GeomNode.SetActorTransform(*ActorElement);
 
-		TAlDagNodePtr<AlShellNode> ShellNode(GeomNode->asShellNodePtr());
-		if (ShellNode && ShellNode->shell())
+		TAlObjectPtr<AlShell> Shell;
+		if (GeomNode.GetShell(Shell))
 		{
-			TAlObjectPtr<AlShell> Shell(ShellNode->shell());
 			TAlObjectPtr<AlShader> Shader(Shell->firstShader());
 			int32 SlotIndex = 0;
-			while (Shader.IsValid())
+			while (Shader)
 			{
 				TSharedPtr<IDatasmithMaterialIDElement> MaterialIDElement = FindOrAddMaterial(Shader);
 				MaterialIDElement->SetId(SlotIndex++);
@@ -606,17 +684,17 @@ namespace UE_DATASMITHWIRETRANSLATOR_NAMESPACE
 		{
 			TSharedPtr<IDatasmithMaterialIDElement> MaterialIDElement;
 
-			TAlDagNodePtr<AlSurfaceNode> SurfaceNode(GeomNode->asSurfaceNodePtr());
-			if (SurfaceNode && SurfaceNode->surface())
+			TAlObjectPtr<AlSurface> Surface;
+			if (GeomNode.GetSurface(Surface))
 			{
-				MaterialIDElement = FindOrAddMaterial(SurfaceNode->surface()->firstShader());
+				MaterialIDElement = FindOrAddMaterial(Surface->firstShader());
 			}
 			else
 			{
-				TAlDagNodePtr<AlMeshNode> MeshNode(GeomNode->asMeshNodePtr());
-				if (MeshNode && MeshNode->mesh())
+				TAlObjectPtr<AlMesh> Mesh;
+				if (GeomNode.GetMesh(Mesh))
 				{
-					MaterialIDElement = FindOrAddMaterial(MeshNode->mesh()->firstShader());
+					MaterialIDElement = FindOrAddMaterial(Mesh->firstShader());
 				}
 			}
 
@@ -635,22 +713,23 @@ namespace UE_DATASMITHWIRETRANSLATOR_NAMESPACE
 				return TSharedPtr<IDatasmithActorElement>();
 			}
 
-			ensure(false);
+			ensureWire(false);
 		}
 
 		return ActorElement;
 	}
 
-	TSharedPtr<IDatasmithActorElement> FWireTranslatorImpl::ProcessBodyNode(TSharedPtr<FBodyNode>& BodyNode, const TAlDagNodePtr<AlGroupNode>& GroupNode, const TAlObjectPtr<AlLayer>& ParentLayer)
+	TSharedPtr<IDatasmithActorElement> FWireTranslatorImpl::ProcessBodyNode(TSharedPtr<FBodyNode>& BodyNode, const FAlDagNodePtr& GroupNode, const TAlObjectPtr<AlLayer>& ParentLayer)
 	{
 		if (!BodyNode->HasContent())
 		{
 			return TSharedPtr<IDatasmithActorElement>();
 		}
 
-		if (BodyNode->HasSingleContent())
+		FAlDagNodePtr Singleton;
+		if (BodyNode->GetSingleContent(Singleton))
 		{
-			return ProcessGeometryNode(BodyNode->GetSingleContent(), ParentLayer);
+			return ProcessGeometryNode(Singleton, ParentLayer);
 		}
 
 		TSharedPtr<IDatasmithMeshElement> MeshElement = FindOrAddMeshElement(BodyNode);
@@ -675,9 +754,9 @@ namespace UE_DATASMITHWIRETRANSLATOR_NAMESPACE
 			ActorElement->SetLayer(*CsvLayerString);
 		}
 
-		if (!BodyNode->GetLayer()->isSymmetric())
+		if (BodyNode->GetLayer() && !BodyNode->GetLayer()->isSymmetric())
 		{
-			OpenModelUtils::SetActorTransform(*ActorElement, GroupNode.Get());
+			GroupNode.SetActorTransform(*ActorElement);
 		}
 
 		if (WireSettings.bUseLayerAsActor && BodyNode->GetLayer() != ParentLayer)
@@ -688,28 +767,28 @@ namespace UE_DATASMITHWIRETRANSLATOR_NAMESPACE
 				return TSharedPtr<IDatasmithActorElement>();
 			}
 
-			ensure(false);
+			ensureWire(false);
 		}
 
 		return ActorElement;
 	}
 
-	TSharedPtr<IDatasmithActorElement> FWireTranslatorImpl::ProcessPatchMesh(TSharedPtr<FPatchMesh>& PatchMesh, const TAlDagNodePtr<AlGroupNode>& GroupNode, const TAlObjectPtr<AlLayer>& ParentLayer)
+	TSharedPtr<IDatasmithActorElement> FWireTranslatorImpl::ProcessPatchMesh(TSharedPtr<FPatchMesh>& PatchMesh, const FAlDagNodePtr& GroupNode, const TAlObjectPtr<AlLayer>& ParentLayer)
 	{
 		if (!PatchMesh->HasContent())
 		{
 			return TSharedPtr<IDatasmithActorElement>();
 		}
 
-		TAlObjectPtr<AlLayer> Layer(GroupNode->layer());
-		if (Layer && Layer->invisible())
+		if (!GroupNode.IsVisible())
 		{
 			return TSharedPtr<IDatasmithActorElement>();
 		}
 
-		if (PatchMesh->HasSingleContent())
+		FAlDagNodePtr Singleton;
+		if (PatchMesh->GetSingleContent(Singleton))
 		{
-			return ProcessGeometryNode(PatchMesh->GetSingleContent().Get(), ParentLayer);
+			return ProcessGeometryNode(Singleton, ParentLayer);
 		}
 
 		TSharedPtr<IDatasmithMeshElement> MeshElement = FindOrAddMeshElement(PatchMesh);
@@ -734,7 +813,7 @@ namespace UE_DATASMITHWIRETRANSLATOR_NAMESPACE
 			ActorElement->SetLayer(*CsvLayerString);
 		}
 
-		OpenModelUtils::SetActorTransform(*ActorElement, GroupNode.Get());
+		GroupNode.SetActorTransform(*ActorElement);
 
 		if (WireSettings.bUseLayerAsActor && PatchMesh->GetLayer() != ParentLayer)
 		{
@@ -744,13 +823,13 @@ namespace UE_DATASMITHWIRETRANSLATOR_NAMESPACE
 				return TSharedPtr<IDatasmithActorElement>();
 			}
 
-			ensure(false);
+			ensureWire(false);
 		}
 
 		return ActorElement;
 	}
 
-	TSharedPtr<IDatasmithMeshElement> FWireTranslatorImpl::FindOrAddMeshElement(const TAlDagNodePtr<AlDagNode>& GeomNode)
+	TSharedPtr<IDatasmithMeshElement> FWireTranslatorImpl::FindOrAddMeshElement(const FAlDagNodePtr& GeomNode)
 	{
 		// Look if geometry has not been already processed, return it if found
 		if (TSharedPtr<IDatasmithMeshElement>* MeshElementPtr = GeomNodeToMeshElement.Find(GeomNode.GetHash()))
@@ -758,7 +837,7 @@ namespace UE_DATASMITHWIRETRANSLATOR_NAMESPACE
 			return *MeshElementPtr;
 		}
 
-		if (!OpenModelUtils::IsGeometryValid(GeomNode))
+		if (!GeomNode.HasGeometry())
 		{
 			// #wire_import: Log an error
 			return TSharedPtr<IDatasmithMeshElement>();
@@ -770,9 +849,6 @@ namespace UE_DATASMITHWIRETRANSLATOR_NAMESPACE
 		MeshElement->SetLightmapSourceUV(-1);
 #if TRACK_MESHELEMENT
 		{
-			//static const FString ToTrack(TEXT("shell_30181"));
-			//static const FString ToTrack(TEXT("Blend_srf_1395"));
-			static const FString ToTrack(TEXT("M0047942"));
 			if (!ToTrack.Equals(MeshElement->GetLabel()))
 			{
 				return TSharedPtr<IDatasmithMeshElement>();
@@ -789,37 +865,35 @@ namespace UE_DATASMITHWIRETRANSLATOR_NAMESPACE
 			};
 
 
-		TAlDagNodePtr<AlShellNode> ShellNode(GeomNode->asShellNodePtr());
-		if (ShellNode)
+		TAlObjectPtr<AlShell> Shell;
+		if (GeomNode.GetShell(Shell))
 		{
-			TAlObjectPtr<AlShell> Shell(ShellNode->shell());
-
 			TAlObjectPtr<AlShader> Shader(Shell->firstShader());
 			int32 SlotIndex = 0;
 			while (Shader)
 			{
 				ApplyMaterial(Shader, SlotIndex++);
-				Shader = TAlObjectPtr<AlShader>(Shell->nextShader(Shader.Get()));
+				Shader = Shell->nextShader(Shader.Get());
 			}
 			// #wire_import: Check There are as many shaders as trim regions
-			MeshElementToParametricNode.Add(MeshElement, GeomNode.Get());
+			MeshElementToParametricNode.Add(MeshElement, GeomNode);
 		}
 		else
 		{
-			TAlDagNodePtr<AlSurfaceNode> SurfaceNode(GeomNode->asSurfaceNodePtr());
-			if (SurfaceNode && AlIsValid(SurfaceNode->surface()))
+			TAlObjectPtr<AlSurface> Surface;
+			if (GeomNode.GetSurface(Surface))
 			{
 				// #wire_import: Check for trim regions
-				ApplyMaterial(SurfaceNode->surface()->firstShader(), 0);
-				MeshElementToParametricNode.Add(MeshElement, GeomNode.Get());
+				ApplyMaterial(Surface->firstShader(), 0);
+				MeshElementToParametricNode.Add(MeshElement, GeomNode);
 			}
 			else
 			{
-				TAlDagNodePtr<AlMeshNode> MeshNode(GeomNode->asMeshNodePtr());
-				if (MeshNode && AlIsValid(MeshNode->mesh()))
+				TAlObjectPtr<AlMesh> Mesh;
+				if (GeomNode.GetMesh(Mesh))
 				{
-					ApplyMaterial(MeshNode->mesh()->firstShader(), 0);
-					MeshElementToMeshNode.Add(MeshElement, MeshNode.Get());
+					ApplyMaterial(Mesh->firstShader(), 0);
+					MeshElementToMeshNode.Add(MeshElement, GeomNode);
 				}
 				else
 				{
@@ -855,9 +929,6 @@ namespace UE_DATASMITHWIRETRANSLATOR_NAMESPACE
 		MeshElement->SetLightmapSourceUV(-1);
 #if TRACK_MESHELEMENT
 		{
-			//static const FString ToTrack(TEXT("shell_30181"));
-			//static const FString ToTrack(TEXT("Blend_srf_1395"));
-			static const FString ToTrack(TEXT("M0047942"));
 			if (!ToTrack.Equals(MeshElement->GetLabel()))
 			{
 				return TSharedPtr<IDatasmithMeshElement>();
@@ -914,11 +985,12 @@ namespace UE_DATASMITHWIRETRANSLATOR_NAMESPACE
 
 
 		int32 SlotIndex = 0;
-		PatchMesh->IterateOnMeshNodes([&ApplyMaterial, &SlotIndex](const TAlDagNodePtr<AlMeshNode>& MeshNode)
+		PatchMesh->IterateOnMeshNodes([&ApplyMaterial, &SlotIndex](const FAlDagNodePtr& MeshNode)
 			{
-				if (MeshNode->mesh())
+				TAlObjectPtr<AlMesh> Mesh;
+				if (MeshNode.GetMesh(Mesh))
 				{
-					ApplyMaterial(MeshNode->mesh()->firstShader(), SlotIndex++);
+					ApplyMaterial(Mesh->firstShader(), SlotIndex++);
 				}
 			});
 
@@ -943,7 +1015,7 @@ namespace UE_DATASMITHWIRETRANSLATOR_NAMESPACE
 		}
 
 		TSharedPtr<IDatasmithActorElement> ParentLayerActor;
-		TAlObjectPtr<AlLayer> ParentLayer(Layer->parentLayer());
+		TAlObjectPtr<AlLayer> ParentLayer = FLayerContainer::FindOrAdd(Layer->parentLayer());
 		if (ParentLayer)
 		{
 			ParentLayerActor = FindOrAddLayerActor(ParentLayer);
@@ -980,11 +1052,14 @@ namespace UE_DATASMITHWIRETRANSLATOR_NAMESPACE
 	{
 		CADLibrary::FMeshParameters MeshParameters;
 
+#if WIRE_MEMORY_CHECK
+		int32 PrevAllocatedObjects = AllocatedObjects;
+#endif
 		if (TOptional<FMeshDescription> Mesh = GetMeshDescription(MeshElement, MeshParameters))
 		{
 #if TRACK_MESHELEMENT
 			FMeshDescription& MeshDescription = OutMeshPayload.LodMeshes.Add_GetRef(MoveTemp(Mesh.GetValue()));
-			//MakeMeshVisible(MeshDescription);
+			MakeMeshVisible(MeshDescription);
 #else
 			OutMeshPayload.LodMeshes.Add(MoveTemp(Mesh.GetValue()));
 #endif
@@ -1000,6 +1075,9 @@ namespace UE_DATASMITHWIRETRANSLATOR_NAMESPACE
 				}
 			}
 
+#if WIRE_MEMORY_CHECK
+			ensure(PrevAllocatedObjects == AllocatedObjects);
+#endif
 			return true;
 		}
 
@@ -1010,19 +1088,20 @@ namespace UE_DATASMITHWIRETRANSLATOR_NAMESPACE
 	{
 		if (WireSettings.bAliasUseNative)
 		{
-			if (AlDagNode** GeomNodePtr = MeshElementToParametricNode.Find(MeshElement))
+			if (FAlDagNodePtr* GeomNodePtr = MeshElementToParametricNode.Find(MeshElement))
 			{
-				TAlDagNodePtr<AlDagNode> GeomNode(*GeomNodePtr);
+				FAlDagNodePtr& GeomNode = *GeomNodePtr;
 
 				// #wire_import: Check whether parametric geometry with symmetry keeps the symmetry
 				// #wire_import: the best way, should be to don't have to apply inverse global transform to the generated mesh
-				TAlDagNodePtr<AlMeshNode> MeshNode = OpenModelUtils::TesselateDagLeaf(*GeomNode, ETesselatorType::Fast, WireSettings.ChordTolerance);
+				FAlDagNodePtr MeshNode = OpenModelUtils::TesselateDagLeaf(*GeomNode, ETesselatorType::Fast, WireSettings.ChordTolerance);
 
-				if (MeshNode && MeshNode->mesh())
+				TAlObjectPtr<AlMesh> Mesh;
+				if (MeshNode.GetMesh(Mesh))
 				{
 					AlMatrix4x4 AlMatrix;
 					GeomNode->inverseGlobalTransformationMatrix(AlMatrix);
-					MeshNode->mesh()->transform(AlMatrix);
+					Mesh->transform(AlMatrix);
 
 					// Get the meshes from the dag nodes. Note that removing the mesh's DAG.
 					// will also removes the meshes, so we have to do it later.
@@ -1030,9 +1109,9 @@ namespace UE_DATASMITHWIRETRANSLATOR_NAMESPACE
 				}
 			}
 			
-			if (AlMeshNode** MeshNodePtr = MeshElementToMeshNode.Find(MeshElement))
+			if (FAlDagNodePtr* MeshNodePtr = MeshElementToMeshNode.Find(MeshElement))
 			{
-				TAlDagNodePtr<AlMeshNode> MeshNode(*MeshNodePtr);
+				FAlDagNodePtr MeshNode = *MeshNodePtr;
 				if (MeshNode)
 				{
 					return GetMeshDescriptionFromMeshNode(MeshNode, MeshElement, OutMeshParameters);
@@ -1047,7 +1126,7 @@ namespace UE_DATASMITHWIRETRANSLATOR_NAMESPACE
 			return GetMeshDescriptionFromBodyNode(*BodyNodePtr, MeshElement, OutMeshParameters);
 		}
 
-		if (AlDagNode** GeomNodePtr = MeshElementToParametricNode.Find(MeshElement))
+		if (FAlDagNodePtr* GeomNodePtr = MeshElementToParametricNode.Find(MeshElement))
 		{
 			return GetMeshDescriptionFromParametricNode(*GeomNodePtr, MeshElement, OutMeshParameters);
 		}
@@ -1057,7 +1136,7 @@ namespace UE_DATASMITHWIRETRANSLATOR_NAMESPACE
 			return GetMeshDescriptionFromPatchMesh(*PatchMeshPtr, MeshElement, OutMeshParameters);
 		}
 
-		if (AlMeshNode** MeshNodePtr = MeshElementToMeshNode.Find(MeshElement))
+		if (FAlDagNodePtr* MeshNodePtr = MeshElementToMeshNode.Find(MeshElement))
 		{
 			return GetMeshDescriptionFromMeshNode(*MeshNodePtr, MeshElement, OutMeshParameters);
 		}
@@ -1116,10 +1195,10 @@ namespace UE_DATASMITHWIRETRANSLATOR_NAMESPACE
 		MeshDescription.Empty();
 
 		constexpr bool bMerge = true;
-		PatchMesh->IterateOnMeshNodes([&](const TAlDagNodePtr<AlMeshNode>& MeshNode)
+		PatchMesh->IterateOnMeshNodes([&](const FAlDagNodePtr& MeshNode)
 			{
-				TAlObjectPtr<AlMesh> Mesh(MeshNode.GetMesh());
-				if (Mesh)
+				TAlObjectPtr<AlMesh> Mesh;
+				if (MeshNode.GetMesh(Mesh))
 				{
 					AlMatrix4x4 AlMatrix;
 					if (OutMeshParameters.bIsSymmetric)
@@ -1147,7 +1226,7 @@ namespace UE_DATASMITHWIRETRANSLATOR_NAMESPACE
 	}
 
 	// #wire_import: AlSurfaceNode can have trim regions. This should be handle at this stage
-	TOptional<FMeshDescription> FWireTranslatorImpl::GetMeshDescriptionFromParametricNode(const TAlDagNodePtr<AlDagNode>& DagNode, TSharedPtr<IDatasmithMeshElement> MeshElement, CADLibrary::FMeshParameters& OutMeshParameters)
+	TOptional<FMeshDescription> FWireTranslatorImpl::GetMeshDescriptionFromParametricNode(const FAlDagNodePtr& DagNode, TSharedPtr<IDatasmithMeshElement> MeshElement, CADLibrary::FMeshParameters& OutMeshParameters)
 	{
 		OutMeshParameters = OpenModelUtils::GetMeshParameters(DagNode.GetLayer());
 
@@ -1165,7 +1244,7 @@ namespace UE_DATASMITHWIRETRANSLATOR_NAMESPACE
 			ObjectReference = EAliasObjectReference::WorldReference;
 		}
 
-		ensure(MeshElement->GetMaterialSlotCount() == 1);
+		ensureWire(MeshElement->GetMaterialSlotCount() == 1);
 
 		const FDagNodeGeometry DagNodeGeometry( (int32)ECADModelGeometryType::DagNode, ObjectReference, DagNode );
 		if (!ModelConverter->AddGeometry(DagNodeGeometry))
@@ -1180,7 +1259,7 @@ namespace UE_DATASMITHWIRETRANSLATOR_NAMESPACE
 		FMeshDescription MeshDescription;
 		DatasmithMeshHelper::PrepareAttributeForStaticMesh(MeshDescription);
 
-		OutMeshParameters = OpenModelUtils::GetMeshParameters(*DagNode);
+		OutMeshParameters = DagNode.GetMeshParameters();
 
 		if (ModelConverter->Tessellate(OutMeshParameters, MeshDescription))
 		{
@@ -1194,20 +1273,20 @@ namespace UE_DATASMITHWIRETRANSLATOR_NAMESPACE
 		return TOptional<FMeshDescription>();
 	}
 
-	TOptional<FMeshDescription> FWireTranslatorImpl::GetMeshDescriptionFromMeshNode(const TAlDagNodePtr<AlMeshNode>& MeshNode, TSharedPtr<IDatasmithMeshElement> MeshElement, CADLibrary::FMeshParameters& OutMeshParameters)
+	TOptional<FMeshDescription> FWireTranslatorImpl::GetMeshDescriptionFromMeshNode(const FAlDagNodePtr& MeshNode, TSharedPtr<IDatasmithMeshElement> MeshElement, CADLibrary::FMeshParameters& OutMeshParameters)
 	{
 		if (!MeshNode)
 		{
 			return TOptional<FMeshDescription>();
 		}
 
-		TAlObjectPtr<AlMesh> Mesh(MeshNode->mesh());
-		if (!Mesh)
+		TAlObjectPtr<AlMesh> Mesh;
+		if (!MeshNode.GetMesh(Mesh))
 		{
 			return TOptional<FMeshDescription>();
 		}
 
-		OutMeshParameters = OpenModelUtils::GetMeshParameters(*MeshNode);
+		OutMeshParameters = MeshNode.GetMeshParameters();
 
 		if (OutMeshParameters.bIsSymmetric)
 		{
@@ -2187,6 +2266,13 @@ namespace UE_DATASMITHWIRETRANSLATOR_NAMESPACE
 
 		virtual void ShutdownModule() override
 		{
+#if WIRE_MEMORY_CHECK
+			// #wire_import: Need to investigate why this is crashing when enabled.
+			if (bool(AlUniverse::isInitialized()))
+			{
+				AlUniverse::deleteAll();
+			}
+#endif
 		}
 	};
 
@@ -2194,6 +2280,7 @@ namespace UE_DATASMITHWIRETRANSLATOR_NAMESPACE
 #pragma optimize("", off)
 	void MakeMeshVisible(FMeshDescription& MeshDescription)
 	{
+#if MAKE_VISIBLE
 		using namespace UE::Geometry;
 
 		TVertexAttributesRef<FVector3f> VertexPositions = MeshDescription.GetVertexPositions();
@@ -2207,9 +2294,9 @@ namespace UE_DATASMITHWIRETRANSLATOR_NAMESPACE
 		constexpr float MinSize = 1.f;
 		constexpr float MaxSize = 20.f;
 
-		const float ScaleX = OBox.Extents.X < MinSize ? 2.f / OBox.Extents.X : OBox.Extents.X > MaxSize ? MaxSize / OBox.Extents.X : 1.f;
-		const float ScaleY = OBox.Extents.Y < MinSize ? 2.f / OBox.Extents.Y : OBox.Extents.Y > MaxSize ? MaxSize / OBox.Extents.Y : 1.f;
-		const float ScaleZ = OBox.Extents.Z < MinSize ? 2.f / OBox.Extents.Z : OBox.Extents.Z > MaxSize ? MaxSize / OBox.Extents.Z : 1.f;
+		const float ScaleX = OBox.Extents.X < MinSize ? 20.f / OBox.Extents.X : OBox.Extents.X > MaxSize ? MaxSize / OBox.Extents.X : 10.f;
+		const float ScaleY = OBox.Extents.Y < MinSize ? 20.f / OBox.Extents.Y : OBox.Extents.Y > MaxSize ? MaxSize / OBox.Extents.Y : 10.f;
+		const float ScaleZ = OBox.Extents.Z < MinSize ? 20.f / OBox.Extents.Z : OBox.Extents.Z > MaxSize ? MaxSize / OBox.Extents.Z : 10.f;
 		UE_LOG(LogWireInterface, Warning, TEXT("Scaling factor: %.3f %.3f %.3f"), ScaleX, ScaleY, ScaleZ);
 		const FVector3f AxisX = OBox.AxisX();
 		const FVector3f AxisY = OBox.AxisY();
@@ -2220,6 +2307,7 @@ namespace UE_DATASMITHWIRETRANSLATOR_NAMESPACE
 			FVector3f P = VertexPositions[VertexID] - OBox.Frame.Origin;
 			VertexPositions[VertexID] = ((P | AxisX) * ScaleX) * AxisX + ((P | AxisY) * ScaleY) * AxisY + ((P | AxisZ) * ScaleZ) * AxisZ;
 		}
+#endif
 	}
 #endif
 
diff --git a/Engine/Plugins/Enterprise/DatasmithCADImporter/Source/WireInterface/Private/WireInterfaceImpl.h b/Engine/Plugins/Enterprise/DatasmithCADImporter/Source/WireInterface/Private/WireInterfaceImpl.h
--- a/Engine/Plugins/Enterprise/DatasmithCADImporter/Source/WireInterface/Private/WireInterfaceImpl.h
+++ b/Engine/Plugins/Enterprise/DatasmithCADImporter/Source/WireInterface/Private/WireInterfaceImpl.h
@@ -54,16 +54,16 @@ namespace UE_DATASMITHWIRETRANSLATOR_NAMESPACE
 	private:
 		/** Model traversal */
 		bool TraverseModel();
-		TSharedPtr<IDatasmithActorElement> TraverseDag(const TAlDagNodePtr<AlDagNode>& DagNode);
+		TSharedPtr<IDatasmithActorElement> TraverseDag(const FAlDagNodePtr& DagNode);
 
-		TSharedPtr<IDatasmithActorElement> ProcessGeometryNode(const TAlDagNodePtr<AlDagNode>& GeomNode, const TAlObjectPtr<AlLayer>& ParentLayer = TAlObjectPtr<AlLayer>());
-		TSharedPtr<IDatasmithActorElement> TraverseGroupNode(const TAlDagNodePtr<AlGroupNode>& GroupNode, const TAlObjectPtr<AlLayer>& ParentLayer = TAlObjectPtr<AlLayer>());
-		TSharedPtr<IDatasmithActorElement> ProcessGroupNode(const TAlDagNodePtr<AlGroupNode>& GroupNode, const TAlObjectPtr<AlLayer>& ParentLayer = TAlObjectPtr<AlLayer>());
-		TSharedPtr<IDatasmithActorElement> ProcessBodyNode(TSharedPtr<FBodyNode>& BodyNode, const TAlDagNodePtr<AlGroupNode>& GroupNode, const TAlObjectPtr<AlLayer>& ParentLayer = TAlObjectPtr<AlLayer>());
-		TSharedPtr<IDatasmithActorElement> ProcessPatchMesh(TSharedPtr<FPatchMesh>& PatchMesh, const TAlDagNodePtr<AlGroupNode>& GroupNode, const TAlObjectPtr<AlLayer>& ParentLayer = TAlObjectPtr<AlLayer>());
+		TSharedPtr<IDatasmithActorElement> ProcessGeometryNode(const FAlDagNodePtr& GeomNode, const TAlObjectPtr<AlLayer>& ParentLayer = TAlObjectPtr<AlLayer>());
+		TSharedPtr<IDatasmithActorElement> TraverseGroupNode(const FAlDagNodePtr& GroupNode, const TAlObjectPtr<AlLayer>& ParentLayer = TAlObjectPtr<AlLayer>());
+		TSharedPtr<IDatasmithActorElement> ProcessGroupNode(const FAlDagNodePtr& GroupNode, const TAlObjectPtr<AlLayer>& ParentLayer = TAlObjectPtr<AlLayer>());
+		TSharedPtr<IDatasmithActorElement> ProcessBodyNode(TSharedPtr<FBodyNode>& BodyNode, const FAlDagNodePtr& GroupNode, const TAlObjectPtr<AlLayer>& ParentLayer = TAlObjectPtr<AlLayer>());
+		TSharedPtr<IDatasmithActorElement> ProcessPatchMesh(TSharedPtr<FPatchMesh>& PatchMesh, const FAlDagNodePtr& GroupNode, const TAlObjectPtr<AlLayer>& ParentLayer = TAlObjectPtr<AlLayer>());
 
 		TSharedPtr<IDatasmithActorElement> FindOrAddLayerActor(const TAlObjectPtr<AlLayer>& Layer);
-		TSharedPtr<IDatasmithMeshElement> FindOrAddMeshElement(const TAlDagNodePtr<AlDagNode>& GeomNode);
+		TSharedPtr<IDatasmithMeshElement> FindOrAddMeshElement(const FAlDagNodePtr& GeomNode);
 		TSharedPtr<IDatasmithMeshElement> FindOrAddMeshElement(TSharedPtr<FBodyNode>& BodyNode);
 		TSharedPtr<IDatasmithMeshElement> FindOrAddMeshElement(TSharedPtr<FPatchMesh>& PatchMesh);
 
@@ -88,13 +88,34 @@ namespace UE_DATASMITHWIRETRANSLATOR_NAMESPACE
 		TSharedPtr<CADLibrary::ICADModelConverter> GetModelConverter() const;
 		TOptional<FMeshDescription> GetMeshDescriptionFromBodyNode(TSharedPtr<FBodyNode>& BodyNode, TSharedPtr<IDatasmithMeshElement> MeshElement, CADLibrary::FMeshParameters& MeshParameters);
 		TOptional<FMeshDescription> GetMeshDescriptionFromPatchMesh(TSharedPtr<FPatchMesh>& PatchMesh, TSharedPtr<IDatasmithMeshElement> MeshElement, CADLibrary::FMeshParameters& MeshParameters);
-		TOptional<FMeshDescription> GetMeshDescriptionFromParametricNode(const TAlDagNodePtr<AlDagNode>& DagNode, TSharedPtr<IDatasmithMeshElement> MeshElement, CADLibrary::FMeshParameters& MeshParameters);
-		TOptional<FMeshDescription> GetMeshDescriptionFromMeshNode(const TAlDagNodePtr<AlMeshNode>& MeshNode, TSharedPtr<IDatasmithMeshElement> MeshElement, CADLibrary::FMeshParameters& MeshParameters);
+		TOptional<FMeshDescription> GetMeshDescriptionFromParametricNode(const FAlDagNodePtr& DagNode, TSharedPtr<IDatasmithMeshElement> MeshElement, CADLibrary::FMeshParameters& MeshParameters);
+		TOptional<FMeshDescription> GetMeshDescriptionFromMeshNode(const FAlDagNodePtr& MeshNode, TSharedPtr<IDatasmithMeshElement> MeshElement, CADLibrary::FMeshParameters& MeshParameters);
 
+		FAlDagNodePtr FindOrAddDagNode(AlDagNode* InDagNode)
+		{
+			if (!InDagNode)
+			{
+				return FAlDagNodePtr();
+			}
+
+			if (FAlDagNodePtr* DagNodePtr = EncounteredNodes.Find(InDagNode))
+			{
+#if WIRE_MEMORY_CHECK
+				ensure(false);
+#endif
+				return *DagNodePtr;
+			}
+
+			FAlDagNodePtr& NewDagNode = EncounteredNodes.Add(InDagNode);
+			NewDagNode = InDagNode;
+
+			return NewDagNode;
+		}
 	private:
 		TSharedPtr<IDatasmithScene> DatasmithScene;
 		FString OutputPath;
 		FString SceneFullPath;
+		FString SceneVersion;
 
 		FWireSettings WireSettings;
 
@@ -105,8 +126,8 @@ namespace UE_DATASMITHWIRETRANSLATOR_NAMESPACE
 		TMap<FString, TSharedPtr<IDatasmithBaseMaterialElement>> ShaderNameToMaterial;
 
 		TMap<uint32, TSharedPtr<IDatasmithMeshElement>> GeomNodeToMeshElement;
-		TMap<TSharedPtr<IDatasmithMeshElement>, AlDagNode*> MeshElementToParametricNode;
-		TMap<TSharedPtr<IDatasmithMeshElement>, AlMeshNode*> MeshElementToMeshNode;
+		TMap<TSharedPtr<IDatasmithMeshElement>, FAlDagNodePtr> MeshElementToParametricNode;
+		TMap<TSharedPtr<IDatasmithMeshElement>, FAlDagNodePtr> MeshElementToMeshNode;
 
 		TMap<uint32, TSharedPtr<IDatasmithMeshElement>> BodyNodeToMeshElement;
 		TMap<TSharedPtr<IDatasmithMeshElement>, TSharedPtr<FBodyNode>> MeshElementToBodyNode;
@@ -114,6 +135,8 @@ namespace UE_DATASMITHWIRETRANSLATOR_NAMESPACE
 		TMap<uint32, TSharedPtr<IDatasmithMeshElement>> PatchMeshToMeshElement;
 		TMap<TSharedPtr<IDatasmithMeshElement>, TSharedPtr<FPatchMesh>> MeshElementToPatchMesh;
 
+		TMap<AlDagNode*, FAlDagNodePtr> EncounteredNodes;
+
 		TMap<uint32, TSharedPtr<IDatasmithActorElement>> LayerToActor;
 	};
 } // namespace
diff --git a/Engine/Plugins/Enterprise/DatasmithCADImporter/Source/WireInterface/WireInterface2020.Build.cs b/Engine/Plugins/Enterprise/DatasmithCADImporter/Source/WireInterface/WireInterface2020.Build.cs
--- a/Engine/Plugins/Enterprise/DatasmithCADImporter/Source/WireInterface/WireInterface2020.Build.cs
+++ b/Engine/Plugins/Enterprise/DatasmithCADImporter/Source/WireInterface/WireInterface2020.Build.cs
@@ -56,7 +56,7 @@ public abstract class WireInterfaceBase : ModuleRules
 		PublicDefinitions.Add($"UE_DATASMITHWIRETRANSLATOR_MODULE_NAME={GetType().Name}");
 		PublicDefinitions.Add($"UE_OPENMODEL_MAJOR_VERSION={GetMajorVersion()}");
 		PublicDefinitions.Add($"UE_OPENMODEL_MINOR_VERSION={GetMinorVersion()}");
-		PublicDefinitions.Add("THINFACE_ENABLED=0");
+		PublicDefinitions.Add("WIRE_THINFACE_ENABLED=0");
 
 		if (System.Type.GetType(GetAliasVersion()) != null)
 		{
diff --git a/Engine/Plugins/Enterprise/DatasmithCADImporter/Source/WireInterface/WireInterface2024_1.Build.cs b/Engine/Plugins/Enterprise/DatasmithCADImporter/Source/WireInterface/WireInterface2024_1.Build.cs
new file mode 100644
--- /dev/null
+++ b/Engine/Plugins/Enterprise/DatasmithCADImporter/Source/WireInterface/WireInterface2024_1.Build.cs
@@ -0,0 +1,32 @@
+// Copyright Epic Games, Inc. All Rights Reserved.
+
+using System.IO;
+using UnrealBuildTool;
+
+public class WireInterface2024_1 : WireInterfaceBase
+{
+	public WireInterface2024_1(ReadOnlyTargetRules Target) 
+		: base(Target)
+	{
+	}
+
+	public override string GetAliasVersion()
+	{
+		return "OpenModel2024_1";
+	}
+
+	public override string GetAliasDefinition()
+	{
+		return "OPEN_MODEL_2024_1";
+	}
+
+	public override int GetMajorVersion()
+	{
+		return 2024;
+	}
+
+	public override int GetMinorVersion()
+	{
+		return 1;
+	}
+}
\ No newline at end of file
diff --git a/Engine/Plugins/Enterprise/DatasmithCADImporter/Source/WireInterface/WireInterface2025_0.Build.cs b/Engine/Plugins/Enterprise/DatasmithCADImporter/Source/WireInterface/WireInterface2025_0.Build.cs
new file mode 100644
--- /dev/null
+++ b/Engine/Plugins/Enterprise/DatasmithCADImporter/Source/WireInterface/WireInterface2025_0.Build.cs
@@ -0,0 +1,32 @@
+// Copyright Epic Games, Inc. All Rights Reserved.
+
+using System.IO;
+using UnrealBuildTool;
+
+public class WireInterface2025_0 : WireInterfaceBase
+{
+	public WireInterface2025_0(ReadOnlyTargetRules Target) 
+		: base(Target)
+	{
+	}
+
+	public override string GetAliasVersion()
+	{
+		return "OpenModel2025_0";
+	}
+
+	public override string GetAliasDefinition()
+	{
+		return "OPEN_MODEL_2025_0";
+	}
+
+	public override int GetMajorVersion()
+	{
+		return 2025;
+	}
+
+	public override int GetMinorVersion()
+	{
+		return 0;
+	}
+}
\ No newline at end of file
