Fixed linkage issue on cooking

#rnx

diff --git a/Engine/Plugins/Enterprise/DatasmithCADImporter/Source/WireInterface/Private/OpenModelUtils.h b/Engine/Plugins/Enterprise/DatasmithCADImporter/Source/WireInterface/Private/OpenModelUtils.h
--- a/Engine/Plugins/Enterprise/DatasmithCADImporter/Source/WireInterface/Private/OpenModelUtils.h
+++ b/Engine/Plugins/Enterprise/DatasmithCADImporter/Source/WireInterface/Private/OpenModelUtils.h
@@ -18,8 +18,6 @@
 #include "MessageLogModule.h"
 #endif
 
-DECLARE_LOG_CATEGORY_EXTERN(LogWireInterface, Log, All);
-
 #if PLATFORM_WINDOWS
 #include "Windows/AllowWindowsPlatformTypes.h"
 #endif
@@ -68,6 +66,8 @@ struct FMeshDescription;
 
 namespace UE_DATASMITHWIRETRANSLATOR_NAMESPACE
 {
+	DECLARE_LOG_CATEGORY_EXTERN(LogWireInterface, Log, All)
+
 #if WIRE_MEMORY_CHECK
 	extern TSet<AlDagNode*> DagNodeSet;
 	extern TSet<AlObject*> ObjectSet;
