Fixed infinite loop when importing wire files using techSoft

#rnx

diff --git a/Engine/Plugins/Enterprise/DatasmithCADImporter/Source/WireInterface/Private/AliasModelToTechSoftConverter.cpp b/Engine/Plugins/Enterprise/DatasmithCADImporter/Source/WireInterface/Private/AliasModelToTechSoftConverter.cpp
--- a/Engine/Plugins/Enterprise/DatasmithCADImporter/Source/WireInterface/Private/AliasModelToTechSoftConverter.cpp
+++ b/Engine/Plugins/Enterprise/DatasmithCADImporter/Source/WireInterface/Private/AliasModelToTechSoftConverter.cpp
@@ -340,7 +340,7 @@ A3DTopoFace* FAliasModelToTechSoftConverter::AddTrimRegion(const AlTrimRegion& I
 			LinkEdgesLoop(*TrimBoundary);
 		}
 
-		TrimBoundary->nextBoundaryD();
+		Status = TrimBoundary->nextBoundaryD();
 	}
 
 	if (Loops.Num() == 0)
