Made sure ensure is disabled when compiling the CADWorker

#rnx

diff --git a/Engine/Plugins/Enterprise/DatasmithCADImporter/Source/CADInterfaces/CADInterfaces.Build.cs b/Engine/Plugins/Enterprise/DatasmithCADImporter/Source/CADInterfaces/CADInterfaces.Build.cs
--- a/Engine/Plugins/Enterprise/DatasmithCADImporter/Source/CADInterfaces/CADInterfaces.Build.cs
+++ b/Engine/Plugins/Enterprise/DatasmithCADImporter/Source/CADInterfaces/CADInterfaces.Build.cs
@@ -11,6 +11,11 @@ namespace UnrealBuildTool.Rules
 			bLegalToDistributeObjectCode = true;
 			bUseUnity = false;
 
+			if(Target.bIsBuildingConsoleApplication)
+			{
+				PublicDefinitions.Add("DO_ENSURE=0");
+			}
+
 			PublicDependencyModuleNames.AddRange(
 				new string[]
 				{
diff --git a/Engine/Plugins/Enterprise/DatasmithCADImporter/Source/CADLibrary/CADLibrary.Build.cs b/Engine/Plugins/Enterprise/DatasmithCADImporter/Source/CADLibrary/CADLibrary.Build.cs
--- a/Engine/Plugins/Enterprise/DatasmithCADImporter/Source/CADLibrary/CADLibrary.Build.cs
+++ b/Engine/Plugins/Enterprise/DatasmithCADImporter/Source/CADLibrary/CADLibrary.Build.cs
@@ -10,6 +10,11 @@ public class CADLibrary : ModuleRules
 	{
 		bLegalToDistributeObjectCode = true;
 
+		if (Target.bIsBuildingConsoleApplication)
+		{
+			PublicDefinitions.Add("DO_ENSURE=0");
+		}
+
 		PublicDependencyModuleNames.AddRange(
 			new string[]
 			{
diff --git a/Engine/Plugins/Enterprise/DatasmithCADImporter/Source/DatasmithCADTranslator/DatasmithCADTranslator.Build.cs b/Engine/Plugins/Enterprise/DatasmithCADImporter/Source/DatasmithCADTranslator/DatasmithCADTranslator.Build.cs
--- a/Engine/Plugins/Enterprise/DatasmithCADImporter/Source/DatasmithCADTranslator/DatasmithCADTranslator.Build.cs
+++ b/Engine/Plugins/Enterprise/DatasmithCADImporter/Source/DatasmithCADTranslator/DatasmithCADTranslator.Build.cs
@@ -16,6 +16,11 @@ namespace UnrealBuildTool.Rules
 				}
 			);
 
+			if (Target.bIsBuildingConsoleApplication)
+			{
+				PublicDefinitions.Add("DO_ENSURE=0");
+			}
+
 			PrivateDependencyModuleNames.AddRange(
 				new string[]
 				{
diff --git a/Engine/Plugins/Enterprise/DatasmithCADImporter/Source/DatasmithDispatcher/DatasmithDispatcher.Build.cs b/Engine/Plugins/Enterprise/DatasmithCADImporter/Source/DatasmithDispatcher/DatasmithDispatcher.Build.cs
--- a/Engine/Plugins/Enterprise/DatasmithCADImporter/Source/DatasmithDispatcher/DatasmithDispatcher.Build.cs
+++ b/Engine/Plugins/Enterprise/DatasmithCADImporter/Source/DatasmithDispatcher/DatasmithDispatcher.Build.cs
@@ -23,6 +23,11 @@ namespace UnrealBuildTool.Rules
 				{
                 }
 			);
-        }
-    }
+
+			if (Target.bIsBuildingConsoleApplication)
+			{
+				PublicDefinitions.Add("DO_ENSURE=0");
+			}
+		}
+	}
 }
diff --git a/Engine/Source/Runtime/Datasmith/CADKernel/Base/CADKernel.Build.cs b/Engine/Source/Runtime/Datasmith/CADKernel/Base/CADKernel.Build.cs
--- a/Engine/Source/Runtime/Datasmith/CADKernel/Base/CADKernel.Build.cs
+++ b/Engine/Source/Runtime/Datasmith/CADKernel/Base/CADKernel.Build.cs
@@ -13,10 +13,16 @@ namespace UnrealBuildTool.Rules
 			
 			PublicDefinitions.Add("CADKERNEL_THINZONE=0");
 
+			if (Target.bIsBuildingConsoleApplication)
+			{
+				PublicDefinitions.Add("DO_ENSURE=0");
+			}
+
 			PublicDependencyModuleNames.AddRange(
 				new string[]
 				{
 					"Core",
+					"GeometryCore",
 				}
 			);
 		}
