Fixed CadAutomationTests

#rnx

diff --git a/Engine/Plugins/Enterprise/DatasmithCADImporter/Source/CADInterfaces/Private/TechSoftFileParser.cpp b/Engine/Plugins/Enterprise/DatasmithCADImporter/Source/CADInterfaces/Private/TechSoftFileParser.cpp
--- a/Engine/Plugins/Enterprise/DatasmithCADImporter/Source/CADInterfaces/Private/TechSoftFileParser.cpp
+++ b/Engine/Plugins/Enterprise/DatasmithCADImporter/Source/CADInterfaces/Private/TechSoftFileParser.cpp
@@ -214,6 +214,7 @@ void TraverseAttribute(const A3DMiscAttributeData& AttributeData, TMap<FString,
 
 void SetIOOption(A3DImport& Importer)
 {
+	using namespace CADLibrary;
 	// A3DRWParamsGeneralData Importer.m_sGeneral
 	Importer.m_sLoadData.m_sGeneral.m_bReadSolids = A3D_TRUE;
 	Importer.m_sLoadData.m_sGeneral.m_bReadSurfaces = A3D_TRUE;
@@ -231,7 +232,7 @@ void SetIOOption(A3DImport& Importer)
 	Importer.m_sLoadData.m_sGeneral.m_bReadConstraints = A3D_FALSE;
 	Importer.m_sLoadData.m_sGeneral.m_iNbMultiProcess = 1;
 
-	Importer.m_sLoadData.m_sIncremental.m_bLoadNoDependencies = CADLibrary::FImportParameters::bGEnableCADCache;
+	Importer.m_sLoadData.m_sIncremental.m_bLoadNoDependencies = FImportParameters::bGEnableCADCache && GMaxImportThreads != 1;
 	Importer.m_sLoadData.m_sIncremental.m_bLoadStructureOnly = false;
 }
 
