Updated message output of CAD translator

#rnx

diff --git a/Engine/Plugins/Enterprise/DatasmithCADImporter/Source/CADInterfaces/Private/CADInterfacesModule.cpp b/Engine/Plugins/Enterprise/DatasmithCADImporter/Source/CADInterfaces/Private/CADInterfacesModule.cpp
--- a/Engine/Plugins/Enterprise/DatasmithCADImporter/Source/CADInterfaces/Private/CADInterfacesModule.cpp
+++ b/Engine/Plugins/Enterprise/DatasmithCADImporter/Source/CADInterfaces/Private/CADInterfacesModule.cpp
@@ -44,6 +44,11 @@ ECADInterfaceAvailability ICADInterfacesModule::GetAvailability()
 	return ECADInterfaceAvailability::Unavailable;
 }
 
+const TCHAR* ICADInterfacesModule::GetLibraryVersion()
+{
+	return CADLibrary::FTechSoftInterface::Get().GetVersion();
+}
+
 void FCADInterfacesModule::StartupModule()
 {
 
diff --git a/Engine/Plugins/Enterprise/DatasmithCADImporter/Source/CADInterfaces/Private/TechSoftInterface.cpp b/Engine/Plugins/Enterprise/DatasmithCADImporter/Source/CADInterfaces/Private/TechSoftInterface.cpp
--- a/Engine/Plugins/Enterprise/DatasmithCADImporter/Source/CADInterfaces/Private/TechSoftInterface.cpp
+++ b/Engine/Plugins/Enterprise/DatasmithCADImporter/Source/CADInterfaces/Private/TechSoftInterface.cpp
@@ -28,6 +28,17 @@ FTechSoftInterface& FTechSoftInterface::Get()
 	return TechSoftInterface;
 }
 
+const TCHAR* FTechSoftInterface::GetVersion()
+{
+#ifdef USE_TECHSOFT_SDK
+	static const FString Version = TEXT("TechSoft ") + FString::FromInt(A3D_DLL_MAJORVERSION) + TEXT(".") + FString::FromInt(A3D_DLL_MINORVERSION) + TEXT(".") + FString::FromInt(A3D_DLL_UPDATEVERSION);
+	return bIsInitialize ? *Version : TEXT("TechSoft uninitialized");
+#else
+	return TEXT("TechSoft unavailable");
+#endif
+}
+
+#pragma optimize("", off)
 bool FTechSoftInterface::InitializeKernel(const TCHAR* InEnginePluginsPath)
 {
 #ifdef USE_TECHSOFT_SDK
@@ -54,12 +65,14 @@ bool FTechSoftInterface::InitializeKernel(const TCHAR* InEnginePluginsPath)
 
 	if (A3DSDKLoadLibraryA(TCHAR_TO_UTF8(*TechSoftDllPath)))
 	{
-		A3DLicPutUnifiedLicense(HOOPS_LICENSE);
+		A3DStatus Status = A3DLicPutUnifiedLicense(HOOPS_LICENSE);
 
 		A3DInt32 iMajorVersion = 0, iMinorVersion = 0;
-		if (A3DDllGetVersion(&iMajorVersion, &iMinorVersion) == A3D_SUCCESS)
+		Status = A3DDllGetVersion(&iMajorVersion, &iMinorVersion);
+		if (Status == A3D_SUCCESS)
 		{
-			if (A3DDllInitialize(A3D_DLL_MAJORVERSION, A3D_DLL_MINORVERSION) == A3D_SUCCESS)
+			Status = A3DDllInitialize(A3D_DLL_MAJORVERSION, A3D_DLL_MINORVERSION);
+			if (Status == A3D_SUCCESS)
 			{
 				bIsInitialize = true;
 				return true;
@@ -73,6 +86,7 @@ bool FTechSoftInterface::InitializeKernel(const TCHAR* InEnginePluginsPath)
 #endif
 	return false;
 }
+#pragma optimize("", on)
 
 namespace TechSoftInterface
 {
@@ -84,13 +98,7 @@ bool TECHSOFT_InitializeKernel(const TCHAR* InEnginePluginsPath)
 
 FString GetTechSoftVersion()
 {
-#ifdef USE_TECHSOFT_SDK
-	A3DInt32 MajorVersion = 0, MinorVersion = 0;
-	A3DDllGetVersion(&MajorVersion, &MinorVersion);
-	return FString::Printf(TEXT("Techsoft %d.%d"), MajorVersion, MinorVersion);
-#else
-	return FString();
-#endif
+	return FTechSoftInterface::Get().GetVersion();
 }
 
 #ifdef USE_TECHSOFT_SDK
diff --git a/Engine/Plugins/Enterprise/DatasmithCADImporter/Source/CADInterfaces/Public/CADInterfacesModule.h b/Engine/Plugins/Enterprise/DatasmithCADImporter/Source/CADInterfaces/Public/CADInterfacesModule.h
--- a/Engine/Plugins/Enterprise/DatasmithCADImporter/Source/CADInterfaces/Public/CADInterfacesModule.h
+++ b/Engine/Plugins/Enterprise/DatasmithCADImporter/Source/CADInterfaces/Public/CADInterfacesModule.h
@@ -20,4 +20,5 @@ class CADINTERFACES_API ICADInterfacesModule : public IModuleInterface
 public:
 	static ICADInterfacesModule& Get();
 	static ECADInterfaceAvailability GetAvailability();
+	static const TCHAR* GetLibraryVersion();
 };
diff --git a/Engine/Plugins/Enterprise/DatasmithCADImporter/Source/CADInterfaces/Public/TechSoftInterface.h b/Engine/Plugins/Enterprise/DatasmithCADImporter/Source/CADInterfaces/Public/TechSoftInterface.h
--- a/Engine/Plugins/Enterprise/DatasmithCADImporter/Source/CADInterfaces/Public/TechSoftInterface.h
+++ b/Engine/Plugins/Enterprise/DatasmithCADImporter/Source/CADInterfaces/Public/TechSoftInterface.h
@@ -60,6 +60,7 @@ public:
 	}
 
 	bool InitializeKernel(const TCHAR* = TEXT(""));
+	const TCHAR* GetVersion();
 
 private:
 
diff --git a/Engine/Plugins/Enterprise/DatasmithCADImporter/Source/DatasmithCADTranslator/Private/DatasmithCADTranslator.cpp b/Engine/Plugins/Enterprise/DatasmithCADImporter/Source/DatasmithCADTranslator/Private/DatasmithCADTranslator.cpp
--- a/Engine/Plugins/Enterprise/DatasmithCADImporter/Source/DatasmithCADTranslator/Private/DatasmithCADTranslator.cpp
+++ b/Engine/Plugins/Enterprise/DatasmithCADImporter/Source/DatasmithCADTranslator/Private/DatasmithCADTranslator.cpp
@@ -158,7 +158,7 @@ bool FDatasmithCADTranslator::LoadScene(TSharedRef<IDatasmithScene> DatasmithSce
 	CADLibrary::FFileDescriptor FileDescriptor(*FPaths::ConvertRelativePathToFull(GetSource().GetSourceFile()));
 
 	UE_LOG(LogCADTranslator, Display, TEXT("CAD translation [%s]."), *FileDescriptor.GetSourcePath());
-	UE_LOG(LogCADTranslator, Display, TEXT(" - Parsing Library:      %s"), TEXT("TechSoft"));
+	UE_LOG(LogCADTranslator, Display, TEXT(" - Parsing Library:      %s"), ICADInterfacesModule::GetLibraryVersion());
 	UE_LOG(LogCADTranslator, Display, TEXT(" - Tessellation Library: %s")
 		, FImportParameters::bGDisableCADKernelTessellation ? TEXT("TechSoft") : TEXT("CADKernel"));
 	UE_LOG(LogCADTranslator, Display, TEXT(" - Cache mode:           %s")
