Fixed wrong initial value for GeometricTolerance

#rnx

diff --git a/Engine/Plugins/Enterprise/DatasmithCADImporter/Source/DatasmithOpenNurbsTranslator/Private/DatasmithOpenNurbsTranslator.cpp b/Engine/Plugins/Enterprise/DatasmithCADImporter/Source/DatasmithOpenNurbsTranslator/Private/DatasmithOpenNurbsTranslator.cpp
--- a/Engine/Plugins/Enterprise/DatasmithCADImporter/Source/DatasmithOpenNurbsTranslator/Private/DatasmithOpenNurbsTranslator.cpp
+++ b/Engine/Plugins/Enterprise/DatasmithCADImporter/Source/DatasmithOpenNurbsTranslator/Private/DatasmithOpenNurbsTranslator.cpp
@@ -3341,6 +3341,7 @@ void FDatasmithOpenNurbsTranslator::SetSceneImportOptions(const TArray<TObjectPt
 		else if (UDatasmithOpenNurbsImportOptions* OpenNurbsOptionsObj = Cast<UDatasmithOpenNurbsImportOptions>(Option))
 		{
 			OpenNurbsOptions = OpenNurbsOptionsObj->Options;
+			OpenNurbsOptionsObj->SaveConfig(CPF_Config);
 		}
 	}
 
diff --git a/Engine/Plugins/Enterprise/DatasmithCADImporter/Source/ParametricSurface/Private/ParametricSurfaceTranslator.cpp b/Engine/Plugins/Enterprise/DatasmithCADImporter/Source/ParametricSurface/Private/ParametricSurfaceTranslator.cpp
--- a/Engine/Plugins/Enterprise/DatasmithCADImporter/Source/ParametricSurface/Private/ParametricSurfaceTranslator.cpp
+++ b/Engine/Plugins/Enterprise/DatasmithCADImporter/Source/ParametricSurface/Private/ParametricSurfaceTranslator.cpp
@@ -40,6 +40,7 @@ void FParametricSurfaceTranslator::SetSceneImportOptions(const TArray<TObjectPtr
 		if (UDatasmithCommonTessellationOptions* TessellationOptionsObject = Cast<UDatasmithCommonTessellationOptions>(OptionPtr))
 		{
 			CommonTessellationOptions = TessellationOptionsObject->Options;
+			TessellationOptionsObject->SaveConfig(CPF_Config);
 		}
 	}
 }
