Fixed crash importing CAD scene with root node(and below) hidden

- added ensures to all the places where root node it accessed

diff --git a/Engine/Plugins/Enterprise/DatasmithCADImporter/Source/CADInterfaces/Private/TechSoftFileParser.cpp b/Engine/Plugins/Enterprise/DatasmithCADImporter/Source/CADInterfaces/Private/TechSoftFileParser.cpp
--- a/Engine/Plugins/Enterprise/DatasmithCADImporter/Source/CADInterfaces/Private/TechSoftFileParser.cpp
+++ b/Engine/Plugins/Enterprise/DatasmithCADImporter/Source/CADInterfaces/Private/TechSoftFileParser.cpp
@@ -448,7 +448,6 @@ ECADParsingResult FTechSoftFileParser::Process()
 
 	ProcessReport.ImportTime += FPlatformTime::ToMilliseconds64(FPlatformTime::Cycles64() - StartTraversTime);
 
-
 	if (Result == ECADParsingResult::ProcessOk)
 	{
 		GenerateBodyMeshes();
@@ -460,8 +459,11 @@ ECADParsingResult FTechSoftFileParser::Process()
 		FString TechSoftVersion = TechSoftInterface::GetTechSoftVersion();
 		if (!TechSoftVersion.IsEmpty())
 		{
-			FArchiveReference& RootReference = SceneGraph.GetReference(1);
-			RootReference.MetaData.Add(TEXT("TechsoftVersion"), TechSoftVersion);
+			FArchiveReference* RootReference = SceneGraph.GetRootReference();
+			if (ensureMsgf(RootReference, TEXT("FTechSoftFileParser expected root Reference in the Scene Graph")))
+			{
+				RootReference->MetaData.Add(TEXT("TechsoftVersion"), TechSoftVersion);
+			}
 		}
 	}
 
@@ -839,7 +841,11 @@ void FTechSoftFileParser::TraverseReference(const A3DAsmProductOccurrence* A3DRe
 
 	if (Reference.bIsRemoved || !Reference.bShow)
 	{
-		SceneGraph.RemoveLastReference();
+		if (Reference.bIsRemoved)
+		{
+			// Keep node in the SceneGraph if occurrence is not removed but just not visible
+			SceneGraph.RemoveLastReference();
+		}
 		return;
 	}
 
diff --git a/Engine/Plugins/Enterprise/DatasmithCADImporter/Source/CADInterfaces/Public/CADSceneGraph.h b/Engine/Plugins/Enterprise/DatasmithCADImporter/Source/CADInterfaces/Public/CADSceneGraph.h
--- a/Engine/Plugins/Enterprise/DatasmithCADImporter/Source/CADInterfaces/Public/CADSceneGraph.h
+++ b/Engine/Plugins/Enterprise/DatasmithCADImporter/Source/CADInterfaces/Public/CADSceneGraph.h
@@ -265,6 +265,12 @@ public:
 	FArchiveReference& AddReference(FArchiveUnloadedReference&);
 	FArchiveReference& AddReference(FArchiveInstance& Parent);
 	FArchiveReference& GetReference(const FCadId CadId);
+	FArchiveReference* GetRootReference()
+	{
+		constexpr FCadId RootId = 1;
+		return CADIdToIndex.IsValidIndex(RootId) ? &GetReference(RootId) : nullptr;
+	}
+
 	void RemoveLastReference();
 	bool IsAReference(FCadId CadId) const;
 
diff --git a/Engine/Plugins/Enterprise/DatasmithCADImporter/Source/DatasmithCADTranslator/Private/DatasmithCADTranslator.cpp b/Engine/Plugins/Enterprise/DatasmithCADImporter/Source/DatasmithCADTranslator/Private/DatasmithCADTranslator.cpp
--- a/Engine/Plugins/Enterprise/DatasmithCADImporter/Source/DatasmithCADTranslator/Private/DatasmithCADTranslator.cpp
+++ b/Engine/Plugins/Enterprise/DatasmithCADImporter/Source/DatasmithCADTranslator/Private/DatasmithCADTranslator.cpp
@@ -273,6 +273,36 @@ bool FDatasmithCADTranslator::LoadScene(TSharedRef<IDatasmithScene> DatasmithSce
 	FCADFileReader FileReader(ImportParameters, FileDescriptor, *FPaths::EnginePluginsDir(), CachePath);
 	if (FileReader.ProcessFile() != ECADParsingResult::ProcessOk)
 	{
+		// todo: copied from FDatasmithDispatcher::LogMessages
+		for (const TPair<uint8, FString>& Message : FileReader.GetCADFileData().GetMessages())
+		{
+			switch (Message.Key)
+			{
+			case ELogVerbosity::Fatal:
+				UE_LOG(LogCADTranslator, Fatal, TEXT("%s"), *Message.Value);
+				break;
+			case ELogVerbosity::Error:
+				UE_LOG(LogCADTranslator, Error, TEXT("%s"), *Message.Value);
+				break;
+			case ELogVerbosity::Warning:
+				UE_LOG(LogCADTranslator, Warning, TEXT("%s"), *Message.Value);
+				break;
+			case ELogVerbosity::Display:
+				UE_LOG(LogCADTranslator, Display, TEXT("%s"), *Message.Value);
+				break;
+			case ELogVerbosity::Log:
+				UE_LOG(LogCADTranslator, Log, TEXT("%s"), *Message.Value);
+				break;
+			case ELogVerbosity::Verbose:
+				UE_LOG(LogCADTranslator, Verbose, TEXT("%s"), *Message.Value);
+				break;
+			case ELogVerbosity::VeryVerbose:
+			default:
+				UE_LOG(LogCADTranslator, VeryVerbose, TEXT("%s"), *Message.Value);
+				break;
+			}
+		}
+
 		return false;
 	}
 
diff --git a/Engine/Plugins/Enterprise/DatasmithCADImporter/Source/DatasmithCADTranslator/Private/DatasmithSceneGraphBuilder.cpp b/Engine/Plugins/Enterprise/DatasmithCADImporter/Source/DatasmithCADTranslator/Private/DatasmithSceneGraphBuilder.cpp
--- a/Engine/Plugins/Enterprise/DatasmithCADImporter/Source/DatasmithCADTranslator/Private/DatasmithSceneGraphBuilder.cpp
+++ b/Engine/Plugins/Enterprise/DatasmithCADImporter/Source/DatasmithCADTranslator/Private/DatasmithSceneGraphBuilder.cpp
@@ -179,21 +179,25 @@ void FDatasmithSceneGraphBuilder::FillAnchorActor(const TSharedRef<IDatasmithAct
 
 	FActorData Data(TEXT(""));
 
-	const FCadId RootId = 1;
 	FActorData ParentData(ActorElement->GetName());
-	CADLibrary::FArchiveReference& Reference = SceneGraph->GetReference(RootId);
+
+	CADLibrary::FArchiveReference* Reference = SceneGraph->GetRootReference();
+	if (!ensureMsgf(Reference, TEXT("FDatasmithSceneGraphBuilder::FillAnchorActor expected root Reference in the Scene Graph")))
+	{
+		return;
+	}
 
 	CADLibrary::FArchiveInstance EmptyInstance;
 	FString ActorUUID;
 	FString ActorLabel;
-	DatasmithSceneGraphBuilderImpl::GetNodeUuidAndLabel(EmptyInstance, Reference, ParentData.Uuid, ActorUUID, ActorLabel);
+	DatasmithSceneGraphBuilderImpl::GetNodeUuidAndLabel(EmptyInstance, *Reference, ParentData.Uuid, ActorUUID, ActorLabel);
 
-	AddMetaData(ActorElement, EmptyInstance, Reference);
+	AddMetaData(ActorElement, EmptyInstance, *Reference);
 
 	FActorData ComponentData(*ActorUUID, ParentData);
 	//DatasmithSceneGraphBuilderImpl::GetMainMaterial(Reference.MetaData, ComponentData, bMaterialPropagationIsTopDown);
 
-	AddChildren(ActorElement, Reference, ComponentData);
+	AddChildren(ActorElement, *Reference, ComponentData);
 
 	ActorElement->SetLabel(*ActorLabel);
 }
@@ -224,11 +228,13 @@ FDatasmithSceneBaseGraphBuilder::FDatasmithSceneBaseGraphBuilder(CADLibrary::FAr
 
 bool FDatasmithSceneBaseGraphBuilder::Build()
 {
-	const FCadId RootId = 1;
-
 	FActorData Data(TEXT(""));
-	CADLibrary::FArchiveReference& Reference = SceneGraph->GetReference(RootId);
-	TSharedPtr<IDatasmithActorElement> RootActor = BuildReference(Reference, Data);
+	CADLibrary::FArchiveReference* Reference = SceneGraph->GetRootReference();
+	if (!ensureMsgf(Reference, TEXT("FDatasmithSceneBaseGraphBuilder::Build: expected root Reference in the Scene Graph")))
+	{
+		return false;
+	}
+	TSharedPtr<IDatasmithActorElement> RootActor = BuildReference(*Reference, Data);
 	DatasmithScene->AddActor(RootActor);
 
 	// Set ProductName, ProductVersion in DatasmithScene for Analytics purpose
@@ -236,7 +242,7 @@ bool FDatasmithSceneBaseGraphBuilder::Build()
 	DatasmithScene->SetVendor(TEXT("Techsoft"));
 
 	FString ProductVersion;
-	if (Reference.MetaData.RemoveAndCopyValue(TEXT("TechsoftVersion"), ProductVersion))
+	if (Reference->MetaData.RemoveAndCopyValue(TEXT("TechsoftVersion"), ProductVersion))
 	{
 		DatasmithScene->SetProductVersion(*ProductVersion);
 	}
@@ -363,9 +369,8 @@ TSharedPtr<IDatasmithActorElement>  FDatasmithSceneBaseGraphBuilder::BuildInstan
 				{
 					AncestorSceneGraphHash.Add(InstanceSceneGraphHash);
 
-					FCadId RootId = 1;
-					Reference = &(SceneGraph->GetReference(RootId));
-					if (Reference)
+					Reference = SceneGraph->GetRootReference();
+					if (ensureMsgf(Reference, TEXT("FDatasmithSceneBaseGraphBuilder::BuildInstance expected root Reference in the Scene Graph of an external reference")))
 					{
 						Reference->SetGraphicProperties(UnloadedReference);
 					}
