[Wire] Fixed missing colors when using group option
- Fixed coding error in FDatasmithStaticMeshImporter::SetupStaticMesh which was eliminating sections when some were sharing the same material
- Simplified material assignment to MeshElement's slots.
- removed redundant material assignment on MeshActor.
- Fixed wrong material slot name used in FMeshDescription. It has to be an integer to work in Datasmith import.

#rnx

diff --git a/Engine/Plugins/Enterprise/DatasmithCADImporter/Source/WireInterface/Private/OpenModelUtils.h b/Engine/Plugins/Enterprise/DatasmithCADImporter/Source/WireInterface/Private/OpenModelUtils.h
--- a/Engine/Plugins/Enterprise/DatasmithCADImporter/Source/WireInterface/Private/OpenModelUtils.h
+++ b/Engine/Plugins/Enterprise/DatasmithCADImporter/Source/WireInterface/Private/OpenModelUtils.h
@@ -413,7 +413,7 @@ namespace UE_DATASMITHWIRETRANSLATOR_NAMESPACE
 
 		void AddMeshNode(FAlDagNodePtr& MeshNode)
 		{
-			ensure(Layer == MeshNode.GetLayer());
+			ensureWire(Layer == MeshNode.GetLayer());
 			MeshNodes.Add(MeshNode);
 		}
 
diff --git a/Engine/Plugins/Enterprise/DatasmithCADImporter/Source/WireInterface/Private/WireInterfaceImpl.cpp b/Engine/Plugins/Enterprise/DatasmithCADImporter/Source/WireInterface/Private/WireInterfaceImpl.cpp
--- a/Engine/Plugins/Enterprise/DatasmithCADImporter/Source/WireInterface/Private/WireInterfaceImpl.cpp
+++ b/Engine/Plugins/Enterprise/DatasmithCADImporter/Source/WireInterface/Private/WireInterfaceImpl.cpp
@@ -667,44 +667,6 @@ namespace UE_DATASMITHWIRETRANSLATOR_NAMESPACE
 
 		GeomNode.SetActorTransform(*ActorElement);
 
-		TAlObjectPtr<AlShell> Shell;
-		if (GeomNode.GetShell(Shell))
-		{
-			TAlObjectPtr<AlShader> Shader(Shell->firstShader());
-			int32 SlotIndex = 0;
-			while (Shader)
-			{
-				TSharedPtr<IDatasmithMaterialIDElement> MaterialIDElement = FindOrAddMaterial(Shader);
-				MaterialIDElement->SetId(SlotIndex++);
-				ActorElement->AddMaterialOverride(MaterialIDElement);
-				Shader = Shell->nextShader(Shader.Get());
-			}
-		}
-		else
-		{
-			TSharedPtr<IDatasmithMaterialIDElement> MaterialIDElement;
-
-			TAlObjectPtr<AlSurface> Surface;
-			if (GeomNode.GetSurface(Surface))
-			{
-				MaterialIDElement = FindOrAddMaterial(Surface->firstShader());
-			}
-			else
-			{
-				TAlObjectPtr<AlMesh> Mesh;
-				if (GeomNode.GetMesh(Mesh))
-				{
-					MaterialIDElement = FindOrAddMaterial(Mesh->firstShader());
-				}
-			}
-
-			if (MaterialIDElement)
-			{
-				MaterialIDElement->SetId(0);
-				ActorElement->AddMaterialOverride(MaterialIDElement);
-			}
-		}
-
 		if (WireSettings.bUseLayerAsActor && Layer != ParentLayer)
 		{
 			if (TSharedPtr<IDatasmithActorElement> LayerActor = FindOrAddLayerActor(Layer))
@@ -857,10 +819,9 @@ namespace UE_DATASMITHWIRETRANSLATOR_NAMESPACE
 #endif
 		auto ApplyMaterial = [this, &MeshElement](const TAlObjectPtr<AlShader>& Shader, int32 SlotIndex)
 			{
-				if (TSharedPtr<IDatasmithMaterialIDElement> MaterialIDElement = this->FindOrAddMaterial(Shader))
+				if (TSharedPtr<IDatasmithBaseMaterialElement> MaterialElement = this->FindOrAddMaterial(Shader))
 				{
-					MaterialIDElement->SetId(SlotIndex);
-					MeshElement->SetMaterial(MaterialIDElement->GetName(), SlotIndex);
+					MeshElement->SetMaterial(MaterialElement->GetName(), SlotIndex);
 				}
 			};
 
@@ -938,10 +899,9 @@ namespace UE_DATASMITHWIRETRANSLATOR_NAMESPACE
 
 		auto ApplyMaterial = [this, &MeshElement](int32 SlotIndex, const TAlObjectPtr<AlShader>& Shader)
 			{
-				if (TSharedPtr<IDatasmithMaterialIDElement> MaterialIDElement = this->FindOrAddMaterial(Shader))
+				if (TSharedPtr<IDatasmithBaseMaterialElement> MaterialElement = this->FindOrAddMaterial(Shader))
 				{
-					MaterialIDElement->SetId(SlotIndex);
-					MeshElement->SetMaterial(MaterialIDElement->GetName(), SlotIndex);
+					MeshElement->SetMaterial(MaterialElement->GetName(), SlotIndex);
 				}
 			};
 
@@ -973,13 +933,20 @@ namespace UE_DATASMITHWIRETRANSLATOR_NAMESPACE
 
 		MeshElement->SetLabel(*PatchMesh->GetName());
 		MeshElement->SetLightmapSourceUV(-1);
+#if TRACK_MESHELEMENT
+		{
+			if (!ToTrack.Equals(MeshElement->GetLabel()))
+			{
+				return TSharedPtr<IDatasmithMeshElement>();
+			}
+		}
+#endif
 
 		auto ApplyMaterial = [this, &MeshElement](const TAlObjectPtr<AlShader>& Shader, int32 SlotIndex)
 			{
-				if (TSharedPtr<IDatasmithMaterialIDElement> MaterialIDElement = this->FindOrAddMaterial(Shader))
+				if (TSharedPtr<IDatasmithBaseMaterialElement> MaterialElement = this->FindOrAddMaterial(Shader))
 				{
-					MaterialIDElement->SetId(SlotIndex);
-					MeshElement->SetMaterial(MaterialIDElement->GetName(), SlotIndex);
+					MeshElement->SetMaterial(MaterialElement->GetName(), SlotIndex);
 				}
 			};
 
@@ -1195,6 +1162,7 @@ namespace UE_DATASMITHWIRETRANSLATOR_NAMESPACE
 		MeshDescription.Empty();
 
 		constexpr bool bMerge = true;
+		int32 SlotIndex = 0;
 		PatchMesh->IterateOnMeshNodes([&](const FAlDagNodePtr& MeshNode)
 			{
 				TAlObjectPtr<AlMesh> Mesh;
@@ -1212,8 +1180,7 @@ namespace UE_DATASMITHWIRETRANSLATOR_NAMESPACE
 
 					Mesh->transform(AlMatrix);
 
-					TSharedPtr<IDatasmithMaterialIDElement> MaterialIDElement = MeshElement->GetMaterialSlotAt(0);
-					const FString SlotMaterialName = MaterialIDElement ? MaterialIDElement->GetName() : FString();
+					const FString SlotMaterialName = DatasmithMeshHelper::DefaultSlotName(SlotIndex).ToString();
 
 					OpenModelUtils::TransferAlMeshToMeshDescription(*Mesh, *SlotMaterialName, MeshDescription, OutMeshParameters, bMerge);
 				}
@@ -1298,11 +1265,8 @@ namespace UE_DATASMITHWIRETRANSLATOR_NAMESPACE
 		FMeshDescription MeshDescription;
 		DatasmithMeshHelper::PrepareAttributeForStaticMesh(MeshDescription);
 
-		TSharedPtr<IDatasmithMaterialIDElement> MaterialIDElement = MeshElement->GetMaterialSlotAt(0);
-		const FString SlotMaterialName = MaterialIDElement ? MaterialIDElement->GetName() : FString();
-
 		const bool bMerge = false;
-		OpenModelUtils::TransferAlMeshToMeshDescription(*Mesh, *SlotMaterialName, MeshDescription, OutMeshParameters, bMerge);
+		OpenModelUtils::TransferAlMeshToMeshDescription(*Mesh, TEXT("0"), MeshDescription, OutMeshParameters, bMerge);
 
 		// Build edge meta data
 		FStaticMeshOperations::DetermineEdgeHardnessesFromVertexInstanceNormals(MeshDescription);
@@ -1334,20 +1298,17 @@ namespace UE_DATASMITHWIRETRANSLATOR_NAMESPACE
 
 	// Material creation
 
-	TSharedPtr<IDatasmithMaterialIDElement> FWireTranslatorImpl::FindOrAddMaterial(const TAlObjectPtr<AlShader>& Shader)
+	TSharedPtr<IDatasmithBaseMaterialElement> FWireTranslatorImpl::FindOrAddMaterial(const TAlObjectPtr<AlShader>& Shader)
 	{
 		const FString ShaderName = Shader.GetName();
 
 		if (TSharedPtr<IDatasmithBaseMaterialElement>* MaterialElementPtr = ShaderNameToMaterial.Find(ShaderName))
 		{
-			return FDatasmithSceneFactory::CreateMaterialId((*MaterialElementPtr)->GetName());
+			return *MaterialElementPtr;
 		}
 
 		const FString ShaderModelName = Shader->shadingModel();
 
-		//const FColor Color = CreateShaderColorFromShaderName(ShaderName);
-		//ShaderNameToColor.Add(ShaderName, Color);
-
 		TSharedPtr<IDatasmithUEPbrMaterialElement> MaterialElement = FDatasmithSceneFactory::CreateUEPbrMaterial(*Shader.GetUniqueID(SHADER_TYPE));
 		MaterialElement->SetLabel(*ShaderName);
 
@@ -1371,7 +1332,7 @@ namespace UE_DATASMITHWIRETRANSLATOR_NAMESPACE
 		DatasmithScene->AddMaterial(MaterialElement);
 		ShaderNameToMaterial.Add(*ShaderName, MaterialElement);
 
-		return FDatasmithSceneFactory::CreateMaterialId(MaterialElement->GetName());
+		return MaterialElement;
 	}
 
 	bool FWireTranslatorImpl::GetCommonParameters(int32 Field, double Value, FColor& Color, FColor& TransparencyColor, FColor& IncandescenceColor, double GlowIntensity)
diff --git a/Engine/Plugins/Enterprise/DatasmithCADImporter/Source/WireInterface/Private/WireInterfaceImpl.h b/Engine/Plugins/Enterprise/DatasmithCADImporter/Source/WireInterface/Private/WireInterfaceImpl.h
--- a/Engine/Plugins/Enterprise/DatasmithCADImporter/Source/WireInterface/Private/WireInterfaceImpl.h
+++ b/Engine/Plugins/Enterprise/DatasmithCADImporter/Source/WireInterface/Private/WireInterfaceImpl.h
@@ -81,7 +81,7 @@ namespace UE_DATASMITHWIRETRANSLATOR_NAMESPACE
 		void AddAlLightSourceParameters(const TAlObjectPtr<AlShader>& Shader, TSharedPtr<IDatasmithUEPbrMaterialElement> MaterialElement);
 		void AddAlPhongParameters(const TAlObjectPtr<AlShader>& Shader, TSharedPtr<IDatasmithUEPbrMaterialElement> MaterialElement);
 		
-		TSharedPtr<IDatasmithMaterialIDElement> FindOrAddMaterial(const TAlObjectPtr<AlShader>& Shader);
+		TSharedPtr<IDatasmithBaseMaterialElement> FindOrAddMaterial(const TAlObjectPtr<AlShader>& Shader);
 
 		/** Geometry retrieval */
 		TOptional<FMeshDescription> GetMeshDescription(TSharedPtr<IDatasmithMeshElement> MeshElement, CADLibrary::FMeshParameters& OutMeshParameters);
@@ -138,6 +138,9 @@ namespace UE_DATASMITHWIRETRANSLATOR_NAMESPACE
 		TMap<AlDagNode*, FAlDagNodePtr> EncounteredNodes;
 
 		TMap<uint32, TSharedPtr<IDatasmithActorElement>> LayerToActor;
+
+		// #cad_debug
+		bool bTrackMesh = false;
 	};
 } // namespace
 #endif
diff --git a/Engine/Plugins/Enterprise/DatasmithImporter/Source/DatasmithImporter/Private/DatasmithStaticMeshImporter.cpp b/Engine/Plugins/Enterprise/DatasmithImporter/Source/DatasmithImporter/Private/DatasmithStaticMeshImporter.cpp
--- a/Engine/Plugins/Enterprise/DatasmithImporter/Source/DatasmithImporter/Private/DatasmithStaticMeshImporter.cpp
+++ b/Engine/Plugins/Enterprise/DatasmithImporter/Source/DatasmithImporter/Private/DatasmithStaticMeshImporter.cpp
@@ -483,7 +483,8 @@ void FDatasmithStaticMeshImporter::PreBuildStaticMeshes( FDatasmithImportContext
 
 class FMeshMaterialSlotBuilder
 {
-	struct SectionInfo
+public:
+	struct FSectionInfo
 	{
 		enum { InvalidIndex = -1 };
 
@@ -495,7 +496,6 @@ class FMeshMaterialSlotBuilder
 		int32 StaticMaterialsIndex = InvalidIndex; // index in the FStaticMesh::StaticMaterials
 	};
 
-public:
 	FMeshMaterialSlotBuilder(UStaticMesh& StaticMesh, int32 MeshElementLODCount)
 	{
 		// Populate section LODs declared in IDatasmithMeshElement
@@ -509,7 +509,7 @@ public:
 				RawSections.Reserve(MeshDescription->PolygonGroups().Num());
 				for (FPolygonGroupID PolygonGroupID : MeshDescription->PolygonGroups().GetElementIDs())
 				{
-					SectionInfo& ThisSection = RawSections.Add_GetRef({});
+					FSectionInfo& ThisSection = RawSections.Add_GetRef({});
 					ThisSection.LodIndex = LodIndex;
 					ThisSection.SectionIndex = SectionIndex++;
 					ThisSection.SectionId = PolygonGroupID;
@@ -520,19 +520,19 @@ public:
 		}
 
 		// Prepares StaticMaterials Slot informations from sections found in the mesh
-		for (SectionInfo& Section : RawSections)
+		for (FSectionInfo& Section : RawSections)
 		{
 			// get associated slot (but skip empty polygon groups)
-			Section.StaticMaterialsIndex = Section.PolyCount <= 0 ? SectionInfo::InvalidIndex : SlotNames.Add(Section.MaterialSlotName).AsInteger();
+			Section.StaticMaterialsIndex = Section.PolyCount <= 0 ? FSectionInfo::InvalidIndex : SlotNames.Add(Section.MaterialSlotName).AsInteger();
 		}
 	}
 
 	FMeshSectionInfoMap GenerateSectionInfoMap() const
 	{
 		FMeshSectionInfoMap InfoMap;
-		for (const SectionInfo& Section : RawSections)
+		for (const FSectionInfo& Section : RawSections)
 		{
-			if (Section.StaticMaterialsIndex != SectionInfo::InvalidIndex)
+			if (Section.StaticMaterialsIndex != FSectionInfo::InvalidIndex)
 			{
 				InfoMap.Set(Section.LodIndex, Section.SectionIndex, FMeshSectionInfo(Section.StaticMaterialsIndex));
 			}
@@ -540,7 +540,7 @@ public:
 		return InfoMap;
 	}
 
-	const TArray<SectionInfo>& GetSectionsInfos() const
+	const TArray<FSectionInfo>& GetSectionsInfos() const
 	{
 		return RawSections;
 	}
@@ -552,7 +552,7 @@ public:
 
 
 private:
-	TArray<SectionInfo> RawSections;
+	TArray<FSectionInfo> RawSections;
 	TSet<FName> SlotNames;
 };
 
@@ -691,11 +691,11 @@ void FDatasmithStaticMeshImporter::SetupStaticMesh( FDatasmithAssetsImportContex
 	FMeshSectionInfoMap MeshSectionInfoMap = Sections.GenerateSectionInfoMap();
 	StaticMeshTemplate->SectionInfoMap.Load( MeshSectionInfoMap );
 
-	// Create one StaticMaterial per material index. Actual materials will be assigned later. This is just so that we create all the slots.
-	for (const FName& SlotName : Sections.GetSlotNames())
+	// Create one StaticMaterial per section. Actual materials will be assigned later. This is just so that we create all the slots.
+	for (const FMeshMaterialSlotBuilder::FSectionInfo& SectionInfo : Sections.GetSectionsInfos())
 	{
 		FDatasmithStaticMaterialTemplate StaticMaterial;
-		StaticMaterial.MaterialSlotName = SlotName;
+		StaticMaterial.MaterialSlotName = SectionInfo.MaterialSlotName;
 
 		StaticMeshTemplate->StaticMaterials.Add(MoveTemp(StaticMaterial));
 	}
