[Wire] Corrected missing incrementation
The mesh was properly sectioned but the missing increment was assigning the same material to each section
Somehow the increment step was deleted before submission :-(

#rnx

diff --git a/Engine/Plugins/Enterprise/DatasmithCADImporter/Source/WireInterface/Private/WireInterfaceImpl.cpp b/Engine/Plugins/Enterprise/DatasmithCADImporter/Source/WireInterface/Private/WireInterfaceImpl.cpp
--- a/Engine/Plugins/Enterprise/DatasmithCADImporter/Source/WireInterface/Private/WireInterfaceImpl.cpp
+++ b/Engine/Plugins/Enterprise/DatasmithCADImporter/Source/WireInterface/Private/WireInterfaceImpl.cpp
@@ -1180,7 +1180,7 @@ namespace UE_DATASMITHWIRETRANSLATOR_NAMESPACE
 
 					Mesh->transform(AlMatrix);
 
-					const FString SlotMaterialName = DatasmithMeshHelper::DefaultSlotName(SlotIndex).ToString();
+					const FString SlotMaterialName = DatasmithMeshHelper::DefaultSlotName(SlotIndex++).ToString();
 
 					OpenModelUtils::TransferAlMeshToMeshDescription(*Mesh, *SlotMaterialName, MeshDescription, OutMeshParameters, bMerge);
 				}
