Replaced FPoint and FPoint2D by FVector and  FVector2d and remove unused FFPoint and FPointH classes - Final

#rnx

diff --git a/Engine/Plugins/Enterprise/DatasmithCADImporter/Source/CADInterfaces/Private/TechSoftBridge.cpp b/Engine/Plugins/Enterprise/DatasmithCADImporter/Source/CADInterfaces/Private/TechSoftBridge.cpp
--- a/Engine/Plugins/Enterprise/DatasmithCADImporter/Source/CADInterfaces/Private/TechSoftBridge.cpp
+++ b/Engine/Plugins/Enterprise/DatasmithCADImporter/Source/CADInterfaces/Private/TechSoftBridge.cpp
@@ -53,9 +53,11 @@ A3DStatus GetSurfaceAsNurbs(const A3DSurfBase* A3DSurface, A3DSurfNurbsData* Dat
 
 UE::CADKernel::FMatrixH CreateCoordinateSystem(const A3DMiscCartesianTransformationData& Transformation, double UnitScale = 1.0)
 {
-	UE::CADKernel::FPoint Origin(&Transformation.m_sOrigin.m_dX);
-	UE::CADKernel::FPoint Ox(&Transformation.m_sXVector.m_dX);
-	UE::CADKernel::FPoint Oy(&Transformation.m_sYVector.m_dX);
+	using namespace UE::CADKernel;
+
+	FVector Origin(Transformation.m_sOrigin.m_dX, Transformation.m_sOrigin.m_dY, Transformation.m_sOrigin.m_dZ);
+	FVector Ox(Transformation.m_sXVector.m_dX, Transformation.m_sXVector.m_dY, Transformation.m_sXVector.m_dZ);
+	FVector Oy(Transformation.m_sYVector.m_dX, Transformation.m_sYVector.m_dY, Transformation.m_sYVector.m_dZ);
 
 	Ox.Normalize();
 	Oy.Normalize();
@@ -64,13 +66,13 @@ UE::CADKernel::FMatrixH CreateCoordinateSystem(const A3DMiscCartesianTransformat
 	{
 		Origin *= UnitScale;
 	}
-	UE::CADKernel::FPoint Oz = Ox ^ Oy;
+	FVector Oz = Ox ^ Oy;
 
-	UE::CADKernel::FMatrixH Matrix = UE::CADKernel::FMatrixH(Origin, Ox, Oy, Oz);
+	FMatrixH Matrix = FMatrixH(Origin, Ox, Oy, Oz);
 
 	if (!FMath::IsNearlyEqual(Transformation.m_sScale.m_dX, 1.) || !FMath::IsNearlyEqual(Transformation.m_sScale.m_dY, 1.) || !FMath::IsNearlyEqual(Transformation.m_sScale.m_dZ, 1.))
 	{
-		UE::CADKernel::FMatrixH Scale = UE::CADKernel::FMatrixH::MakeScaleMatrix(Transformation.m_sScale.m_dX, Transformation.m_sScale.m_dY, Transformation.m_sScale.m_dZ);
+		FMatrixH Scale = FMatrixH::MakeScaleMatrix(Transformation.m_sScale.m_dX, Transformation.m_sScale.m_dY, Transformation.m_sScale.m_dZ);
 		Matrix *= Scale;
 	}
 	return Matrix;
@@ -106,27 +108,31 @@ void FillDoubleArray(const int32 UCount, const int32 VCount, const double* Value
 		}
 	}
 }
-
-void FillPointArray(const int32 Count, const A3DVector3dData* Points, TArray<UE::CADKernel::FPoint>& OutPointsArray, double UnitScale = 1.0)
+void FillPointArray(const int32 Count, const A3DVector3dData* Points, TArray<FVector>& OutPointsArray, double UnitScale = 1.0)
 {
+	using namespace UE::CADKernel;
+
 	OutPointsArray.Reserve(Count);
 	for (int32 Index = 0; Index < Count; Index++)
 	{
-		OutPointsArray.Emplace(&Points[Index].m_dX);
+		OutPointsArray.Emplace(Points[Index].m_dX, Points[Index].m_dY, Points[Index].m_dZ);
 	}
 
 	if (!FMath::IsNearlyEqual(UnitScale, 1.))
 	{
-		for (UE::CADKernel::FPoint& Point : OutPointsArray)
+		for (FVector& Point : OutPointsArray)
 		{
 			Point *= UnitScale;
 		}
 	}
 };
 
-void FillPointArray(const int32 UCount, const int32 VCount, const A3DVector3dData* Points, TArray<UE::CADKernel::FPoint>& OutPointsArray, double UnitScale = 1.0)
+void FillPointArray(const int32 UCount, const int32 VCount, const A3DVector3dData* Points, TArray<FVector>& OutPointsArray, double UnitScale = 1.0)
 {
+	using namespace UE::CADKernel;
+
 	OutPointsArray.SetNum(UCount * VCount);
+
 	for (int32 Undex = 0, PointIndex = 0; Undex < UCount; ++Undex)
 	{
 		int32 Index = Undex;
@@ -138,7 +144,7 @@ void FillPointArray(const int32 UCount, const int32 VCount, const A3DVector3dDat
 
 	if (!FMath::IsNearlyEqual(UnitScale, 1.))
 	{
-		for (UE::CADKernel::FPoint& Point : OutPointsArray)
+		for (FVector& Point : OutPointsArray)
 		{
 			Point *= UnitScale;
 		}
@@ -147,9 +153,11 @@ void FillPointArray(const int32 UCount, const int32 VCount, const A3DVector3dDat
 
 UE::CADKernel::FSurfacicBoundary GetSurfacicBoundary(A3DDomainData& Domain, const TechSoftUtils::FUVReparameterization& UVReparameterization)
 {
+	using namespace UE::CADKernel;
 
-	UE::CADKernel::FPoint2D Min(Domain.m_sMin.m_dX, Domain.m_sMin.m_dY);
-	UE::CADKernel::FPoint2D Max(Domain.m_sMax.m_dX, Domain.m_sMax.m_dY);
+
+	FVector2d Min(Domain.m_sMin.m_dX, Domain.m_sMin.m_dY);
+	FVector2d Max(Domain.m_sMax.m_dX, Domain.m_sMax.m_dY);
 
 	if (UVReparameterization.GetNeedApply())
 	{
@@ -157,10 +165,10 @@ UE::CADKernel::FSurfacicBoundary GetSurfacicBoundary(A3DDomainData& Domain, cons
 		UVReparameterization.Apply(Max);
 	}
 
-	UE::CADKernel::EIso UIndex = UVReparameterization.GetSwapUV() ? UE::CADKernel::EIso::IsoV : UE::CADKernel::EIso::IsoU;
-	UE::CADKernel::EIso VIndex = UVReparameterization.GetSwapUV() ? UE::CADKernel::EIso::IsoU : UE::CADKernel::EIso::IsoV;
+	EIso UIndex = UVReparameterization.GetSwapUV() ? EIso::IsoV : EIso::IsoU;
+	EIso VIndex = UVReparameterization.GetSwapUV() ? EIso::IsoU : EIso::IsoV;
 
-	UE::CADKernel::FSurfacicBoundary Boundary;
+	FSurfacicBoundary Boundary;
 	Boundary[UIndex].Min = FMath::Min(Min.X, Max.X);
 	Boundary[UIndex].Max = FMath::Max(Min.X, Max.X);
 	Boundary[VIndex].Min = FMath::Min(Min.Y, Max.Y);
@@ -829,7 +837,7 @@ void FTechSoftBridge::AddFace(const A3DTopoFace* A3DFace, UE::CADKernel::EOrient
 				continue;
 			}
 
-			TArray<FPoint2D> LoopSampling;
+			TArray<FVector2d> LoopSampling;
 			Loop->Get2DSampling(LoopSampling);
 			FAABB2D Boundary;
 			Boundary += LoopSampling;
@@ -1362,7 +1370,9 @@ TSharedPtr<UE::CADKernel::FSurface> FTechSoftBridge::AddTransformSurface(const A
 
 TSharedPtr<UE::CADKernel::FSurface> FTechSoftBridge::AddSurfaceNurbs(const A3DSurfNurbsData& A3DNurbsData, TechSoftUtils::FUVReparameterization& OutUVReparameterization)
 {
-	UE::CADKernel::FNurbsSurfaceData NurbsData;
+	using namespace UE::CADKernel;
+
+	FNurbsSurfaceData NurbsData;
 
 	NurbsData.PoleUCount = A3DNurbsData.m_uiUCtrlSize;
 	NurbsData.PoleVCount = A3DNurbsData.m_uiVCtrlSize;
@@ -1374,11 +1384,11 @@ TSharedPtr<UE::CADKernel::FSurface> FTechSoftBridge::AddSurfaceNurbs(const A3DSu
 	TechSoftUtils::FillDoubleArray(A3DNurbsData.m_uiUKnotSize, A3DNurbsData.m_pdUKnots, NurbsData.UNodalVector);
 	TechSoftUtils::FillDoubleArray(A3DNurbsData.m_uiVKnotSize, A3DNurbsData.m_pdVKnots, NurbsData.VNodalVector);
 
-	TArray<UE::CADKernel::FPoint> Poles;
+	TArray<FVector> Poles;
 	TechSoftUtils::FillPointArray(NurbsData.PoleUCount, NurbsData.PoleVCount, A3DNurbsData.m_pCtrlPts, NurbsData.Poles);
 	if (!FMath::IsNearlyEqual(BodyScale, 1.))
 	{
-		for (UE::CADKernel::FPoint& Point : NurbsData.Poles)
+		for (FVector& Point : NurbsData.Poles)
 		{
 			Point *= BodyScale;
 		}
diff --git a/Engine/Plugins/Enterprise/DatasmithCADImporter/Source/CADInterfaces/Private/TechSoftBridge.h b/Engine/Plugins/Enterprise/DatasmithCADImporter/Source/CADInterfaces/Private/TechSoftBridge.h
--- a/Engine/Plugins/Enterprise/DatasmithCADImporter/Source/CADInterfaces/Private/TechSoftBridge.h
+++ b/Engine/Plugins/Enterprise/DatasmithCADImporter/Source/CADInterfaces/Private/TechSoftBridge.h
@@ -108,18 +108,20 @@ public:
 		SetNeedApply();
 	}
 
-	void Process(TArray<UE::CADKernel::FPoint>& Poles) const
+	void Process(TArray<FVector>& Poles) const
 	{
-		if(bNeedApply)
+		using namespace UE::CADKernel;
+
+		if (bNeedApply)
 		{
-			for (UE::CADKernel::FPoint& Point : Poles)
+			for (FVector& Point : Poles)
 			{
 				Apply(Point);
 			}
 		}
 		if (bSwapUV)
 		{
-			for (UE::CADKernel::FPoint& Point : Poles)
+			for (FVector& Point : Poles)
 			{
 				GetSwapUV(Point);
 			}
@@ -137,20 +139,24 @@ public:
 		SetNeedApply();
 	}
 
-	void Apply(UE::CADKernel::FPoint2D& Point) const
+	void Apply(FVector2d& Point) const
 	{
-		Point.X = Scale[UE::CADKernel::EIso::IsoU] * Point.X + Offset[UE::CADKernel::EIso::IsoU];
-		Point.Y = Scale[UE::CADKernel::EIso::IsoV] * Point.Y + Offset[UE::CADKernel::EIso::IsoV];
+		using namespace UE::CADKernel;
+
+		Point.X = Scale[EIso::IsoU] * Point.X + Offset[EIso::IsoU];
+		Point.Y = Scale[EIso::IsoV] * Point.Y + Offset[EIso::IsoV];
 	}
 
 private:
-	void Apply(UE::CADKernel::FPoint& Point) const
+	void Apply(FVector& Point) const
 	{
-		Point.X = Scale[UE::CADKernel::EIso::IsoU] * Point.X + Offset[UE::CADKernel::EIso::IsoU];
-		Point.Y = Scale[UE::CADKernel::EIso::IsoV] * Point.Y + Offset[UE::CADKernel::EIso::IsoV];
+		using namespace UE::CADKernel;
+
+		Point.X = Scale[EIso::IsoU] * Point.X + Offset[EIso::IsoU];
+		Point.Y = Scale[EIso::IsoV] * Point.Y + Offset[EIso::IsoV];
 	}
 
-	void GetSwapUV(UE::CADKernel::FPoint& Point) const
+	void GetSwapUV(FVector& Point) const
 	{
 		Swap(Point.X, Point.Y);
 	}
diff --git a/Engine/Source/Runtime/Datasmith/CADKernel/Base/Private/Geo/Curves/BezierCurve.cpp b/Engine/Source/Runtime/Datasmith/CADKernel/Base/Private/Geo/Curves/BezierCurve.cpp
--- a/Engine/Source/Runtime/Datasmith/CADKernel/Base/Private/Geo/Curves/BezierCurve.cpp
+++ b/Engine/Source/Runtime/Datasmith/CADKernel/Base/Private/Geo/Curves/BezierCurve.cpp
@@ -78,7 +78,7 @@ FBezierCurve::FBezierCurve(const FNurbsCurveData& NurbsCurveData)
 		if (!bIsRational && !FMath::IsNearlyEqual(WeightRef, 1., UE_DOUBLE_SMALL_NUMBER))
 		{
 			const double InvWeight = 1. / WeightRef;
-			for (FPoint& Pole : Poles)
+			for (FVector& Pole : Poles)
 			{
 				Pole *= InvWeight;
 			}
@@ -175,10 +175,10 @@ void FBezierCurve::EvaluatePoint(double Coordinate, FCurvePoint& OutPoint, int32
 
 TSharedPtr<FEntityGeom> FBezierCurve::ApplyMatrix(const FMatrixH& InMatrix) const
 {
-	TArray<FPoint> TransformedPoles;
+	TArray<FVector> TransformedPoles;
 	TransformedPoles.Reserve(Poles.Num());
 
-	for (const FPoint& Pole : Poles)
+	for (const FVector& Pole : Poles)
 	{
 		TransformedPoles.Emplace(InMatrix.Multiply(Pole));
 	}
@@ -186,9 +186,9 @@ TSharedPtr<FEntityGeom> FBezierCurve::ApplyMatrix(const FMatrixH& InMatrix) cons
 	return FEntity::MakeShared<FBezierCurve>(TransformedPoles);
 }
 
-void FBezierCurve::Offset(const FPoint& OffsetDirection)
+void FBezierCurve::Offset(const FVector& OffsetDirection)
 {
-	for (FPoint& Pole : Poles)
+	for (FVector& Pole : Poles)
 	{
 		Pole += OffsetDirection;
 	}
@@ -203,7 +203,7 @@ FInfoEntity& FBezierCurve::GetInfo(FInfoEntity& Info) const
 }
 #endif
 
-void FBezierCurve::ExtendTo(const FPoint& Point)
+void FBezierCurve::ExtendTo(const FVector& Point)
 {
 	PolylineTools::ExtendTo(Poles, Point);
 }
diff --git a/Engine/Source/Runtime/Datasmith/CADKernel/Base/Private/Geo/Curves/BezierCurve.h b/Engine/Source/Runtime/Datasmith/CADKernel/Base/Private/Geo/Curves/BezierCurve.h
--- a/Engine/Source/Runtime/Datasmith/CADKernel/Base/Private/Geo/Curves/BezierCurve.h
+++ b/Engine/Source/Runtime/Datasmith/CADKernel/Base/Private/Geo/Curves/BezierCurve.h
@@ -18,10 +18,10 @@ protected:
 	TArray<double> NodalVector;
 	int32 NumSegments = 0;
 
-	TArray<FPoint> Poles;
+	TArray<FVector> Poles;
 	TArray<double> Weights;
 
-	FBezierCurve(const TArray<FPoint>& InPoles)
+	FBezierCurve(const TArray<FVector>& InPoles)
 		: Poles{ InPoles }
 	{
 		Degree = Poles.Num() - 1;
@@ -69,13 +69,13 @@ public:
 		return Degree;
 	}
 
-	const TArray<FPoint>& GetPoles() const
+	const TArray<FVector>& GetPoles() const
 	{
 		return Poles;
 	}
 
 	virtual TSharedPtr<FEntityGeom> ApplyMatrix(const FMatrixH& InMatrix) const override;
-	virtual void Offset(const FPoint& OffsetDirection) override;
+	virtual void Offset(const FVector& OffsetDirection) override;
 
 	virtual void EvaluatePoint(double Coordinate, FCurvePoint& OutPoint, int32 DerivativeOrder = 0) const override;
 	virtual void Evaluate2DPoint(double Coordinate, FCurvePoint2D& OutPoint, int32 DerivativeOrder = 0) const override
@@ -98,7 +98,7 @@ public:
 		}
 	}
 
-	virtual void ExtendTo(const FPoint& Point) override;
+	virtual void ExtendTo(const FVector& Point) override;
 
 
 };
diff --git a/Engine/Source/Runtime/Datasmith/CADKernel/Base/Private/Geo/Curves/BoundedCurve.cpp b/Engine/Source/Runtime/Datasmith/CADKernel/Base/Private/Geo/Curves/BoundedCurve.cpp
--- a/Engine/Source/Runtime/Datasmith/CADKernel/Base/Private/Geo/Curves/BoundedCurve.cpp
+++ b/Engine/Source/Runtime/Datasmith/CADKernel/Base/Private/Geo/Curves/BoundedCurve.cpp
@@ -59,7 +59,7 @@ TSharedPtr<FEntityGeom> FBoundedCurve::ApplyMatrix(const FMatrixH& InMatrix) con
 	return FEntity::MakeShared<FBoundedCurve>(TransformedCurve.ToSharedRef(), Boundary, Dimension);
 }
 
-void FBoundedCurve::Offset(const FPoint& OffsetDirection)
+void FBoundedCurve::Offset(const FVector& OffsetDirection)
 {
 	Curve->Offset(OffsetDirection);
 }
diff --git a/Engine/Source/Runtime/Datasmith/CADKernel/Base/Private/Geo/Curves/BoundedCurve.h b/Engine/Source/Runtime/Datasmith/CADKernel/Base/Private/Geo/Curves/BoundedCurve.h
--- a/Engine/Source/Runtime/Datasmith/CADKernel/Base/Private/Geo/Curves/BoundedCurve.h
+++ b/Engine/Source/Runtime/Datasmith/CADKernel/Base/Private/Geo/Curves/BoundedCurve.h
@@ -61,7 +61,7 @@ public:
 	}
 
 	virtual TSharedPtr<FEntityGeom> ApplyMatrix(const FMatrixH& InMatrix) const override;
-	virtual void Offset(const FPoint& OffsetDirection) override;
+	virtual void Offset(const FVector& OffsetDirection) override;
 
 	virtual void EvaluatePoint(double Coordinate, FCurvePoint& OutPoint, int32 DerivativeOrder = 0) const override;
 	virtual void Evaluate2DPoint(double Coordinate, FCurvePoint2D& OutPoint, int32 DerivativeOrder = 0) const override;
diff --git a/Engine/Source/Runtime/Datasmith/CADKernel/Base/Private/Geo/Curves/CompositeCurve.cpp b/Engine/Source/Runtime/Datasmith/CADKernel/Base/Private/Geo/Curves/CompositeCurve.cpp
--- a/Engine/Source/Runtime/Datasmith/CADKernel/Base/Private/Geo/Curves/CompositeCurve.cpp
+++ b/Engine/Source/Runtime/Datasmith/CADKernel/Base/Private/Geo/Curves/CompositeCurve.cpp
@@ -183,7 +183,7 @@ TSharedPtr<FEntityGeom> FCompositeCurve::ApplyMatrix(const FMatrixH& InMatrix) c
 	return FEntity::MakeShared<FCompositeCurve>(TransformedCurves);
 }
 
-void FCompositeCurve::Offset(const FPoint& OffsetDirection)
+void FCompositeCurve::Offset(const FVector& OffsetDirection)
 {
 	for (FOrientedCurve& Curve : Curves)
 	{
@@ -223,8 +223,8 @@ FCompositeCurve::FCompositeCurve(const TArray<TSharedPtr<FCurve>>& CurveList, bo
 	double Perimeter = 0.0;
 	Coordinates.Add(Perimeter);
 
-	FPoint PreviousPoint1;
-	FPoint PreviousPoint2;
+	FVector PreviousPoint1;
+	FVector PreviousPoint2;
 
 	for (int32 ICurve = 0; ICurve < CurveList.Num(); ICurve++)
 	{
@@ -243,8 +243,8 @@ FCompositeCurve::FCompositeCurve(const TArray<TSharedPtr<FCurve>>& CurveList, bo
 		{
 			if (ICurve > 0)
 			{
-				FPoint Point1;
-				FPoint Point2;
+				FVector Point1;
+				FVector Point2;
 
 				Curve->GetExtremities(&Point1, &Point2);
 				if (Point1.Distance(PreviousPoint2) > Point2.Distance(PreviousPoint2))
diff --git a/Engine/Source/Runtime/Datasmith/CADKernel/Base/Private/Geo/Curves/CompositeCurve.h b/Engine/Source/Runtime/Datasmith/CADKernel/Base/Private/Geo/Curves/CompositeCurve.h
--- a/Engine/Source/Runtime/Datasmith/CADKernel/Base/Private/Geo/Curves/CompositeCurve.h
+++ b/Engine/Source/Runtime/Datasmith/CADKernel/Base/Private/Geo/Curves/CompositeCurve.h
@@ -69,7 +69,7 @@ public:
 	}
 
 	virtual TSharedPtr<FEntityGeom> ApplyMatrix(const FMatrixH& InMatrix) const override;
-	virtual void Offset(const FPoint& OffsetDirection) override;
+	virtual void Offset(const FVector& OffsetDirection) override;
 
 	virtual TSharedPtr<FCurve> GetCurve(int32 Index) const
 	{
diff --git a/Engine/Source/Runtime/Datasmith/CADKernel/Base/Private/Geo/Curves/Curve.cpp b/Engine/Source/Runtime/Datasmith/CADKernel/Base/Private/Geo/Curves/Curve.cpp
--- a/Engine/Source/Runtime/Datasmith/CADKernel/Base/Private/Geo/Curves/Curve.cpp
+++ b/Engine/Source/Runtime/Datasmith/CADKernel/Base/Private/Geo/Curves/Curve.cpp
@@ -41,7 +41,7 @@ void FCurve::EvaluatePoints(const TArray<double>& Coordinates, TArray<FCurvePoin
 	}
 }
 
-void FCurve::EvaluatePoints(const TArray<double>& Coordinates, TArray<FPoint>& OutPoints) const
+void FCurve::EvaluatePoints(const TArray<double>& Coordinates, TArray<FVector>& OutPoints) const
 {
 	OutPoints.Empty(Coordinates.Num());
 	for (double Coordinate : Coordinates)
@@ -65,7 +65,7 @@ void FCurve::Evaluate2DPoints(const TArray<double>& Coordinates, TArray<FCurvePo
  * Evaluate exact 2D points of the curve at the input Coordinates
  * The function can only be used with 2D curve (Dimension == 2)
  */
-void FCurve::Evaluate2DPoints(const TArray<double>& Coordinates, TArray<FPoint2D>& OutPoints) const
+void FCurve::Evaluate2DPoints(const TArray<double>& Coordinates, TArray<FVector2d>& OutPoints) const
 {
 	OutPoints.SetNum(Coordinates.Num());
 	for (int32 iPoint = 0; iPoint < Coordinates.Num(); iPoint++)
diff --git a/Engine/Source/Runtime/Datasmith/CADKernel/Base/Private/Geo/Curves/CurveMaker.cpp b/Engine/Source/Runtime/Datasmith/CADKernel/Base/Private/Geo/Curves/CurveMaker.cpp
--- a/Engine/Source/Runtime/Datasmith/CADKernel/Base/Private/Geo/Curves/CurveMaker.cpp
+++ b/Engine/Source/Runtime/Datasmith/CADKernel/Base/Private/Geo/Curves/CurveMaker.cpp
@@ -43,7 +43,7 @@ TSharedPtr<FCurve> FCurve::MakeNurbsCurve(FNurbsCurveData& InNurbsData)
 
 		if (InNurbsData.Dimension == 2)
 		{
-			TArray<FPoint2D> NewPoles;
+			TArray<FVector2d> NewPoles;
 			NewPoles.Reserve(InNurbsData.Poles.Num());
 
 			if (InNurbsData.bIsRational && InNurbsData.Poles.Num() == InNurbsData.Weights.Num())
@@ -56,7 +56,7 @@ TSharedPtr<FCurve> FCurve::MakeNurbsCurve(FNurbsCurveData& InNurbsData)
 			}
 			else
 			{
-				for (const FPoint& Pole : InNurbsData.Poles)
+				for (const FVector& Pole : InNurbsData.Poles)
 				{
 					NewPoles.Emplace(Pole.X, Pole.Y);
 				}
@@ -67,7 +67,7 @@ TSharedPtr<FCurve> FCurve::MakeNurbsCurve(FNurbsCurveData& InNurbsData)
 
 		if (InNurbsData.bIsRational && InNurbsData.Poles.Num() == InNurbsData.Weights.Num())
 		{
-			TArray<FPoint> NewPoles(InNurbsData.Poles);
+			TArray<FVector> NewPoles(InNurbsData.Poles);
 			for (int32 Index = 0; Index < NewPoles.Num(); ++Index)
 			{
 				NewPoles[Index] /= InNurbsData.Weights[Index];
@@ -82,22 +82,22 @@ TSharedPtr<FCurve> FCurve::MakeNurbsCurve(FNurbsCurveData& InNurbsData)
 	return FEntity::MakeShared<FNURBSCurve>(InNurbsData);
 }
 
-TSharedPtr<FCurve> FCurve::MakeBezierCurve(const TArray<FPoint>& InPoles)
+TSharedPtr<FCurve> FCurve::MakeBezierCurve(const TArray<FVector>& InPoles)
 {
 	return FEntity::MakeShared<UE::CADKernel::FBezierCurve>(InPoles);
 }
 
-TSharedPtr<FCurve> FCurve::MakeSplineCurve(const TArray<FPoint>& InPoles)
+TSharedPtr<FCurve> FCurve::MakeSplineCurve(const TArray<FVector>& InPoles)
 {
 	return FEntity::MakeShared<UE::CADKernel::FSplineCurve>(InPoles);
 }
 
-TSharedPtr<FCurve> FCurve::MakeSplineCurve(const TArray<FPoint>& InPoles, const TArray<FPoint>& InTangents)
+TSharedPtr<FCurve> FCurve::MakeSplineCurve(const TArray<FVector>& InPoles, const TArray<FVector>& InTangents)
 {
 	return FEntity::MakeShared<UE::CADKernel::FSplineCurve>(InPoles, InTangents);
 }
 
-TSharedPtr<FCurve> FCurve::MakeSplineCurve(const TArray<FPoint>& InPoles, const TArray<FPoint>& InArriveTangents, const TArray<FPoint>& InLeaveTangents)
+TSharedPtr<FCurve> FCurve::MakeSplineCurve(const TArray<FVector>& InPoles, const TArray<FVector>& InArriveTangents, const TArray<FVector>& InLeaveTangents)
 {
 	return FEntity::MakeShared<UE::CADKernel::FSplineCurve>(InPoles, InArriveTangents, InLeaveTangents);
 }
diff --git a/Engine/Source/Runtime/Datasmith/CADKernel/Base/Private/Geo/Curves/CurveUtilities.cpp b/Engine/Source/Runtime/Datasmith/CADKernel/Base/Private/Geo/Curves/CurveUtilities.cpp
--- a/Engine/Source/Runtime/Datasmith/CADKernel/Base/Private/Geo/Curves/CurveUtilities.cpp
+++ b/Engine/Source/Runtime/Datasmith/CADKernel/Base/Private/Geo/Curves/CurveUtilities.cpp
@@ -15,7 +15,7 @@ namespace UE::CADKernel::CurveUtilities
 	template<class T>
 	bool GetCurvePoles(const T& Curve, TArray<FVector>& PolesOut)
 	{
-		const TArray<FPoint>& Poles = Curve.GetPoles();
+		const TArray<FVector>& Poles = Curve.GetPoles();
 		const int32 PoleCount = Poles.Num();
 
 		PolesOut.Reserve(PoleCount);
@@ -24,7 +24,7 @@ namespace UE::CADKernel::CurveUtilities
 		{
 			for (int32 Index = 0; Index < PoleCount; ++Index)
 			{
-				const FPoint& Pole = Poles[Index];
+				const FVector& Pole = Poles[Index];
 				PolesOut.Add(FVector(Pole.X, Pole.Y, 0.f));
 			}
 		}
@@ -32,7 +32,7 @@ namespace UE::CADKernel::CurveUtilities
 		{
 			for (int32 Index = 0; Index < PoleCount; ++Index)
 			{
-				const FPoint& Pole = Poles[Index];
+				const FVector& Pole = Poles[Index];
 				PolesOut.Add(FVector(Pole.X, Pole.Y, Pole.Z));
 			}
 		}
@@ -98,11 +98,11 @@ namespace UE::CADKernel::CurveUtilities
 			const FSegmentCurve& Segment = static_cast<const FSegmentCurve&>(Curve);
 			Poles.Reserve(2);
 			{
-				const FPoint& Point = Segment.GetStartPoint();
+				const FVector& Point = Segment.GetStartPoint();
 				Poles.Emplace(Point.X, Point.Y, Point.Z);
 			}
 			{
-				const FPoint& Point = Segment.GetEndPoint();
+				const FVector& Point = Segment.GetEndPoint();
 				Poles.Emplace(Point.X, Point.Y, Point.Z);
 			}
 		}
@@ -110,12 +110,12 @@ namespace UE::CADKernel::CurveUtilities
 
 		case ECurve::Polyline3D:
 		{
-			GetPolylinePoints<FPoint, FCurvePoint>(static_cast<const TPolylineCurve<FPoint, FCurvePoint>&>(Curve), Poles);
+			GetPolylinePoints<FVector, FCurvePoint>(static_cast<const TPolylineCurve<FVector, FCurvePoint>&>(Curve), Poles);
 		}
 
 		case ECurve::Polyline2D:
 		{
-			GetPolylinePoints<FPoint2D, FCurvePoint2D>(static_cast<const TPolylineCurve<FPoint2D, FCurvePoint2D>&>(Curve), Poles);
+			GetPolylinePoints<FVector2d, FCurvePoint2D>(static_cast<const TPolylineCurve<FVector2d, FCurvePoint2D>&>(Curve), Poles);
 		}
 		break;
 
@@ -148,12 +148,12 @@ namespace UE::CADKernel::CurveUtilities
 
 		const FSurfacicPolyline& Polyline = Curve.GetPolyline();
 
-		const TArray<FPoint2D>& Point2Ds = Polyline.Get2DPoints();
+		const TArray<FVector2d>& Point2Ds = Polyline.Get2DPoints();
 
 		TArray<FVector2d> SurfacicPoints;
 		SurfacicPoints.Reserve(Point2Ds.Num());
 
-		for (const FPoint2D& Point2D : Point2Ds)
+		for (const FVector2d& Point2D : Point2Ds)
 		{
 			SurfacicPoints.Emplace(Point2D.X, Point2D.Y);
 		}
@@ -167,12 +167,12 @@ namespace UE::CADKernel::CurveUtilities
 
 		const FSurfacicPolyline& Polyline = Curve.GetPolyline();
 
-		const TArray<FPoint>& Point3Ds = Polyline.GetPoints();
+		const TArray<FVector>& Point3Ds = Polyline.GetPoints();
 
 		TArray<FVector> Points;
 		Points.Reserve(Point3Ds.Num());
 
-		for (const FPoint& Point : Point3Ds)
+		for (const FVector& Point : Point3Ds)
 		{
 			Points.Emplace(Point.X, Point.Y, Point.Z);
 		}
@@ -209,12 +209,12 @@ namespace UE::CADKernel::CurveUtilities
 
 		case ECurve::Polyline3D:
 		{
-			return static_cast<const TPolylineCurve<FPoint, FCurvePoint>&>(Curve).GetPolylinePoints().Num();
+			return static_cast<const TPolylineCurve<FVector, FCurvePoint>&>(Curve).GetPolylinePoints().Num();
 		}
 
 		case ECurve::Polyline2D:
 		{
-			return static_cast<const TPolylineCurve<FPoint2D, FCurvePoint2D>&>(Curve).GetPolylinePoints().Num();
+			return static_cast<const TPolylineCurve<FVector2d, FCurvePoint2D>&>(Curve).GetPolylinePoints().Num();
 		}
 
 		default:
diff --git a/Engine/Source/Runtime/Datasmith/CADKernel/Base/Private/Geo/Curves/EllipseCurve.cpp b/Engine/Source/Runtime/Datasmith/CADKernel/Base/Private/Geo/Curves/EllipseCurve.cpp
--- a/Engine/Source/Runtime/Datasmith/CADKernel/Base/Private/Geo/Curves/EllipseCurve.cpp
+++ b/Engine/Source/Runtime/Datasmith/CADKernel/Base/Private/Geo/Curves/EllipseCurve.cpp
@@ -10,7 +10,7 @@ TSharedPtr<FEntityGeom> FEllipseCurve::ApplyMatrix(const FMatrixH& InMatrix) con
 	return FEntity::MakeShared<FEllipseCurve>(NewMatrix, RadiusU, RadiusV, Boundary);
 }
 
-void FEllipseCurve::Offset(const FPoint& OffsetDirection)
+void FEllipseCurve::Offset(const FVector& OffsetDirection)
 {
 	FMatrixH Offset = FMatrixH::MakeTranslationMatrix(OffsetDirection);
 	Matrix *= Offset;
diff --git a/Engine/Source/Runtime/Datasmith/CADKernel/Base/Private/Geo/Curves/EllipseCurve.h b/Engine/Source/Runtime/Datasmith/CADKernel/Base/Private/Geo/Curves/EllipseCurve.h
--- a/Engine/Source/Runtime/Datasmith/CADKernel/Base/Private/Geo/Curves/EllipseCurve.h
+++ b/Engine/Source/Runtime/Datasmith/CADKernel/Base/Private/Geo/Curves/EllipseCurve.h
@@ -65,45 +65,49 @@ public:
 	}
 
 	virtual TSharedPtr<FEntityGeom> ApplyMatrix(const FMatrixH& InMatrix) const override;
-	virtual void Offset(const FPoint& OffsetDirection) override;
+	virtual void Offset(const FVector& OffsetDirection) override;
 
 	virtual void EvaluatePoint(double Coordinate, FCurvePoint& OutPoint, int32 DerivativeOrder = 0) const override
 	{
 		ensureCADKernel(Dimension == 3);
-		Evaluate<FCurvePoint, FPoint>(Coordinate, OutPoint, DerivativeOrder);
+
+		OutPoint.DerivativeOrder = DerivativeOrder;
+
+		const double CosU = cos(Coordinate);
+		const double SinU = sin(Coordinate);
+
+		OutPoint.Point = Matrix.Multiply(FVector(RadiusU * CosU, RadiusV * SinU, 0.));
+
+		if (DerivativeOrder > 0)
+		{
+			OutPoint.Gradient = Matrix.MultiplyVector(FVector(RadiusU * (-SinU), RadiusV * CosU, 0.));
+
+			if (DerivativeOrder > 1)
+			{
+				OutPoint.Laplacian = Matrix.MultiplyVector(FVector(RadiusU * (-CosU), RadiusV * (-SinU), 0.));
+			}
+		}
 	}
 
 	virtual void Evaluate2DPoint(double Coordinate, FCurvePoint2D& OutPoint, int32 DerivativeOrder = 0) const override
 	{
 		ensureCADKernel(Dimension == 2);
-		Evaluate<FCurvePoint2D, FPoint2D>(Coordinate, OutPoint, DerivativeOrder);
-	}
 
-private:
-
-	template <typename CurvePointType, typename PointType>
-	void Evaluate(double Coordinate, CurvePointType& OutPoint, int32 DerivativeOrder) const
-	{
 		OutPoint.DerivativeOrder = DerivativeOrder;
 
-		double CosU = cos(Coordinate);
-		double SinU = sin(Coordinate);
+		const double CosU = cos(Coordinate);
+		const double SinU = sin(Coordinate);
 
-		{
-			PointType Result(RadiusU * CosU, RadiusV * SinU);
-			OutPoint.Point = Matrix.Multiply(Result);
-		}
+		OutPoint.Point = Matrix.Multiply2D(FVector2d(RadiusU * CosU, RadiusV * SinU));
 
 		if (DerivativeOrder > 0)
 		{
-			PointType Result(RadiusU * (-SinU), RadiusV * CosU);
-			OutPoint.Gradient = Matrix.MultiplyVector(Result);
-		}
+			OutPoint.Gradient = Matrix.MultiplyVector2D(FVector2d(RadiusU * (-SinU), RadiusV * CosU));
 
-		if (DerivativeOrder > 1)
-		{
-			PointType Result(RadiusU * (-CosU), RadiusV * (-SinU));
-			OutPoint.Laplacian = Matrix.MultiplyVector(Result);
+			if (DerivativeOrder > 1)
+			{
+				OutPoint.Laplacian = Matrix.MultiplyVector2D(FVector2d(RadiusU * (-CosU), RadiusV * (-SinU)));
+			}
 		}
 	}
 };
diff --git a/Engine/Source/Runtime/Datasmith/CADKernel/Base/Private/Geo/Curves/HyperbolaCurve.cpp b/Engine/Source/Runtime/Datasmith/CADKernel/Base/Private/Geo/Curves/HyperbolaCurve.cpp
--- a/Engine/Source/Runtime/Datasmith/CADKernel/Base/Private/Geo/Curves/HyperbolaCurve.cpp
+++ b/Engine/Source/Runtime/Datasmith/CADKernel/Base/Private/Geo/Curves/HyperbolaCurve.cpp
@@ -10,7 +10,7 @@ TSharedPtr<FEntityGeom> FHyperbolaCurve::ApplyMatrix(const FMatrixH& InMatrix) c
 	return FEntity::MakeShared<FHyperbolaCurve>(NewMatrix, SemiMajorAxis, SemiImaginaryAxis, Boundary);
 }
 
-void FHyperbolaCurve::Offset(const FPoint& OffsetDirection)
+void FHyperbolaCurve::Offset(const FVector& OffsetDirection)
 {
 	FMatrixH Offset = FMatrixH::MakeTranslationMatrix(OffsetDirection);
 	Matrix *= Offset;
diff --git a/Engine/Source/Runtime/Datasmith/CADKernel/Base/Private/Geo/Curves/HyperbolaCurve.h b/Engine/Source/Runtime/Datasmith/CADKernel/Base/Private/Geo/Curves/HyperbolaCurve.h
--- a/Engine/Source/Runtime/Datasmith/CADKernel/Base/Private/Geo/Curves/HyperbolaCurve.h
+++ b/Engine/Source/Runtime/Datasmith/CADKernel/Base/Private/Geo/Curves/HyperbolaCurve.h
@@ -52,45 +52,50 @@ public:
 	}
 
 	virtual TSharedPtr<FEntityGeom> ApplyMatrix(const FMatrixH& InMatrix) const override;
-	virtual void Offset(const FPoint& OffsetDirection) override;
+	virtual void Offset(const FVector& OffsetDirection) override;
+
 
 	virtual void EvaluatePoint(double Coordinate, FCurvePoint& OutPoint, int32 DerivativeOrder = 0) const override
 	{
 		ensureCADKernel(Dimension == 3);
-		Evaluate<FCurvePoint, FPoint>(Coordinate, OutPoint, DerivativeOrder);
+
+		OutPoint.DerivativeOrder = DerivativeOrder;
+
+		const double CosHUCoord = cosh(Coordinate);
+		const double SinHUCoord = sinh(Coordinate);
+
+		OutPoint.Point = Matrix.Multiply(FVector(SemiMajorAxis * CosHUCoord, SemiImaginaryAxis * SinHUCoord, 0.));
+
+		if (DerivativeOrder > 0)
+		{
+			OutPoint.Gradient = Matrix.MultiplyVector(FVector(SemiMajorAxis * SinHUCoord, SemiImaginaryAxis * CosHUCoord, 0.));
+
+			if (DerivativeOrder > 1)
+			{
+				OutPoint.Laplacian = Matrix.MultiplyVector(FVector(SemiMajorAxis * CosHUCoord, SemiImaginaryAxis * SinHUCoord, 0.));
+			}
+		}
 	}
 
 	virtual void Evaluate2DPoint(double Coordinate, FCurvePoint2D& OutPoint, int32 DerivativeOrder = 0) const override
 	{
 		ensureCADKernel(Dimension == 2);
-		Evaluate<FCurvePoint2D, FPoint2D>(Coordinate, OutPoint, DerivativeOrder);
-	}
 
-private:
-
-	template <typename CurvePointType, typename PointType>
-	void Evaluate(double Coordinate, CurvePointType& OutPoint, int32 DerivativeOrder) const
-	{
 		OutPoint.DerivativeOrder = DerivativeOrder;
 
 		const double CosHUCoord = cosh(Coordinate);
 		const double SinHUCoord = sinh(Coordinate);
 
-		{
-			PointType Result(SemiMajorAxis * CosHUCoord, SemiImaginaryAxis * SinHUCoord);
-			OutPoint.Point = Matrix.Multiply(Result);
-		}
+		OutPoint.Point = Matrix.Multiply2D(FVector2d(SemiMajorAxis * CosHUCoord, SemiImaginaryAxis * SinHUCoord));
 
 		if (DerivativeOrder > 0)
 		{
-			FPoint Result(SemiMajorAxis * SinHUCoord, SemiImaginaryAxis * CosHUCoord);
-			OutPoint.Gradient = Matrix.MultiplyVector(Result);
-		}
+			OutPoint.Gradient = Matrix.MultiplyVector2D(FVector2d(SemiMajorAxis * SinHUCoord, SemiImaginaryAxis * CosHUCoord));
 
-		if (DerivativeOrder > 1)
-		{
-			FPoint Result(SemiMajorAxis * CosHUCoord, SemiImaginaryAxis * SinHUCoord);
-			OutPoint.Laplacian = Matrix.MultiplyVector(Result);
+			if (DerivativeOrder > 1)
+			{
+				OutPoint.Laplacian = Matrix.MultiplyVector2D(FVector2d(SemiMajorAxis * CosHUCoord, SemiImaginaryAxis * SinHUCoord));
+			}
 		}
 	}
 };
diff --git a/Engine/Source/Runtime/Datasmith/CADKernel/Base/Private/Geo/Curves/NURBSCurve.cpp b/Engine/Source/Runtime/Datasmith/CADKernel/Base/Private/Geo/Curves/NURBSCurve.cpp
--- a/Engine/Source/Runtime/Datasmith/CADKernel/Base/Private/Geo/Curves/NURBSCurve.cpp
+++ b/Engine/Source/Runtime/Datasmith/CADKernel/Base/Private/Geo/Curves/NURBSCurve.cpp
@@ -13,7 +13,7 @@
 namespace UE::CADKernel
 { 
 
-FNURBSCurve::FNURBSCurve(int32 InDegre, const TArray<double>& InNodalVector, const TArray<FPoint>& InPoles, const TArray<double>& InWeights, int8 InDimension)
+FNURBSCurve::FNURBSCurve(int32 InDegre, const TArray<double>& InNodalVector, const TArray<FVector>& InPoles, const TArray<double>& InWeights, int8 InDimension)
 	: FCurve(InDimension)
 	, Degree(InDegre)
 	, NodalVector(InNodalVector)
@@ -37,10 +37,10 @@ FNURBSCurve::FNURBSCurve(FNurbsCurveData& NurbsCurveData)
 
 TSharedPtr<FEntityGeom> FNURBSCurve::ApplyMatrix(const FMatrixH& InMatrix) const
 {
-	TArray<FPoint> TransformedPoles;
+	TArray<FVector> TransformedPoles;
 
 	TransformedPoles.Reserve(Poles.Num());
-	for (const FPoint& Pole : Poles) 
+	for (const FVector& Pole : Poles) 
 	{
 		TransformedPoles.Emplace(InMatrix.Multiply(Pole));
 	}
@@ -48,9 +48,9 @@ TSharedPtr<FEntityGeom> FNURBSCurve::ApplyMatrix(const FMatrixH& InMatrix) const
 	return FEntity::MakeShared<FNURBSCurve>(Degree, NodalVector, TransformedPoles, Weights, Dimension);
 }
 
-void FNURBSCurve::Offset(const FPoint& OffsetDirection)
+void FNURBSCurve::Offset(const FVector& OffsetDirection)
 {
-	for (FPoint& Pole : Poles)
+	for (FVector& Pole : Poles)
 	{
 		Pole += OffsetDirection;
 	}
@@ -88,7 +88,7 @@ void FNURBSCurve::Finalize()
 		{
 			if (!FMath::IsNearlyEqual(1., FirstWeigth))
 			{
-				for (FPoint& Pole : Poles)
+				for (FVector& Pole : Poles)
 				{
 					Pole /= FirstWeigth;
 				}
@@ -127,7 +127,7 @@ void FNURBSCurve::Finalize()
 		int32 Jndex = 0;
 		if (Dimension == 2)
 		{
-			for (FPoint& Pole : Poles)
+			for (FVector& Pole : Poles)
 			{
 				HomogeneousPoles[Jndex++] = Pole.X;
 				HomogeneousPoles[Jndex++] = Pole.Y;
@@ -135,14 +135,14 @@ void FNURBSCurve::Finalize()
 		}
 		else
 		{
-			memcpy(HomogeneousPoles.GetData(), Poles.GetData(), sizeof(FPoint) * Poles.Num());
+			memcpy(HomogeneousPoles.GetData(), Poles.GetData(), sizeof(FVector) * Poles.Num());
 		}
 	}
 
 	Boundary.Set(NodalVector[Degree], NodalVector[NodalVector.Num() - 1 - Degree]);
 }
 
-void FNURBSCurve::ExtendTo(const FPoint& Point)
+void FNURBSCurve::ExtendTo(const FVector& Point)
 {
 	PolylineTools::ExtendTo(Poles, Point);
 	Finalize();
@@ -163,7 +163,7 @@ void FNURBSCurve::SetStartNodalCoordinate(double NewStartBoundary)
 void FNURBSCurve::Invert()
 {
 	{
-		TArray<FPoint> NewPoles;
+		TArray<FVector> NewPoles;
 		NewPoles.Reserve(Poles.Num());
 		for (int32 Index = Poles.Num() - 1; Index >= 0; --Index)
 		{
diff --git a/Engine/Source/Runtime/Datasmith/CADKernel/Base/Private/Geo/Curves/NURBSCurve.h b/Engine/Source/Runtime/Datasmith/CADKernel/Base/Private/Geo/Curves/NURBSCurve.h
--- a/Engine/Source/Runtime/Datasmith/CADKernel/Base/Private/Geo/Curves/NURBSCurve.h
+++ b/Engine/Source/Runtime/Datasmith/CADKernel/Base/Private/Geo/Curves/NURBSCurve.h
@@ -15,7 +15,7 @@ struct FNurbsCurveData
 	TArray<double> NodalVector;
 
 	TArray<double> Weights;
-	TArray<FPoint> Poles;
+	TArray<FVector> Poles;
 };
 
 class CADKERNEL_API FNURBSCurve : public FCurve
@@ -29,7 +29,7 @@ protected:
 	TArray<double> NodalVector;
 	TArray<double> Weights;
 
-	TArray<FPoint> Poles;
+	TArray<FVector> Poles;
 
 	bool bIsRational;
 
@@ -39,7 +39,7 @@ protected:
 	TArray<double> HomogeneousPoles;
 	int32 PoleDimension;
 
-	FNURBSCurve(int32 InDegree, const TArray<double>& InNodalVector, const TArray<FPoint>& InPoles, const TArray<double>& InWeights, int8 InDimension = 3);
+	FNURBSCurve(int32 InDegree, const TArray<double>& InNodalVector, const TArray<FVector>& InPoles, const TArray<double>& InWeights, int8 InDimension = 3);
 	FNURBSCurve(FNurbsCurveData& NurbsCurveData);
 
 	FNURBSCurve() = default;
@@ -80,7 +80,7 @@ public:
 		return Poles.Num();
 	}
 
-	const TArray<FPoint>& GetPoles() const
+	const TArray<FVector>& GetPoles() const
 	{
 		return Poles;
 	}
@@ -106,7 +106,7 @@ public:
 	}
 
 	virtual TSharedPtr<FEntityGeom> ApplyMatrix(const FMatrixH& InMatrix) const override;
-	virtual void Offset(const FPoint& OffsetDirection) override;
+	virtual void Offset(const FVector& OffsetDirection) override;
 
 	virtual void EvaluatePoint(double Coordinate, FCurvePoint& OutPoint, int32 DerivativeOrder = 0) const override
 	{
@@ -123,7 +123,7 @@ public:
 		BSpline::FindNotDerivableParameters(*this, DerivativeOrder, InBoundary, OutNotDerivableCoordinates);
 	}
 
-	virtual void ExtendTo(const FPoint& DesiredPosition) override;
+	virtual void ExtendTo(const FVector& DesiredPosition) override;
 
 	void Invert();
 	void SetStartNodalCoordinate(double NewStartBoundary);
diff --git a/Engine/Source/Runtime/Datasmith/CADKernel/Base/Private/Geo/Curves/ParabolaCurve.cpp b/Engine/Source/Runtime/Datasmith/CADKernel/Base/Private/Geo/Curves/ParabolaCurve.cpp
--- a/Engine/Source/Runtime/Datasmith/CADKernel/Base/Private/Geo/Curves/ParabolaCurve.cpp
+++ b/Engine/Source/Runtime/Datasmith/CADKernel/Base/Private/Geo/Curves/ParabolaCurve.cpp
@@ -13,7 +13,7 @@ TSharedPtr<FEntityGeom> FParabolaCurve::ApplyMatrix(const FMatrixH& InMatrix) co
 	return FEntity::MakeShared<FParabolaCurve>(NewMatrix, FocalDistance, Boundary, Dimension);
 }
 
-void FParabolaCurve::Offset(const FPoint& OffsetDirection)
+void FParabolaCurve::Offset(const FVector& OffsetDirection)
 {
 	FMatrixH Offset = FMatrixH::MakeTranslationMatrix(OffsetDirection);
 	Matrix *= Offset;
diff --git a/Engine/Source/Runtime/Datasmith/CADKernel/Base/Private/Geo/Curves/ParabolaCurve.h b/Engine/Source/Runtime/Datasmith/CADKernel/Base/Private/Geo/Curves/ParabolaCurve.h
--- a/Engine/Source/Runtime/Datasmith/CADKernel/Base/Private/Geo/Curves/ParabolaCurve.h
+++ b/Engine/Source/Runtime/Datasmith/CADKernel/Base/Private/Geo/Curves/ParabolaCurve.h
@@ -43,9 +43,9 @@ public:
 	}
 
 	virtual TSharedPtr<FEntityGeom> ApplyMatrix(const FMatrixH& InMatrix) const override;
-	virtual void Offset(const FPoint& OffsetDirection) override;
+	virtual void Offset(const FVector& OffsetDirection) override;
 
-	virtual void ExtendTo(const FPoint& Point) override
+	virtual void ExtendTo(const FVector& Point) override
 	{
 		ensureCADKernel(false);
 	}
@@ -53,37 +53,38 @@ public:
 	virtual void EvaluatePoint(double Coordinate, FCurvePoint& OutPoint, int32 DerivativeOrder = 0) const override
 	{
 		ensureCADKernel(Dimension == 3);
-		Evaluate<FCurvePoint, FPoint>(Coordinate, OutPoint, DerivativeOrder);
+
+		OutPoint.DerivativeOrder = DerivativeOrder;
+
+		OutPoint.Point = Matrix.Multiply(FVector(Coordinate * Coordinate * FocalDistance, Coordinate, 0.));
+
+		if (DerivativeOrder > 0)
+		{
+			OutPoint.Gradient = Matrix.MultiplyVector(FVector(2.0 * Coordinate * FocalDistance, 1., 0.));
+
+			if (DerivativeOrder > 1)
+			{
+				OutPoint.Laplacian = Matrix.MultiplyVector(FVector(2.0 * FocalDistance, 0., 0.));
+			}
+		}
 	}
 
 	virtual void Evaluate2DPoint(double Coordinate, FCurvePoint2D& OutPoint, int32 DerivativeOrder = 0) const override
 	{
 		ensureCADKernel(Dimension == 2);
-		Evaluate<FCurvePoint2D, FPoint2D>(Coordinate, OutPoint, DerivativeOrder);
-	}
-
-private:
 
-	template <typename CurvePointType, typename PointType>
-	void Evaluate(double Coordinate, CurvePointType& OutPoint, int32 DerivativeOrder) const
-	{
 		OutPoint.DerivativeOrder = DerivativeOrder;
 
-		{
-			PointType Result(Coordinate * Coordinate * FocalDistance, Coordinate);
-			OutPoint.Point = Matrix.Multiply(Result);
-		}
+		OutPoint.Point = Matrix.Multiply2D(FVector2d(Coordinate * Coordinate * FocalDistance, Coordinate));
 
 		if (DerivativeOrder > 0)
 		{
-			FPoint Result(2.0 * Coordinate * FocalDistance, 1);
-			OutPoint.Gradient = Matrix.MultiplyVector(Result);
-		}
+			OutPoint.Gradient = Matrix.MultiplyVector2D(FVector2d(2.0 * Coordinate * FocalDistance, 1.));
 
-		if (DerivativeOrder > 1)
-		{
-			FPoint Result(2.0 * FocalDistance, 0.);
-			OutPoint.Laplacian = Matrix.MultiplyVector(Result);
+			if (DerivativeOrder > 1)
+			{
+				OutPoint.Laplacian = Matrix.MultiplyVector2D(FVector2d(2.0 * FocalDistance, 0.));
+			}
 		}
 	}
 };
diff --git a/Engine/Source/Runtime/Datasmith/CADKernel/Base/Private/Geo/Curves/PolylineCurve.h b/Engine/Source/Runtime/Datasmith/CADKernel/Base/Private/Geo/Curves/PolylineCurve.h
--- a/Engine/Source/Runtime/Datasmith/CADKernel/Base/Private/Geo/Curves/PolylineCurve.h
+++ b/Engine/Source/Runtime/Datasmith/CADKernel/Base/Private/Geo/Curves/PolylineCurve.h
@@ -127,7 +127,7 @@ public:
 		return FEntity::MakeShared<PolylineType>(NewPoints, this->Coordinates);
 	}
 
-	virtual void Offset(const FPoint& OffsetDirection) override
+	virtual void Offset(const FVector& OffsetDirection) override
 	{
 		for (PointType& Pole : this->Points)
 		{
@@ -135,24 +135,24 @@ public:
 		}
 	}
 
-	virtual void ExtendTo(const FPoint& DesiredPoint) override
+	virtual void ExtendTo(const FVector& DesiredPoint) override
 	{
 		PolylineTools::ExtendTo(this->Points, (PointType)DesiredPoint);
 	}
 };
 
-class CADKERNEL_API FPolylineCurve : public TPolylineCurve<FPoint, FCurvePoint>
+class CADKERNEL_API FPolylineCurve : public TPolylineCurve<FVector, FCurvePoint>
 {
 	friend class FEntity;
 
 protected:
-	FPolylineCurve(const TArray<FPoint>& InPoints, const TArray<double>& InCoordinates)
-		: TPolylineCurve<FPoint, FCurvePoint>(InPoints, InCoordinates, 3)
+	FPolylineCurve(const TArray<FVector>& InPoints, const TArray<double>& InCoordinates)
+		: TPolylineCurve<FVector, FCurvePoint>(InPoints, InCoordinates, 3)
 	{
 	}
 
-	FPolylineCurve(const TArray<FPoint>& InPoints)
-		: TPolylineCurve<FPoint, FCurvePoint>(InPoints, 3)
+	FPolylineCurve(const TArray<FVector>& InPoints)
+		: TPolylineCurve<FVector, FCurvePoint>(InPoints, 3)
 	{
 	}
 
@@ -170,7 +170,7 @@ public:
 		return ApplyMatrixImpl<FPolylineCurve>(InMatrix);
 	}
 
-	virtual FPoint EvaluatePoint(double InCoordinate) const override
+	virtual FVector EvaluatePoint(double InCoordinate) const override
 	{
 		return EvaluatePointAt(InCoordinate);
 	}
@@ -181,18 +181,18 @@ public:
 
 };
 
-class CADKERNEL_API FPolyline2DCurve : public TPolylineCurve<FPoint2D, FCurvePoint2D>
+class CADKERNEL_API FPolyline2DCurve : public TPolylineCurve<FVector2d, FCurvePoint2D>
 {
 	friend class FEntity;
 
 protected:
-	FPolyline2DCurve(const TArray<FPoint2D>& InPoints, const TArray<double>& InCoordinates)
-		: TPolylineCurve<FPoint2D, FCurvePoint2D>(InPoints, InCoordinates, 2)
+	FPolyline2DCurve(const TArray<FVector2d>& InPoints, const TArray<double>& InCoordinates)
+		: TPolylineCurve<FVector2d, FCurvePoint2D>(InPoints, InCoordinates, 2)
 	{
 	}
 
-	FPolyline2DCurve(const TArray<FPoint2D>& InPoints)
-		: TPolylineCurve<FPoint2D, FCurvePoint2D>(InPoints, 2)
+	FPolyline2DCurve(const TArray<FVector2d>& InPoints)
+		: TPolylineCurve<FVector2d, FCurvePoint2D>(InPoints, 2)
 	{
 	}
 
@@ -215,7 +215,7 @@ public:
 		EvaluateCurvesPoint(InCoordinate, OutPoint, InDerivativeOrder);
 	}
 
-	virtual FPoint2D Evaluate2DPoint(double InCoordinate) const override
+	virtual FVector2d Evaluate2DPoint(double InCoordinate) const override
 	{
 		return EvaluatePointAt(InCoordinate);
 	}
diff --git a/Engine/Source/Runtime/Datasmith/CADKernel/Base/Private/Geo/Curves/RestrictionCurve.cpp b/Engine/Source/Runtime/Datasmith/CADKernel/Base/Private/Geo/Curves/RestrictionCurve.cpp
--- a/Engine/Source/Runtime/Datasmith/CADKernel/Base/Private/Geo/Curves/RestrictionCurve.cpp
+++ b/Engine/Source/Runtime/Datasmith/CADKernel/Base/Private/Geo/Curves/RestrictionCurve.cpp
@@ -5,15 +5,15 @@
 namespace UE::CADKernel
 {
 
-void FRestrictionCurve::ExtendTo(const FPoint2D& Point)
+void FRestrictionCurve::ExtendTo(const FVector2d& Point)
 {
 	Curve2D->ExtendTo(Point);
 	EvaluateSurfacicPolyline(Polyline);
 }
 
-void FRestrictionCurve::Offset2D(const FPoint2D& Offset)
+void FRestrictionCurve::Offset2D(const FVector2d& Offset)
 {
-	Curve2D->Offset(Offset);
+	Curve2D->Offset(FVector(Offset, 0.));
 	EvaluateSurfacicPolyline(Polyline);
 }
 
diff --git a/Engine/Source/Runtime/Datasmith/CADKernel/Base/Private/Geo/Curves/SegmentCurve.cpp b/Engine/Source/Runtime/Datasmith/CADKernel/Base/Private/Geo/Curves/SegmentCurve.cpp
--- a/Engine/Source/Runtime/Datasmith/CADKernel/Base/Private/Geo/Curves/SegmentCurve.cpp
+++ b/Engine/Source/Runtime/Datasmith/CADKernel/Base/Private/Geo/Curves/SegmentCurve.cpp
@@ -7,13 +7,13 @@ namespace UE::CADKernel
 
 TSharedPtr<FEntityGeom> FSegmentCurve::ApplyMatrix(const FMatrixH& InMatrix) const
 {
-	FPoint TransformedStartPoint = InMatrix.Multiply(StartPoint);
-	FPoint TransformedEndPoint = InMatrix.Multiply(EndPoint);
+	FVector TransformedStartPoint = InMatrix.Multiply(StartPoint);
+	FVector TransformedEndPoint = InMatrix.Multiply(EndPoint);
 
 	return FEntity::MakeShared<FSegmentCurve>(TransformedStartPoint, TransformedEndPoint, Dimension);
 }
 
-void FSegmentCurve::Offset(const FPoint& OffsetDirection)
+void FSegmentCurve::Offset(const FVector& OffsetDirection)
 {
 	StartPoint += OffsetDirection;
 	EndPoint += OffsetDirection;
diff --git a/Engine/Source/Runtime/Datasmith/CADKernel/Base/Private/Geo/Curves/SegmentCurve.h b/Engine/Source/Runtime/Datasmith/CADKernel/Base/Private/Geo/Curves/SegmentCurve.h
--- a/Engine/Source/Runtime/Datasmith/CADKernel/Base/Private/Geo/Curves/SegmentCurve.h
+++ b/Engine/Source/Runtime/Datasmith/CADKernel/Base/Private/Geo/Curves/SegmentCurve.h
@@ -12,16 +12,23 @@ class CADKERNEL_API FSegmentCurve : public FCurve
 	friend class FEntity;
 
 protected:
-	FPoint StartPoint;
-	FPoint EndPoint;
+	FVector StartPoint;
+	FVector EndPoint;
 
-	FSegmentCurve(const FPoint& InStartPoint, const FPoint& InEndPoint, int8 InDimension = 3)
+	FSegmentCurve(const FVector& InStartPoint, const FVector& InEndPoint, int8 InDimension = 3)
 		: FCurve(InDimension)
 		, StartPoint(InStartPoint)
 		, EndPoint(InEndPoint)
 	{
 	}
 
+	FSegmentCurve(const FVector2d& InStartPoint, const FVector2d& InEndPoint, int8 InDimension = 3)
+		: FCurve(InDimension)
+		, StartPoint(InStartPoint, 0.)
+		, EndPoint(InEndPoint, 0.)
+	{
+	}
+
 	FSegmentCurve() = default;
 
 public:
@@ -43,31 +50,31 @@ public:
 	}
 
 	virtual TSharedPtr<FEntityGeom> ApplyMatrix(const FMatrixH& InMatrix) const override;
-	virtual void Offset(const FPoint& OffsetDirection) override;
+	virtual void Offset(const FVector& OffsetDirection) override;
 
 	virtual void EvaluatePoint(double Coordinate, FCurvePoint& OutPoint, int32 DerivativeOrder = 0) const override
 	{
 		ensureCADKernel(Dimension == 3);
-		Evaluate<FCurvePoint, FPoint>(Coordinate, OutPoint, DerivativeOrder);
+		Evaluate<FCurvePoint, FVector>(Coordinate, OutPoint, DerivativeOrder);
 	}
 
 	virtual void Evaluate2DPoint(double Coordinate, FCurvePoint2D& OutPoint, int32 DerivativeOrder = 0) const override
 	{
 		ensureCADKernel(Dimension == 2);
-		Evaluate<FCurvePoint2D, FPoint2D>(Coordinate, OutPoint, DerivativeOrder);
+		Evaluate<FCurvePoint2D, FVector2d>(Coordinate, OutPoint, DerivativeOrder);
 	}
 
-	const FPoint& GetStartPoint() const
+	const FVector& GetStartPoint() const
 	{
 		return StartPoint;
 	}
 
-	const FPoint& GetEndPoint() const
+	const FVector& GetEndPoint() const
 	{
 		return EndPoint;
 	}
 
-	virtual void ExtendTo(const FPoint& DesiredPosition) override
+	virtual void ExtendTo(const FVector& DesiredPosition) override
 	{
 		double DistanceToStartPoint = FVector::DistSquared(DesiredPosition, StartPoint);
 		double DistanceToEndPoint = FVector::DistSquared(DesiredPosition, EndPoint);
@@ -82,14 +89,14 @@ public:
 	}
 
 private:
-	void GetTangent(FPoint& Tangent) const
+	void GetTangent(FVector& Tangent) const
 	{
 		Tangent = EndPoint - StartPoint;
 	}
 
-	void GetTangent(FPoint2D& Tangent) const
+	void GetTangent(FVector2d& Tangent) const
 	{
-		Tangent = FPoint2D(EndPoint.X - StartPoint.X, EndPoint.Y - StartPoint.Y);
+		Tangent = FVector2d(EndPoint.X - StartPoint.X, EndPoint.Y - StartPoint.Y);
 	}
 
 	template <typename CurvePointType, typename PointType>
diff --git a/Engine/Source/Runtime/Datasmith/CADKernel/Base/Private/Geo/Curves/SplineCurve.cpp b/Engine/Source/Runtime/Datasmith/CADKernel/Base/Private/Geo/Curves/SplineCurve.cpp
--- a/Engine/Source/Runtime/Datasmith/CADKernel/Base/Private/Geo/Curves/SplineCurve.cpp
+++ b/Engine/Source/Runtime/Datasmith/CADKernel/Base/Private/Geo/Curves/SplineCurve.cpp
@@ -11,10 +11,10 @@ namespace UE::CADKernel
 TSharedPtr<FEntityGeom> FSplineCurve::ApplyMatrix(const FMatrixH& InMatrix) const
 {
 	ensureCADKernel(false);
-	//TArray<FPoint> TransformedPoles;
+	//TArray<FVector> TransformedPoles;
 	//TransformedPoles.Reserve(Poles.Num());
 
-	//for (const FPoint& Pole : Poles)
+	//for (const FVector& Pole : Poles)
 	//{
 	//	TransformedPoles.Emplace(InMatrix.Multiply(Pole));
 	//}
@@ -22,10 +22,10 @@ TSharedPtr<FEntityGeom> FSplineCurve::ApplyMatrix(const FMatrixH& InMatrix) cons
 	return FEntity::MakeShared<FSplineCurve>(/*TransformedPoles*/);
 }
 
-void FSplineCurve::Offset(const FPoint& OffsetDirection)
+void FSplineCurve::Offset(const FVector& OffsetDirection)
 {
 	ensureCADKernel(false);
-	//for (FPoint& Pole : Poles)
+	//for (FVector& Pole : Poles)
 	//{
 	//	Pole += OffsetDirection;
 	//}
@@ -40,13 +40,13 @@ FInfoEntity& FSplineCurve::GetInfo(FInfoEntity& Info) const
 }
 #endif
 
-void FSplineCurve::ExtendTo(const FPoint& Point)
+void FSplineCurve::ExtendTo(const FVector& Point)
 {
 	ensureCADKernel(false);
 	//PolylineTools::ExtendTo(Poles, Point);
 }
 
-void FSplineCurve::SetSplinePoints(const TArray<FPoint>& Points)
+void FSplineCurve::SetSplinePoints(const TArray<FVector>& Points)
 {
 	const int32 NumPoints = Points.Num();
 	Position.Points.Reset(NumPoints);
@@ -54,9 +54,9 @@ void FSplineCurve::SetSplinePoints(const TArray<FPoint>& Points)
 	//Scale.Points.Reset(NumPoints);
 
 	double InputKey = 0.0;
-	for (const FPoint& Point : Points)
+	for (const FVector& Point : Points)
 	{
-		Position.Points.Emplace(InputKey, Point, FPoint::ZeroVector, FPoint::ZeroVector, CIM_CurveAuto);
+		Position.Points.Emplace(InputKey, Point, FVector::ZeroVector, FVector::ZeroVector, CIM_CurveAuto);
 		InputKey += 1.0;
 	}
 
@@ -66,7 +66,7 @@ void FSplineCurve::SetSplinePoints(const TArray<FPoint>& Points)
 	Boundary.SetMax(InputKey);
 }
 
-void FSplineCurve::SetSplinePoints(const TArray<FPoint>& Points, const TArray<FPoint>& Tangents)
+void FSplineCurve::SetSplinePoints(const TArray<FVector>& Points, const TArray<FVector>& Tangents)
 {
 	const int32 NumPoints = Points.Num();
 	ensureCADKernel(Tangents.Num() == NumPoints);
@@ -85,7 +85,7 @@ void FSplineCurve::SetSplinePoints(const TArray<FPoint>& Points, const TArray<FP
 	Boundary.SetMax(InputKey);
 }
 
-void FSplineCurve::SetSplinePoints(const TArray<FPoint>& Points, const TArray<FPoint>& ArriveTangents, const TArray<FPoint>& LeaveTangents)
+void FSplineCurve::SetSplinePoints(const TArray<FVector>& Points, const TArray<FVector>& ArriveTangents, const TArray<FVector>& LeaveTangents)
 {
 	const int32 NumPoints = Points.Num();
 	Position.Points.Reset(NumPoints);
@@ -158,19 +158,19 @@ void FSplineCurve::UpdateSpline(bool bInClosedLoop, bool bStationaryEndpoints, i
 	ReparamTable.Points.Emplace(AccumulatedLength, static_cast<double>(NumSegments), 0.0, 0.0, CIM_Linear);
 }
 
-FPoint FSplineCurve::GetLocationAtSplineInputKey(double Coordinate) const
+FVector FSplineCurve::GetLocationAtSplineInputKey(double Coordinate) const
 {
-	return Position.Eval((float)Coordinate, FPoint::ZeroVector);
+	return Position.Eval((float)Coordinate, FVector::ZeroVector);
 }
 
-FPoint FSplineCurve::GetTangentAtSplineInputKey(double Coordinate) const
+FVector FSplineCurve::GetTangentAtSplineInputKey(double Coordinate) const
 {
-	return Position.EvalDerivative((float)Coordinate, FPoint::ZeroVector);
+	return Position.EvalDerivative((float)Coordinate, FVector::ZeroVector);
 }
 
-FPoint FSplineCurve::GetDirectionAtSplineInputKey(double Coordinate) const
+FVector FSplineCurve::GetDirectionAtSplineInputKey(double Coordinate) const
 {
-	return Position.EvalDerivative(Coordinate, FPoint::ZeroVector).GetSafeNormal();
+	return Position.EvalDerivative(Coordinate, FVector::ZeroVector).GetSafeNormal();
 }
 
 void FSplineCurve::EvaluatePoint(double Coordinate, FCurvePoint& OutPoint, int32 DerivativeOrder) const
@@ -243,9 +243,9 @@ double FSplineCurve::GetSegmentLength(const int32 Index, const double Param, boo
 	}
 
 	// Cache the coefficients to be fed into the function to calculate the spline derivative at each sample point as they are constant.
-	const FPoint Coeff1 = ((P0 - P1) * 2.0 + T0 + T1) * 3.0;
-	const FPoint Coeff2 = (P1 - P0) * 6.0 - T0 * 4.0f - T1 * 2.0;
-	const FPoint Coeff3 = T0;
+	const FVector Coeff1 = ((P0 - P1) * 2.0 + T0 + T1) * 3.0;
+	const FVector Coeff2 = (P1 - P0) * 6.0 - T0 * 4.0f - T1 * 2.0;
+	const FVector Coeff3 = T0;
 
 	const double HalfParam = Param * 0.5;
 
@@ -254,7 +254,7 @@ double FSplineCurve::GetSegmentLength(const int32 Index, const double Param, boo
 	{
 		// Calculate derivative at each Legendre-Gauss sample, and perform a weighted sum
 		const double Alpha = HalfParam * (1.0 + LegendreGaussCoefficient.Abscissa);
-		const FPoint Derivative = ((Coeff1 * Alpha + Coeff2) * Alpha + Coeff3);
+		const FVector Derivative = ((Coeff1 * Alpha + Coeff2) * Alpha + Coeff3);
 		Length += Derivative.Length() * LegendreGaussCoefficient.Weight;
 	}
 	Length *= HalfParam;
diff --git a/Engine/Source/Runtime/Datasmith/CADKernel/Base/Private/Geo/Curves/SplineCurve.h b/Engine/Source/Runtime/Datasmith/CADKernel/Base/Private/Geo/Curves/SplineCurve.h
--- a/Engine/Source/Runtime/Datasmith/CADKernel/Base/Private/Geo/Curves/SplineCurve.h
+++ b/Engine/Source/Runtime/Datasmith/CADKernel/Base/Private/Geo/Curves/SplineCurve.h
@@ -7,7 +7,7 @@
 #include "Math/InterpCurve.h"
 
 DEFINE_INTERPCURVE_WRAPPER_STRUCT(FInterpCurveDouble, double)
-DEFINE_INTERPCURVE_WRAPPER_STRUCT(FInterpCurveFPoint, UE::CADKernel::FPoint)
+DEFINE_INTERPCURVE_WRAPPER_STRUCT(FInterpCurveFVector, FVector)
 
 namespace UE::CADKernel
 {
@@ -18,7 +18,7 @@ class CADKERNEL_API FSplineCurve : public FCurve
 
 protected:
 
-	FInterpCurveFPoint Position;
+	FInterpCurveFVector Position;
 
 	/** Spline built from rotation data. */
 	//FInterpCurveQuat Rotation;
@@ -38,17 +38,17 @@ protected:
 	//bool bLoopPositionOverride;
 	//double LoopPosition;
 
-	FSplineCurve(const TArray<FPoint>& InPoles)
+	FSplineCurve(const TArray<FVector>& InPoles)
 	{
 		SetSplinePoints(InPoles);
 	}
 
-	FSplineCurve(const TArray<FPoint>& InPoles, const TArray<FPoint>& InTangents)
+	FSplineCurve(const TArray<FVector>& InPoles, const TArray<FVector>& InTangents)
 	{
 		SetSplinePoints(InPoles, InTangents);
 	}
 
-	FSplineCurve(const TArray<FPoint>& InPoles, const TArray<FPoint>& InArriveTangents, const TArray<FPoint>& InLeaveTangents)
+	FSplineCurve(const TArray<FVector>& InPoles, const TArray<FVector>& InArriveTangents, const TArray<FVector>& InLeaveTangents)
 	{
 		SetSplinePoints(InPoles, InArriveTangents, InLeaveTangents);
 	}
@@ -56,13 +56,13 @@ protected:
 	FSplineCurve() = default;
 
 	/** Get location along spline at the provided input coordinate value */
-	FPoint GetLocationAtSplineInputKey(double InCoordinate) const;
+	FVector GetLocationAtSplineInputKey(double InCoordinate) const;
 
 	/** Get tangent along spline at the provided input coordinate value */
-	FPoint GetTangentAtSplineInputKey(double InCoordinate) const;
+	FVector GetTangentAtSplineInputKey(double InCoordinate) const;
 
 	/** Get unit direction (normalized tangent) along spline at the provided input coordinate value */
-	FPoint GetDirectionAtSplineInputKey(double InCoordinate) const;
+	FVector GetDirectionAtSplineInputKey(double InCoordinate) const;
 
 public:
 
@@ -84,11 +84,11 @@ public:
 	}
 
 	virtual TSharedPtr<FEntityGeom> ApplyMatrix(const FMatrixH& InMatrix) const override;
-	virtual void Offset(const FPoint& OffsetDirection) override;
+	virtual void Offset(const FVector& OffsetDirection) override;
 
 	virtual void EvaluatePoint(double Coordinate, FCurvePoint& OutPoint, int32 DerivativeOrder = 0) const override;
 
-	virtual void ExtendTo(const FPoint& Point) override;
+	virtual void ExtendTo(const FVector& Point) override;
 
 	/**
 	 * Update the spline's internal data according to the passed-in params
@@ -105,12 +105,11 @@ public:
 	double GetSplineLength() const;
 
 	/** Sets the spline to an array of points */
-	void SetSplinePoints(const TArray<FPoint>& InPoints);
-	void SetSplinePoints(const TArray<FPoint>& InPoints, const TArray<FPoint>& InTangents);
-	void SetSplinePoints(const TArray<FPoint>& InPoints, const TArray<FPoint>& InArriveTangents, const TArray<FPoint>& InLeaveTangents);
-
-	const FInterpCurveFPoint& GetSplinePointsPosition() const { return Position; }
+	void SetSplinePoints(const TArray<FVector>& InPoints);
+	void SetSplinePoints(const TArray<FVector>& InPoints, const TArray<FVector>& InTangents);
+	void SetSplinePoints(const TArray<FVector>& InPoints, const TArray<FVector>& InArriveTangents, const TArray<FVector>& InLeaveTangents);
 
+	const FInterpCurveFVector& GetSplinePointsPosition() const { return Position; }
 };
 
 } // namespace UE::CADKernel
diff --git a/Engine/Source/Runtime/Datasmith/CADKernel/Base/Private/Geo/Sampler/SamplerAbstract.h b/Engine/Source/Runtime/Datasmith/CADKernel/Base/Private/Geo/Sampler/SamplerAbstract.h
--- a/Engine/Source/Runtime/Datasmith/CADKernel/Base/Private/Geo/Sampler/SamplerAbstract.h
+++ b/Engine/Source/Runtime/Datasmith/CADKernel/Base/Private/Geo/Sampler/SamplerAbstract.h
@@ -6,6 +6,7 @@
 #include "Geo/Surfaces/Surface.h"
 #include "Geo/Curves/Curve.h"
 #include "Math/Boundary.h"
+#include "Math/Geometry.h"
 #include "Math/Point.h"
 
 //#define DEBUG_CURVE_SAMPLING 
@@ -347,7 +348,7 @@ protected:
 
 			const PointType& PreviousPoint = PreviousSampingCoordinate > PreviousCandidateCoordinate ? Sampling.GetPoints()[PreviousSampingIndex] : CandidatePoints.GetPoints()[PreviousCandidateIndex];
 
-			double ChordA = DistanceOfPointToSegment(APoint, PreviousPoint, Middle);
+			double ChordA = DistanceOfPointToSegment<PointType>(APoint, PreviousPoint, Middle);
 			if (ChordA > 4 * DesiredChordError)
 			{
 				CountOfNeededPoints = CountOfNeededPointsToRespectChordError(APoint, PreviousPoint, ChordA);
@@ -365,7 +366,7 @@ protected:
 
 			const PointType& NextPoint = NextSampingCoordinate < NextCandidateCoordinate ? Sampling.GetPoints()[NextSampingIndex] : CandidatePoints.GetPoints()[NextCandidateIndex];
 
-			double ChordB = DistanceOfPointToSegment(BPoint, NextPoint, Middle);
+			double ChordB = DistanceOfPointToSegment<PointType>(BPoint, NextPoint, Middle);
 			if (ChordB > 4 * DesiredChordError)
 			{
 				CountOfNeededPoints = FMath::Max(CountOfNeededPointsToRespectChordError(BPoint, NextPoint, ChordB), CountOfNeededPoints);
@@ -424,7 +425,7 @@ public:
 			{
 				int32 Index = 0;
 				F3DDebugSession G(FString::Printf(TEXT("Curve %d"), Step));
-				for (const FPoint& Point : Sampling.GetPoints())
+				for (const FVector& Point : Sampling.GetPoints())
 				{
 					UE::CADKernel::DisplayPoint(Point, EVisuProperty::BluePoint, Index++);
 				}
@@ -444,7 +445,7 @@ public:
 			{
 				int32 Index = 0;
 				F3DDebugSession G(FString::Printf(TEXT("Next Point Tmp %d"), Step));
-				for (const FPoint& Point : CandidatePoints.GetPoints())
+				for (const FVector& Point : CandidatePoints.GetPoints())
 				{
 					UE::CADKernel::DisplayPoint(Point, EVisuProperty::YellowPoint, Index++);
 				}
@@ -477,7 +478,7 @@ public:
 		for (int32 Index = 0; Index < EvaluationCount; ++Index)
 		{
 			double CurvilinearAbscise;
-			PointType ProjectedPoint = ProjectPointOnSegment(CandidatePoints.GetPoints()[Index], Sampling.GetPoints()[Index], Sampling.GetPoints()[Index + 1], CurvilinearAbscise);
+			PointType ProjectedPoint = ProjectPointOnSegment<PointType>(CandidatePoints.GetPoints()[Index], Sampling.GetPoints()[Index], Sampling.GetPoints()[Index + 1], CurvilinearAbscise);
 			double ChordError = ProjectedPoint.Distance(CandidatePoints.GetPoints()[Index]);
 
 			MinError[0] = FMath::Min(MinError[0], ChordError);
diff --git a/Engine/Source/Runtime/Datasmith/CADKernel/Base/Private/Geo/Sampler/SamplerOnChord.h b/Engine/Source/Runtime/Datasmith/CADKernel/Base/Private/Geo/Sampler/SamplerOnChord.h
--- a/Engine/Source/Runtime/Datasmith/CADKernel/Base/Private/Geo/Sampler/SamplerOnChord.h
+++ b/Engine/Source/Runtime/Datasmith/CADKernel/Base/Private/Geo/Sampler/SamplerOnChord.h
@@ -58,7 +58,7 @@ protected:
 		for (int32 Index = FirstIndex; Index < EndIndex; ++Index)
 		{
 			const PointType& Middle = this->CandidatePoints.GetPointAt(Index);
-			double Chord = DistanceOfPointToSegment(Middle, APoint, BPoint);
+			double Chord = DistanceOfPointToSegment<PointType>(Middle, APoint, BPoint);
 			if (Chord > MaxChord)
 			{
 				MaxChord = Chord;
@@ -129,11 +129,11 @@ protected:
 /**
  * Sampler of curve based on parametrization and chord error control
  */
-class FCurveSamplerOnChord : public TSamplerBasedOnChordError<FPoint>
+class FCurveSamplerOnChord : public TSamplerBasedOnChordError<FVector>
 {
 public:
 	FCurveSamplerOnChord(const FCurve& InCurve, const FLinearBoundary& InBoundary, double InMaxAllowedError, FPolyline3D& OutPolyline)
-		: TSamplerBasedOnChordError<FPoint>(InBoundary, InMaxAllowedError, OutPolyline)
+		: TSamplerBasedOnChordError<FVector>(InBoundary, InMaxAllowedError, OutPolyline)
 		, Curve(InCurve)
 	{
 	};
@@ -154,11 +154,11 @@ protected:
 
 };
 
-class FCurve2DSamplerOnChord : public TSamplerBasedOnChordError<FPoint2D>
+class FCurve2DSamplerOnChord : public TSamplerBasedOnChordError<FVector2d>
 {
 public:
 	FCurve2DSamplerOnChord(const FCurve& InCurve, const FLinearBoundary& InBoundary, double InMaxAllowedError, FPolyline2D& OutPolyline)
-		: TSamplerBasedOnChordError<FPoint2D>(InBoundary, InMaxAllowedError, OutPolyline)
+		: TSamplerBasedOnChordError<FVector2d>(InBoundary, InMaxAllowedError, OutPolyline)
 		, Curve(InCurve)
 	{
 	};
@@ -179,11 +179,11 @@ protected:
 
 };
 
-class FIsoCurve3DSamplerOnChord : public TSamplerBasedOnChordError<FPoint>
+class FIsoCurve3DSamplerOnChord : public TSamplerBasedOnChordError<FVector>
 {
 public:
 	FIsoCurve3DSamplerOnChord(const FSurface& InSurface, double InMaxAllowedError, FPolyline3D& OutPolyline)
-		: TSamplerBasedOnChordError<FPoint>(FLinearBoundary(), InMaxAllowedError, OutPolyline)
+		: TSamplerBasedOnChordError<FVector>(FLinearBoundary(), InMaxAllowedError, OutPolyline)
 		, Surface(InSurface)
 	{
 	};
diff --git a/Engine/Source/Runtime/Datasmith/CADKernel/Base/Private/Geo/Sampler/SamplerOnParam.h b/Engine/Source/Runtime/Datasmith/CADKernel/Base/Private/Geo/Sampler/SamplerOnParam.h
--- a/Engine/Source/Runtime/Datasmith/CADKernel/Base/Private/Geo/Sampler/SamplerOnParam.h
+++ b/Engine/Source/Runtime/Datasmith/CADKernel/Base/Private/Geo/Sampler/SamplerOnParam.h
@@ -58,28 +58,28 @@ protected:
 
 	virtual int32 CheckSamplingError(int32 FirstIndex, int32 EndIndex) override
 	{
-		const FPoint& APoint = this->Sampling.GetPointAt(this->StartSamplingSegmentIndex);
+		const FVector& APoint = this->Sampling.GetPointAt(this->StartSamplingSegmentIndex);
 		double ACoordinate = this->Sampling.Coordinates[this->StartSamplingSegmentIndex];
 
-		const FPoint& BPoint = this->EndStudySegment.Polyline->GetPointAt(this->EndStudySegment.Index);
+		const FVector& BPoint = this->EndStudySegment.Polyline->GetPointAt(this->EndStudySegment.Index);
 		double BCoordinate = this->EndStudySegment.Polyline->Coordinates[this->EndStudySegment.Index];
 
 		double ABCoordinate = BCoordinate - ACoordinate;
-		FPoint ABVector = BPoint - APoint;
+		FVector ABVector = BPoint - APoint;
 
 		double MaxChord = 0;
 		double MaxParametrizationError = 0;
 
 		for (int32 Index = FirstIndex; Index < EndIndex; ++Index)
 		{
-			const FPoint& PointM = this->CandidatePoints.GetPointAt(Index);
+			const FVector& PointM = this->CandidatePoints.GetPointAt(Index);
 
 			double CurvilinearAbscise;
-			FPoint ProjectedPoint = ProjectPointOnSegment(PointM, APoint, BPoint, CurvilinearAbscise, /*bRestrictCoodinateToInside*/ true);
+			FVector ProjectedPoint = ProjectPointOnSegment<FVector>(PointM, APoint, BPoint, CurvilinearAbscise, /*bRestrictCoodinateToInside*/ true);
 
 			double Chord = FVector::Distance(ProjectedPoint, PointM);
 			double AMCoordinate = this->CandidatePoints.Coordinates[Index] - ACoordinate;
-			FPoint MApproximate = APoint + (AMCoordinate / ABCoordinate) * ABVector;
+			FVector MApproximate = APoint + (AMCoordinate / ABCoordinate) * ABVector;
 			double ParametrizationError = FVector::Distance(MApproximate, ProjectedPoint);
 
 			if (Chord > MaxChord)
@@ -130,18 +130,18 @@ protected:
 /**
  * Sampler of surfacic curve based on parametrization and chord error control
  */
-class FSurfacicCurveSamplerOnParam : public TSamplerBasedOnParametrizationAndChordError<FSurfacicPolyline, FPoint>
+class FSurfacicCurveSamplerOnParam : public TSamplerBasedOnParametrizationAndChordError<FSurfacicPolyline, FVector>
 {
 public:
 	FSurfacicCurveSamplerOnParam(const FRestrictionCurve& InCurve, const FLinearBoundary& InBoundary, double InMaxSagError, double InMaxParameterizationError, FSurfacicPolyline& OutPolyline)
-		: TSamplerBasedOnParametrizationAndChordError<FSurfacicPolyline, FPoint>(InBoundary, InMaxSagError, InMaxParameterizationError, OutPolyline)
+		: TSamplerBasedOnParametrizationAndChordError<FSurfacicPolyline, FVector>(InBoundary, InMaxSagError, InMaxParameterizationError, OutPolyline)
 		, Surface(*InCurve.GetCarrierSurface())
 		, Curve(*InCurve.Get2DCurve())
 	{
 	};
 
 	FSurfacicCurveSamplerOnParam(const FSurface& InSurface, const FCurve& InCurve, const FLinearBoundary& InBoundary, double InMaxSagError, double InMaxParameterizationError, FSurfacicPolyline& OutPolyline)
-		: TSamplerBasedOnParametrizationAndChordError<FSurfacicPolyline, FPoint>(InBoundary, InMaxSagError, InMaxParameterizationError, OutPolyline)
+		: TSamplerBasedOnParametrizationAndChordError<FSurfacicPolyline, FVector>(InBoundary, InMaxSagError, InMaxParameterizationError, OutPolyline)
 		, Surface(InSurface)
 		, Curve(InCurve)
 	{
@@ -181,11 +181,11 @@ protected:
 /**
  * Sampler of curve based on parametrization and chord error control
  */
-class FCurveSamplerOnParam : public TSamplerBasedOnParametrizationAndChordError<FPolyline3D, FPoint>
+class FCurveSamplerOnParam : public TSamplerBasedOnParametrizationAndChordError<FPolyline3D, FVector>
 {
 public:
 	FCurveSamplerOnParam(const FCurve& InCurve, const FLinearBoundary& InBoundary, double InMaxSagError, double InMaxParameterizationError, FPolyline3D& OutPolyline)
-		: TSamplerBasedOnParametrizationAndChordError<FPolyline3D, FPoint>(InBoundary, InMaxSagError, InMaxParameterizationError, OutPolyline)
+		: TSamplerBasedOnParametrizationAndChordError<FPolyline3D, FVector>(InBoundary, InMaxSagError, InMaxParameterizationError, OutPolyline)
 		, Curve(InCurve)
 	{
 	};
@@ -209,12 +209,12 @@ protected:
 /**
  * Sampler of surface based on parametrization and chord error control
  */
-class FSurfaceSamplerOnParam : public TSamplerBasedOnParametrizationAndChordError<FPolyline3D, FPoint>
+class FSurfaceSamplerOnParam : public TSamplerBasedOnParametrizationAndChordError<FPolyline3D, FVector>
 {
 
 public:
 	FSurfaceSamplerOnParam(const FSurface& InSurface, const FSurfacicBoundary& InBoundary, double InMaxSagError, double InMaxParameterizationError, FPolyline3D& TemporaryPolyline, FCoordinateGrid& OutSampling)
-		: TSamplerBasedOnParametrizationAndChordError<FPolyline3D, FPoint>(FLinearBoundary(), InMaxSagError, InMaxParameterizationError, TemporaryPolyline)
+		: TSamplerBasedOnParametrizationAndChordError<FPolyline3D, FVector>(FLinearBoundary(), InMaxSagError, InMaxParameterizationError, TemporaryPolyline)
 		, Surface(InSurface)
 		, SurfaceBoundary(InBoundary)
 		, SurfaceSampling(OutSampling)
diff --git a/Engine/Source/Runtime/Datasmith/CADKernel/Base/Private/Geo/Sampling/SurfacicPolyline.cpp b/Engine/Source/Runtime/Datasmith/CADKernel/Base/Private/Geo/Sampling/SurfacicPolyline.cpp
--- a/Engine/Source/Runtime/Datasmith/CADKernel/Base/Private/Geo/Sampling/SurfacicPolyline.cpp
+++ b/Engine/Source/Runtime/Datasmith/CADKernel/Base/Private/Geo/Sampling/SurfacicPolyline.cpp
@@ -39,11 +39,11 @@ FSurfacicPolyline::FSurfacicPolyline(TSharedRef<FSurface> InCarrierSurface, TSha
 
 void FSurfacicPolyline::CheckIfDegenerated(const double Tolerance3D, const FSurfacicTolerance& ToleranceIso, const FLinearBoundary& Boudary, bool& bDegeneration2D, bool& bDegeneration3D, double& Length3D) const
 {
-	TPolylineApproximator<FPoint> Approximator(Coordinates, Points3D);
+	TPolylineApproximator<FVector> Approximator(Coordinates, Points3D);
 	int32 BoundaryIndices[2];
 	Approximator.GetStartEndIndex(Boudary, BoundaryIndices);
 
-	TPolylineApproximator<FPoint> Approximator3D(Coordinates, Points3D);
+	TPolylineApproximator<FVector> Approximator3D(Coordinates, Points3D);
 	Length3D = Approximator3D.ComputeLengthOfSubPolyline(BoundaryIndices, Boudary);
 
 	if (!FMath::IsNearlyZero(Length3D, Tolerance3D))
@@ -62,7 +62,7 @@ void FSurfacicPolyline::CheckIfDegenerated(const double Tolerance3D, const FSurf
 	// The tolerance along ans iso is the length of the boundary along this iso divide by 100 000: if the curve length in 3d is 10m, the tolerance is 0.01mm
 
 	FAABB2D Aabb;
-	TPolylineApproximator<FPoint2D> Approximator2D(Coordinates, Points2D);
+	TPolylineApproximator<FVector2d> Approximator2D(Coordinates, Points2D);
 	Approximator2D.ComputeBoundingBox<2>(BoundaryIndices, Boudary, Aabb);
 
 	bDegeneration2D = (Aabb.GetSize(0) < ToleranceIso[EIso::IsoU] && Aabb.GetSize(1) < ToleranceIso[EIso::IsoV]);
diff --git a/Engine/Source/Runtime/Datasmith/CADKernel/Base/Private/Geo/Sampling/SurfacicSampling.h b/Engine/Source/Runtime/Datasmith/CADKernel/Base/Private/Geo/Sampling/SurfacicSampling.h
--- a/Engine/Source/Runtime/Datasmith/CADKernel/Base/Private/Geo/Sampling/SurfacicSampling.h
+++ b/Engine/Source/Runtime/Datasmith/CADKernel/Base/Private/Geo/Sampling/SurfacicSampling.h
@@ -13,8 +13,8 @@ namespace UE::CADKernel
 struct FSurfacicSampling
 {
 	bool bWithNormals = false;
-	TArray<FPoint2D> Points2D;
-	TArray<FPoint> Points3D;
+	TArray<FVector2d> Points2D;
+	TArray<FVector> Points3D;
 	TArray<FVector3f> Normals;
 
 	int32 Count()
diff --git a/Engine/Source/Runtime/Datasmith/CADKernel/Base/Private/Geo/Surfaces/BezierSurface.cpp b/Engine/Source/Runtime/Datasmith/CADKernel/Base/Private/Geo/Surfaces/BezierSurface.cpp
--- a/Engine/Source/Runtime/Datasmith/CADKernel/Base/Private/Geo/Surfaces/BezierSurface.cpp
+++ b/Engine/Source/Runtime/Datasmith/CADKernel/Base/Private/Geo/Surfaces/BezierSurface.cpp
@@ -9,11 +9,11 @@
 namespace UE::CADKernel
 {
 
-void FBezierSurface::EvaluatePoint(const FPoint2D& InPoint2D, FSurfacicPoint& OutPoint3D, int32 InDerivativeOrder) const
+void FBezierSurface::EvaluatePoint(const FVector2d& InPoint2D, FSurfacicPoint& OutPoint3D, int32 InDerivativeOrder) const
 {
-	TArray<FPoint> VCurvePoles;
-	TArray<FPoint> VCurveUGradient;
-	TArray<FPoint> VCurveULaplacian;
+	TArray<FVector> VCurvePoles;
+	TArray<FVector> VCurveUGradient;
+	TArray<FVector> VCurveULaplacian;
 
 	VCurvePoles.SetNum(VPoleNum);
 
@@ -21,7 +21,7 @@ void FBezierSurface::EvaluatePoint(const FPoint2D& InPoint2D, FSurfacicPoint& Ou
 	
 	if (InDerivativeOrder > 0)
 	{
-		VCurveUGradient.Init(FPoint::ZeroVector, VPoleNum);
+		VCurveUGradient.Init(FVector::ZeroVector, VPoleNum);
 	}
 
 	if (InDerivativeOrder > 1)
@@ -32,11 +32,11 @@ void FBezierSurface::EvaluatePoint(const FPoint2D& InPoint2D, FSurfacicPoint& Ou
 	{
 		double Coordinate = InPoint2D.X;
 
-		TArray<FPoint> UAuxiliaryPoles;
+		TArray<FVector> UAuxiliaryPoles;
 		UAuxiliaryPoles.SetNum(UPoleNum);
 
-		TArray<FPoint> UAuxiliaryGradient;
-		TArray<FPoint> UAuxiliaryLaplacian;
+		TArray<FVector> UAuxiliaryGradient;
+		TArray<FVector> UAuxiliaryLaplacian;
 
 		// For each iso V Curve compute Point, Gradient and Laplacian at U coordinate
 		for (int32 Vndex = 0, PoleIndex = 0; Vndex < VPoleNum; Vndex++, PoleIndex += UPoleNum)
@@ -46,11 +46,11 @@ void FBezierSurface::EvaluatePoint(const FPoint2D& InPoint2D, FSurfacicPoint& Ou
 
 			if (InDerivativeOrder > 0)
 			{
-				UAuxiliaryGradient.Init(FPoint::ZeroVector, UPoleNum);
+				UAuxiliaryGradient.Init(FVector::ZeroVector, UPoleNum);
 			}
 			if (InDerivativeOrder > 1)
 			{
-				UAuxiliaryLaplacian.Init(FPoint::ZeroVector, UPoleNum);
+				UAuxiliaryLaplacian.Init(FVector::ZeroVector, UPoleNum);
 			}
 
 			// Compute Point, Gradient and Laplacian at U coordinate with De Casteljau's algorithm, 
@@ -58,16 +58,16 @@ void FBezierSurface::EvaluatePoint(const FPoint2D& InPoint2D, FSurfacicPoint& Ou
 			{
 				for (int32 Index = 0; Index <= Undex; Index++)
 				{
-					const FPoint PointI = UAuxiliaryPoles[Index];
-					const FPoint& PointB = UAuxiliaryPoles[Index + 1];
-					const FPoint VectorIB = PointB - PointI;
+					const FVector PointI = UAuxiliaryPoles[Index];
+					const FVector& PointB = UAuxiliaryPoles[Index + 1];
+					const FVector VectorIB = PointB - PointI;
 
 					UAuxiliaryPoles[Index] = PointI + VectorIB * Coordinate;
 					if (InDerivativeOrder > 0)
 					{
-						const FPoint GradientI = UAuxiliaryGradient[Index];
-						const FPoint& GradientB = UAuxiliaryGradient[Index + 1];
-						const FPoint VectorGradientIB = GradientB - GradientI;
+						const FVector GradientI = UAuxiliaryGradient[Index];
+						const FVector& GradientB = UAuxiliaryGradient[Index + 1];
+						const FVector VectorGradientIB = GradientB - GradientI;
 
 						UAuxiliaryGradient[Index] = GradientI + VectorGradientIB * Coordinate + VectorIB;
 						if (InDerivativeOrder > 1)
@@ -93,16 +93,16 @@ void FBezierSurface::EvaluatePoint(const FPoint2D& InPoint2D, FSurfacicPoint& Ou
 		}
 	}
 
-	TArray<FPoint> VAuxiliaryGradient;
-	TArray<FPoint> UVAuxiliaryLaplacian;
-	TArray<FPoint> VVAuxiliaryLaplacian;
+	TArray<FVector> VAuxiliaryGradient;
+	TArray<FVector> UVAuxiliaryLaplacian;
+	TArray<FVector> VVAuxiliaryLaplacian;
 	if (InDerivativeOrder > 0)
 	{
-		VAuxiliaryGradient.Init(FPoint::ZeroVector, VPoleNum);
+		VAuxiliaryGradient.Init(FVector::ZeroVector, VPoleNum);
 		if (InDerivativeOrder > 1)
 		{
-			VVAuxiliaryLaplacian.Init(FPoint::ZeroVector, VPoleNum);
-			UVAuxiliaryLaplacian.Init(FPoint::ZeroVector, VPoleNum);
+			VVAuxiliaryLaplacian.Init(FVector::ZeroVector, VPoleNum);
+			UVAuxiliaryLaplacian.Init(FVector::ZeroVector, VPoleNum);
 		}
 	}
 
@@ -113,22 +113,22 @@ void FBezierSurface::EvaluatePoint(const FPoint2D& InPoint2D, FSurfacicPoint& Ou
 	{
 		for (int32 Index = 0; Index <= Vndex; Index++)
 		{
-			const FPoint PointI = VCurvePoles[Index];
-			const FPoint& PointB = VCurvePoles[Index + 1];
-			const FPoint VectorIB = PointB - PointI;
+			const FVector PointI = VCurvePoles[Index];
+			const FVector& PointB = VCurvePoles[Index + 1];
+			const FVector VectorIB = PointB - PointI;
 			VCurvePoles[Index] = PointI + VectorIB * Coordinate;
 
 			if (InDerivativeOrder > 0) 
 			{
-				const FPoint UGradientI = VCurveUGradient[Index];
-				const FPoint& UGradientB = VCurveUGradient[Index + 1];
-				const FPoint VectorUGradientIB = UGradientB - UGradientI;
+				const FVector UGradientI = VCurveUGradient[Index];
+				const FVector& UGradientB = VCurveUGradient[Index + 1];
+				const FVector VectorUGradientIB = UGradientB - UGradientI;
 
 				VCurveUGradient[Index] = UGradientI + VectorUGradientIB * Coordinate;
 
-				const FPoint VGradientI = VAuxiliaryGradient[Index]; 
-				const FPoint& VGradientB = VAuxiliaryGradient[Index + 1];
-				const FPoint VectorVGradientIB = VGradientB - VGradientI;
+				const FVector VGradientI = VAuxiliaryGradient[Index]; 
+				const FVector& VGradientB = VAuxiliaryGradient[Index + 1];
+				const FVector VectorVGradientIB = VGradientB - VGradientI;
 
 				VAuxiliaryGradient[Index] = VGradientI + VectorVGradientIB * Coordinate + VectorIB;
 
@@ -170,9 +170,9 @@ void FBezierSurface::Presample(const FSurfacicBoundary& InBoundaries, FCoordinat
 
 TSharedPtr<FEntityGeom> FBezierSurface::ApplyMatrix(const FMatrixH& InMatrix) const
 {
-	TArray<FPoint> TransformedPoles;
+	TArray<FVector> TransformedPoles;
 	TransformedPoles.Reserve(Poles.Num());
-	for (const FPoint& Pole : Poles) 
+	for (const FVector& Pole : Poles) 
 	{
 		TransformedPoles.Emplace(InMatrix.Multiply(Pole));
 	}
diff --git a/Engine/Source/Runtime/Datasmith/CADKernel/Base/Private/Geo/Surfaces/BezierSurface.h b/Engine/Source/Runtime/Datasmith/CADKernel/Base/Private/Geo/Surfaces/BezierSurface.h
--- a/Engine/Source/Runtime/Datasmith/CADKernel/Base/Private/Geo/Surfaces/BezierSurface.h
+++ b/Engine/Source/Runtime/Datasmith/CADKernel/Base/Private/Geo/Surfaces/BezierSurface.h
@@ -19,9 +19,9 @@ protected:
 	int32 UPoleNum;
 	int32 VPoleNum;
 
-	TArray<FPoint> Poles;
+	TArray<FVector> Poles;
 
-	FBezierSurface(const double InToleranceGeometric, int32 InUDegre, int32 InVDegre, const TArray<FPoint>& InPoles)
+	FBezierSurface(const double InToleranceGeometric, int32 InUDegre, int32 InVDegre, const TArray<FVector>& InPoles)
 		: FSurface(InToleranceGeometric)
 		, UDegre(InUDegre)
 		, VDegre(InVDegre)
@@ -62,14 +62,14 @@ public:
 		return VDegre;
 	}
 
-	virtual const TArray<FPoint>& GetPoles() const
+	virtual const TArray<FVector>& GetPoles() const
 	{
 		return Poles;
 	}
 
 	virtual TSharedPtr<FEntityGeom> ApplyMatrix(const FMatrixH& InMatrix) const override;
 
-	virtual void EvaluatePoint(const FPoint2D& InSurfacicCoordinate, FSurfacicPoint& OutPoint3D, int32 InDerivativeOrder = 0) const override;
+	virtual void EvaluatePoint(const FVector2d& InSurfacicCoordinate, FSurfacicPoint& OutPoint3D, int32 InDerivativeOrder = 0) const override;
 
 	virtual void Presample(const FSurfacicBoundary& InBoundaries, FCoordinateGrid& OutCoordinates) override;
 
diff --git a/Engine/Source/Runtime/Datasmith/CADKernel/Base/Private/Geo/Surfaces/CompositeSurface.cpp b/Engine/Source/Runtime/Datasmith/CADKernel/Base/Private/Geo/Surfaces/CompositeSurface.cpp
--- a/Engine/Source/Runtime/Datasmith/CADKernel/Base/Private/Geo/Surfaces/CompositeSurface.cpp
+++ b/Engine/Source/Runtime/Datasmith/CADKernel/Base/Private/Geo/Surfaces/CompositeSurface.cpp
@@ -70,13 +70,13 @@ void FCompositeSurface::LinesNotDerivables(const FSurfacicBoundary& Bounds, int3
 	}
 }
 
-void FCompositeSurface::EvaluatePoint(const FPoint2D& InSurfacicCoordinate, FSurfacicPoint& OutPoint3D, int32 InDerivativeOrder) const
+void FCompositeSurface::EvaluatePoint(const FVector2d& InSurfacicCoordinate, FSurfacicPoint& OutPoint3D, int32 InDerivativeOrder) const
 {
-	FPoint2D Coordinate = InSurfacicCoordinate;
+	FVector2d Coordinate = InSurfacicCoordinate;
 
 	Boundary.MoveInsideIfNot(Coordinate);
 
-	FPoint2D Point2DPatch;
+	FVector2d Point2DPatch;
 	FDichotomyFinder FinderU(GlobalCoordinates[EIso::IsoU]);
 	int32 PatchUIndex = FinderU.Find(Coordinate.X);
 	Point2DPatch.X = (Coordinate.X - GlobalCoordinates[EIso::IsoU][PatchUIndex]) / (GlobalCoordinates[EIso::IsoU][PatchUIndex + 1] - GlobalCoordinates[EIso::IsoU][PatchUIndex]);
diff --git a/Engine/Source/Runtime/Datasmith/CADKernel/Base/Private/Geo/Surfaces/CompositeSurface.h b/Engine/Source/Runtime/Datasmith/CADKernel/Base/Private/Geo/Surfaces/CompositeSurface.h
--- a/Engine/Source/Runtime/Datasmith/CADKernel/Base/Private/Geo/Surfaces/CompositeSurface.h
+++ b/Engine/Source/Runtime/Datasmith/CADKernel/Base/Private/Geo/Surfaces/CompositeSurface.h
@@ -79,7 +79,7 @@ public:
 
 	virtual void LinesNotDerivables(const FSurfacicBoundary& Bounds, int32 InDerivativeOrder, FCoordinateGrid& OutNotDerivableCoordinates) const override;
 
-	virtual void EvaluatePoint(const FPoint2D& InSurfacicCoordinate, FSurfacicPoint& OutPoint3D, int32 InDerivativeOrder = 0) const override;
+	virtual void EvaluatePoint(const FVector2d& InSurfacicCoordinate, FSurfacicPoint& OutPoint3D, int32 InDerivativeOrder = 0) const override;
 
 	virtual void Presample(const FSurfacicBoundary& InBoundaries, FCoordinateGrid& OutCoordinates) override;
 
diff --git a/Engine/Source/Runtime/Datasmith/CADKernel/Base/Private/Geo/Surfaces/ConeSurface.cpp b/Engine/Source/Runtime/Datasmith/CADKernel/Base/Private/Geo/Surfaces/ConeSurface.cpp
--- a/Engine/Source/Runtime/Datasmith/CADKernel/Base/Private/Geo/Surfaces/ConeSurface.cpp
+++ b/Engine/Source/Runtime/Datasmith/CADKernel/Base/Private/Geo/Surfaces/ConeSurface.cpp
@@ -9,7 +9,7 @@
 namespace UE::CADKernel
 {
 
-void FConeSurface::EvaluatePoint(const FPoint2D& InPoint2D, FSurfacicPoint& OutPoint3D, int32 InDerivativeOrder) const
+void FConeSurface::EvaluatePoint(const FVector2d& InPoint2D, FSurfacicPoint& OutPoint3D, int32 InDerivativeOrder) const
 {
 	OutPoint3D.DerivativeOrder = InDerivativeOrder;
 
@@ -22,18 +22,18 @@ void FConeSurface::EvaluatePoint(const FPoint2D& InPoint2D, FSurfacicPoint& OutP
 
 	if (InDerivativeOrder > 0)
 	{
-		OutPoint3D.GradientU = FPoint(-Radius * sin(InPoint2D.X), Radius * cos(InPoint2D.X), 0.);
-		OutPoint3D.GradientV = FPoint(DeltaVR * cos(InPoint2D.X), DeltaVR * sin(InPoint2D.X), 1.0);
+		OutPoint3D.GradientU = FVector(-Radius * sin(InPoint2D.X), Radius * cos(InPoint2D.X), 0.);
+		OutPoint3D.GradientV = FVector(DeltaVR * cos(InPoint2D.X), DeltaVR * sin(InPoint2D.X), 1.0);
 
 		OutPoint3D.GradientU = Matrix.MultiplyVector(OutPoint3D.GradientU);
 		OutPoint3D.GradientV = Matrix.MultiplyVector(OutPoint3D.GradientV);
 
 		if (InDerivativeOrder > 1)
 		{
-			OutPoint3D.LaplacianU = FPoint(-Radius * cos(InPoint2D.X), -Radius * sin(InPoint2D.X), 0.);
-			OutPoint3D.LaplacianV = FPoint(0., 0., 0.);
+			OutPoint3D.LaplacianU = FVector(-Radius * cos(InPoint2D.X), -Radius * sin(InPoint2D.X), 0.);
+			OutPoint3D.LaplacianV = FVector(0., 0., 0.);
 
-			OutPoint3D.LaplacianUV = FPoint(-DeltaVR * sin(InPoint2D.X), DeltaVR * cos(InPoint2D.X), 0.);
+			OutPoint3D.LaplacianUV = FVector(-DeltaVR * sin(InPoint2D.X), DeltaVR * cos(InPoint2D.X), 0.);
 
 			OutPoint3D.LaplacianU = Matrix.MultiplyVector(OutPoint3D.LaplacianU);
 			OutPoint3D.LaplacianV = Matrix.MultiplyVector(OutPoint3D.LaplacianV);
@@ -78,7 +78,7 @@ void FConeSurface::EvaluatePointGrid(const FCoordinateGrid& Coordinates, FSurfac
 
 		for (int32 Undex = 0; Undex < UCount; Undex++)
 		{
-			FPoint& Point = OutPoints.Points3D.Emplace_GetRef(Radius * CosU[Undex], Radius * SinU[Undex], Coordinates[EIso::IsoV][Vndex]);
+			FVector& Point = OutPoints.Points3D.Emplace_GetRef(Radius * CosU[Undex], Radius * SinU[Undex], Coordinates[EIso::IsoV][Vndex]);
 			Point = Matrix.Multiply(Point);
 		}
 	}
@@ -90,11 +90,11 @@ void FConeSurface::EvaluatePointGrid(const FCoordinateGrid& Coordinates, FSurfac
 			double Radius = StartRadius + Vndex * DeltaVR;
 			for (int32 Undex = 0; Undex < UCount; Undex++)
 			{
-				FPoint GradientU(-Radius * SinU[Undex], Radius * CosU[Undex], 0.);
-				FPoint GradientV(DeltaVR * CosU[Undex], DeltaVR * SinU[Undex], 1.0);
+				FVector GradientU(-Radius * SinU[Undex], Radius * CosU[Undex], 0.);
+				FVector GradientV(DeltaVR * CosU[Undex], DeltaVR * SinU[Undex], 1.0);
 				GradientU = Matrix.MultiplyVector(GradientU);
 				GradientV = Matrix.MultiplyVector(GradientV);
-				const FPoint Normal = GradientU ^ GradientV;
+				const FVector Normal = GradientU ^ GradientV;
 				OutPoints.Normals.Emplace(Normal);
 			}
 		}
diff --git a/Engine/Source/Runtime/Datasmith/CADKernel/Base/Private/Geo/Surfaces/ConeSurface.h b/Engine/Source/Runtime/Datasmith/CADKernel/Base/Private/Geo/Surfaces/ConeSurface.h
--- a/Engine/Source/Runtime/Datasmith/CADKernel/Base/Private/Geo/Surfaces/ConeSurface.h
+++ b/Engine/Source/Runtime/Datasmith/CADKernel/Base/Private/Geo/Surfaces/ConeSurface.h
@@ -58,19 +58,19 @@ protected:
 
 	void ComputeMinToleranceIso()
 	{
-		FPoint Origin = Matrix.Multiply(FPoint::ZeroVector);
+		FVector Origin = Matrix.Multiply(FVector::ZeroVector);
 
 		double DeltaVR = tan(ConeAngle);
 		double Radius1 = FMath::Abs(StartRadius + Boundary[EIso::IsoV].Max * DeltaVR);
 		double Radius2 = FMath::Abs(StartRadius + Boundary[EIso::IsoV].Min * DeltaVR);
 		double Radius = FMath::Max(Radius2, Radius1);
 
-		FPoint Point2DU{ 1 , 0, 0 };
+		FVector Point2DU{ 1 , 0, 0 };
 		double ScaleU = ComputeScaleAlongAxis(Point2DU, Matrix, Origin);
 		double ToleranceU = Tolerance3D / Radius;
 		ToleranceU /= ScaleU;
 
-		FPoint Point2DV{ 0, 1, 0 };
+		FVector Point2DV{ 0, 1, 0 };
 		double ScaleV = ComputeScaleAlongAxis(Point2DV, Matrix, Origin);
 		double ToleranceV = Tolerance3D / FMath::Sqrt(FMath::Square(DeltaVR) + 1);
 		ToleranceV /= ScaleV;
@@ -100,7 +100,7 @@ public:
 
 	virtual TSharedPtr<FEntityGeom> ApplyMatrix(const FMatrixH& InMatrix) const override;
 
-	virtual void EvaluatePoint(const FPoint2D& InSurfacicCoordinate, FSurfacicPoint& OutPoint3D, int32 InDerivativeOrder = 0) const override;
+	virtual void EvaluatePoint(const FVector2d& InSurfacicCoordinate, FSurfacicPoint& OutPoint3D, int32 InDerivativeOrder = 0) const override;
 	virtual void EvaluatePointGrid(const FCoordinateGrid& Coordinates, FSurfacicSampling& OutPoints, bool bComputeNormals = false) const override;
 
 	virtual void Presample(const FSurfacicBoundary& InBoundaries, FCoordinateGrid& OutCoordinates) override
diff --git a/Engine/Source/Runtime/Datasmith/CADKernel/Base/Private/Geo/Surfaces/CoonsSurface.cpp b/Engine/Source/Runtime/Datasmith/CADKernel/Base/Private/Geo/Surfaces/CoonsSurface.cpp
--- a/Engine/Source/Runtime/Datasmith/CADKernel/Base/Private/Geo/Surfaces/CoonsSurface.cpp
+++ b/Engine/Source/Runtime/Datasmith/CADKernel/Base/Private/Geo/Surfaces/CoonsSurface.cpp
@@ -56,10 +56,10 @@ void FCoonsSurface::LinesNotDerivables(const FSurfacicBoundary& Bounds, int32 In
 	}
 }
 
-void FCoonsSurface::EvaluatePoint(const FPoint2D& InPoint2D, FSurfacicPoint& OutPoint3D, int32 InDerivativeOrder) const
+void FCoonsSurface::EvaluatePoint(const FVector2d& InPoint2D, FSurfacicPoint& OutPoint3D, int32 InDerivativeOrder) const
 {
-	TFunction<void(const TSharedPtr<FCurve>*, const FPoint2D&, EIso, FPoint&, FCurvePoint[])> \
-	ComputePointOnRuledSurface = [&](const TSharedPtr<FCurve>* InCurves, const FPoint2D& Coord, EIso Iso, FPoint& OutPoint, FCurvePoint OutCurvePoints[2])
+	TFunction<void(const TSharedPtr<FCurve>*, const FVector2d&, EIso, FVector&, FCurvePoint[])> \
+	ComputePointOnRuledSurface = [&](const TSharedPtr<FCurve>* InCurves, const FVector2d& Coord, EIso Iso, FVector& OutPoint, FCurvePoint OutCurvePoints[2])
 	{
 		const int32 IndexCurve0 = Iso * 2;
 		const int32 IndexCurve1 = Iso * 2 + 1;
@@ -87,17 +87,17 @@ void FCoonsSurface::EvaluatePoint(const FPoint2D& InPoint2D, FSurfacicPoint& Out
 	};
 
 	// 1. First ruled Surface
-	FPoint RuledPoint1;
+	FVector RuledPoint1;
 	FCurvePoint Curve12Points[2];
 	ComputePointOnRuledSurface(Curves, InPoint2D, EIso::IsoU, RuledPoint1, Curve12Points);
 	
 	// 2. Second ruled Surface
-	FPoint RuledPoint2;
+	FVector RuledPoint2;
 	FCurvePoint Curve34Points[2];
 	ComputePointOnRuledSurface(Curves+2, InPoint2D, EIso::IsoV, RuledPoint2, Curve34Points);
 
 	// 3. bilinear surface
-	FPoint BilinearPoint = (Corners[0]*(1.0-InPoint2D.X) + Corners[1]*InPoint2D.X)*(1.0-InPoint2D.Y) + (Corners[2]*(1.0-InPoint2D.X) + Corners[3]*InPoint2D.X)*InPoint2D.Y;
+	FVector BilinearPoint = (Corners[0]*(1.0-InPoint2D.X) + Corners[1]*InPoint2D.X)*(1.0-InPoint2D.Y) + (Corners[2]*(1.0-InPoint2D.X) + Corners[3]*InPoint2D.X)*InPoint2D.Y;
 
 	OutPoint3D.DerivativeOrder = InDerivativeOrder;
 
@@ -106,20 +106,20 @@ void FCoonsSurface::EvaluatePoint(const FPoint2D& InPoint2D, FSurfacicPoint& Out
 
 	if (InDerivativeOrder > 0)
 	{
-		FPoint DeltaU1 = Curve12Points[0].Gradient + InPoint2D.Y * (Curve12Points[1].Gradient - Curve12Points[0].Gradient);
-		FPoint DeltaU2 = Curve34Points[1].Point - Curve34Points[0].Point;
-		FPoint DeltaU3 = (Corners[1] * (1.0 - InPoint2D.Y) + Corners[3] * InPoint2D.Y) - (Corners[0] * (1.0 - InPoint2D.Y) + Corners[2] * InPoint2D.Y);
+		FVector DeltaU1 = Curve12Points[0].Gradient + InPoint2D.Y * (Curve12Points[1].Gradient - Curve12Points[0].Gradient);
+		FVector DeltaU2 = Curve34Points[1].Point - Curve34Points[0].Point;
+		FVector DeltaU3 = (Corners[1] * (1.0 - InPoint2D.Y) + Corners[3] * InPoint2D.Y) - (Corners[0] * (1.0 - InPoint2D.Y) + Corners[2] * InPoint2D.Y);
 		OutPoint3D.GradientU = DeltaU1 + DeltaU2 - DeltaU3;
 
-		FPoint DeltaV1 = Curve12Points[1].Point - Curve12Points[0].Point;;
-		FPoint DeltaV2 = Curve34Points[0].Gradient + InPoint2D.X * (Curve34Points[1].Gradient - Curve34Points[0].Gradient);
-		FPoint DeltaV3 = (Corners[2] * (1.0 - InPoint2D.X) + Corners[3] * InPoint2D.X) - (Corners[0] * (1.0 - InPoint2D.X) + Corners[1] * InPoint2D.X);
+		FVector DeltaV1 = Curve12Points[1].Point - Curve12Points[0].Point;;
+		FVector DeltaV2 = Curve34Points[0].Gradient + InPoint2D.X * (Curve34Points[1].Gradient - Curve34Points[0].Gradient);
+		FVector DeltaV3 = (Corners[2] * (1.0 - InPoint2D.X) + Corners[3] * InPoint2D.X) - (Corners[0] * (1.0 - InPoint2D.X) + Corners[1] * InPoint2D.X);
 		OutPoint3D.GradientV = DeltaV1 + DeltaV2 - DeltaV3;
 
 		if (InDerivativeOrder > 1)
 		{
-			FPoint LaplacianUV1 = (Curve12Points[1].Laplacian - Curve12Points[0].Laplacian);
-			FPoint LaplacianUV2 = (Curve34Points[1].Laplacian - Curve34Points[0].Laplacian);
+			FVector LaplacianUV1 = (Curve12Points[1].Laplacian - Curve12Points[0].Laplacian);
+			FVector LaplacianUV2 = (Curve34Points[1].Laplacian - Curve34Points[0].Laplacian);
 
 			OutPoint3D.LaplacianU = Curve12Points[0].Laplacian + InPoint2D.Y * LaplacianUV1;
 			OutPoint3D.LaplacianV = Curve34Points[0].Laplacian + InPoint2D.X * LaplacianUV2;
diff --git a/Engine/Source/Runtime/Datasmith/CADKernel/Base/Private/Geo/Surfaces/CoonsSurface.h b/Engine/Source/Runtime/Datasmith/CADKernel/Base/Private/Geo/Surfaces/CoonsSurface.h
--- a/Engine/Source/Runtime/Datasmith/CADKernel/Base/Private/Geo/Surfaces/CoonsSurface.h
+++ b/Engine/Source/Runtime/Datasmith/CADKernel/Base/Private/Geo/Surfaces/CoonsSurface.h
@@ -15,7 +15,7 @@ class CADKERNEL_API FCoonsSurface : public FSurface
 
 protected:
 	TSharedPtr<FCurve> Curves[4];
-	TArray<FPoint> Corners;
+	TArray<FVector> Corners;
 
 	FCoonsSurface(const double InToleranceGeometric, TSharedPtr<FCurve> InCurve1, TSharedPtr<FCurve> InCurve2, TSharedPtr<FCurve> InCurve3, TSharedPtr<FCurve> InCurve4);
 	FCoonsSurface(const double InToleranceGeometric, TSharedPtr<FCurve> InCurves[4]);
@@ -45,7 +45,7 @@ public:
 
 	virtual TSharedPtr<FEntityGeom> ApplyMatrix(const FMatrixH& InMatrix) const override;
 
-	virtual void EvaluatePoint(const FPoint2D& InSurfacicCoordinate, FSurfacicPoint& OutPoint3D, int32 InDerivativeOrder = 0) const override;
+	virtual void EvaluatePoint(const FVector2d& InSurfacicCoordinate, FSurfacicPoint& OutPoint3D, int32 InDerivativeOrder = 0) const override;
 
 	virtual void LinesNotDerivables(const FSurfacicBoundary& Bounds, int32 InDerivativeOrder, FCoordinateGrid& OutNotDerivableCoordinates) const override;
 };
diff --git a/Engine/Source/Runtime/Datasmith/CADKernel/Base/Private/Geo/Surfaces/CylinderSurface.cpp b/Engine/Source/Runtime/Datasmith/CADKernel/Base/Private/Geo/Surfaces/CylinderSurface.cpp
--- a/Engine/Source/Runtime/Datasmith/CADKernel/Base/Private/Geo/Surfaces/CylinderSurface.cpp
+++ b/Engine/Source/Runtime/Datasmith/CADKernel/Base/Private/Geo/Surfaces/CylinderSurface.cpp
@@ -25,7 +25,7 @@ void FCylinderSurface::InitBoundary()
 {
 }
 
-void FCylinderSurface::EvaluatePoint(const FPoint2D& InPoint2D, FSurfacicPoint& OutPoint3D, int32 InDerivativeOrder) const
+void FCylinderSurface::EvaluatePoint(const FVector2d& InPoint2D, FSurfacicPoint& OutPoint3D, int32 InDerivativeOrder) const
 {
 	OutPoint3D.DerivativeOrder = InDerivativeOrder;
 
@@ -34,19 +34,19 @@ void FCylinderSurface::EvaluatePoint(const FPoint2D& InPoint2D, FSurfacicPoint&
 
 	if(InDerivativeOrder>0) 
 	{
-		OutPoint3D.GradientU = FPoint(Radius*(-sin(InPoint2D.X)), Radius*(cos(InPoint2D.X)), 0.0);
-		OutPoint3D.GradientV = FPoint(0.0, 0.0, 1.0);
+		OutPoint3D.GradientU = FVector(Radius*(-sin(InPoint2D.X)), Radius*(cos(InPoint2D.X)), 0.0);
+		OutPoint3D.GradientV = FVector(0.0, 0.0, 1.0);
 
 		OutPoint3D.GradientU = Matrix.MultiplyVector(OutPoint3D.GradientU);
 		OutPoint3D.GradientV = Matrix.MultiplyVector(OutPoint3D.GradientV);
 
 		if (InDerivativeOrder > 1) 
 		{
-			OutPoint3D.LaplacianU = FPoint(Radius * (-cos(InPoint2D.X)), Radius * (-sin(InPoint2D.X)), 0.0);
+			OutPoint3D.LaplacianU = FVector(Radius * (-cos(InPoint2D.X)), Radius * (-sin(InPoint2D.X)), 0.0);
 			OutPoint3D.LaplacianU = Matrix.MultiplyVector(OutPoint3D.LaplacianU);
 
-			OutPoint3D.LaplacianV = FPoint(0.0, 0.0, 0.0);
-			OutPoint3D.LaplacianUV = FPoint(0.0, 0.0, 0.0);
+			OutPoint3D.LaplacianV = FVector(0.0, 0.0, 0.0);
+			OutPoint3D.LaplacianUV = FVector(0.0, 0.0, 0.0);
 		}
 	}
 }
diff --git a/Engine/Source/Runtime/Datasmith/CADKernel/Base/Private/Geo/Surfaces/CylinderSurface.h b/Engine/Source/Runtime/Datasmith/CADKernel/Base/Private/Geo/Surfaces/CylinderSurface.h
--- a/Engine/Source/Runtime/Datasmith/CADKernel/Base/Private/Geo/Surfaces/CylinderSurface.h
+++ b/Engine/Source/Runtime/Datasmith/CADKernel/Base/Private/Geo/Surfaces/CylinderSurface.h
@@ -34,15 +34,15 @@ protected:
 
 	void ComputeMinToleranceIso()
 	{
-		FPoint Origin = Matrix.Multiply(FPoint::ZeroVector);
+		FVector Origin = Matrix.Multiply(FVector::ZeroVector);
 
-		FPoint Point2DU{ 1 , 0, 0 };
+		FVector Point2DU{ 1 , 0, 0 };
 		double ToleranceU = Tolerance3D / Radius;
 
 		double ScaleU = ComputeScaleAlongAxis(Point2DU, Matrix, Origin);
 		ToleranceU /= ScaleU;
 
-		FPoint Point2DV{ 0, 1, 0 };
+		FVector Point2DV{ 0, 1, 0 };
 		double ToleranceV = Tolerance3D / ComputeScaleAlongAxis(Point2DV, Matrix, Origin);
 
 		MinToleranceIso.Set(ToleranceU, ToleranceV);
@@ -70,7 +70,7 @@ public:
 
 	virtual void InitBoundary() override;
 
-	virtual void EvaluatePoint(const FPoint2D& InSurfacicCoordinate, FSurfacicPoint& OutPoint3D, int32 InDerivativeOrder = 0) const override;
+	virtual void EvaluatePoint(const FVector2d& InSurfacicCoordinate, FSurfacicPoint& OutPoint3D, int32 InDerivativeOrder = 0) const override;
 
 	virtual void Presample(const FSurfacicBoundary& InBoundaries, FCoordinateGrid& OutCoordinates) override
 	{
diff --git a/Engine/Source/Runtime/Datasmith/CADKernel/Base/Private/Geo/Surfaces/NURBSSurface.cpp b/Engine/Source/Runtime/Datasmith/CADKernel/Base/Private/Geo/Surfaces/NURBSSurface.cpp
--- a/Engine/Source/Runtime/Datasmith/CADKernel/Base/Private/Geo/Surfaces/NURBSSurface.cpp
+++ b/Engine/Source/Runtime/Datasmith/CADKernel/Base/Private/Geo/Surfaces/NURBSSurface.cpp
@@ -10,10 +10,10 @@ namespace UE::CADKernel
 
 TSharedPtr<FEntityGeom> FNURBSSurface::ApplyMatrix(const FMatrixH& InMatrix) const
 {
-	TArray<FPoint> TransformedPoles;
+	TArray<FVector> TransformedPoles;
 	TransformedPoles.Reserve(Poles.Num());
 
-	for (const FPoint& Pole : Poles)
+	for (const FVector& Pole : Poles)
 	{
 		TransformedPoles.Emplace(InMatrix.Multiply(Pole));
 	}
@@ -60,7 +60,7 @@ void FNURBSSurface::Finalize()
 		{
 			if (!FMath::IsNearlyEqual(1., FirstWeigth))
 			{
-				for (FPoint& Pole : Poles)
+				for (FVector& Pole : Poles)
 				{
 					Pole /= FirstWeigth;
 				}
@@ -83,7 +83,7 @@ void FNURBSSurface::Finalize()
 	else
 	{
 		HomogeneousPoles.SetNum(Poles.Num() * 3);
-		memcpy(HomogeneousPoles.GetData(), Poles.GetData(), sizeof(FPoint) * Poles.Num());
+		memcpy(HomogeneousPoles.GetData(), Poles.GetData(), sizeof(FVector) * Poles.Num());
 	}
 
 	double UMin = UNodalVector[UDegree];
@@ -109,30 +109,29 @@ void FNURBSSurface::FillNurbs(FNurbsSurfaceHomogeneousData& NurbsData)
 	Swap(VNodalVector, NurbsData.VNodalVector);
 
 	int32 PoleCount = PoleUCount * PoleVCount;
-	
-	double* RawPolesPtr = NurbsData.HomogeneousPoles.GetData();
-	const int32 Dimension = bIsRational ? 4 : 3;
 	Poles.Reserve(PoleCount);
 
-	const int32 OffSet = Dimension * PoleVCount;
-	for (int32 Vndex = 0; Vndex < PoleVCount; ++Vndex)
+	const int32 Dimension = bIsRational ? 4 : 3;
+	const int32 UOffSet = Dimension * PoleVCount;
+	const double* HomogeneousPolesPtr = NurbsData.HomogeneousPoles.GetData();
+	for (int32 Vndex = 0, VOffset = 0; Vndex < PoleVCount; ++Vndex, VOffset += Dimension)
 	{
-		int32 RawPolesOffset = Dimension * Vndex;
-		for (int32 Undex = 0; Undex < PoleUCount; ++Undex, RawPolesOffset += OffSet)
+		const double* RawPolesPtr = HomogeneousPolesPtr + VOffset;
+		for (int32 Undex = 0; Undex < PoleUCount; ++Undex, RawPolesPtr += UOffSet)
 		{
-			Poles.Emplace(RawPolesPtr + RawPolesOffset);
+			Poles.Emplace(RawPolesPtr[0], RawPolesPtr[1], RawPolesPtr[2]);
 		}
 	}
 
 	if(bIsRational)
 	{
 		Weights.Reserve(PoleCount);
-		for (int32 Vndex = 3; Vndex < (PoleVCount * Dimension); Vndex += Dimension)
+		for (int32 Vndex = 3, VOffset = 3; Vndex < (PoleVCount * Dimension); Vndex += Dimension, VOffset += Dimension)
 		{
-			int32 RawPolesOffset = Vndex;
-			for (int32 Undex = 0; Undex < PoleUCount; ++Undex, RawPolesOffset += OffSet)
+			const double* RawPolesPtr = HomogeneousPolesPtr + VOffset;
+			for (int32 Undex = 0; Undex < PoleUCount; ++Undex, RawPolesPtr += UOffSet)
 			{
-				Weights.Emplace(RawPolesPtr[RawPolesOffset]);
+				Weights.Emplace(*RawPolesPtr);
 			}
 		}
 	}
diff --git a/Engine/Source/Runtime/Datasmith/CADKernel/Base/Private/Geo/Surfaces/NURBSSurface.h b/Engine/Source/Runtime/Datasmith/CADKernel/Base/Private/Geo/Surfaces/NURBSSurface.h
--- a/Engine/Source/Runtime/Datasmith/CADKernel/Base/Private/Geo/Surfaces/NURBSSurface.h
+++ b/Engine/Source/Runtime/Datasmith/CADKernel/Base/Private/Geo/Surfaces/NURBSSurface.h
@@ -23,7 +23,7 @@ protected:
 
 	TArray<double> Weights;
 
-	TArray<FPoint> Poles;
+	TArray<FVector> Poles;
 
 	bool bIsRational;
 
@@ -37,7 +37,7 @@ protected:
 	 * @param NodalVectorU Its size is the number of poles in U + the surface degree in U + 1 (PoleUNum + UDegre + 1)
 	 * @param NodalVectorV Its size is the number of poles in V + the surface degree in V + 1 (PoleVNum + VDegre + 1)
 	 */
-	FNURBSSurface(const double InToleranceGeometric, int32 InPoleUCount, int32 InPoleVCount, int32 InDegreU, int32 InDegreV, const TArray<double>& InNodalVectorU, const TArray<double>& InNodalVectorV, const TArray<FPoint>& InPoles)
+	FNURBSSurface(const double InToleranceGeometric, int32 InPoleUCount, int32 InPoleVCount, int32 InDegreU, int32 InDegreV, const TArray<double>& InNodalVectorU, const TArray<double>& InNodalVectorV, const TArray<FVector>& InPoles)
 		: FSurface(InToleranceGeometric)
 		, PoleUCount(InPoleUCount)
 		, PoleVCount(InPoleVCount)
@@ -56,7 +56,7 @@ protected:
 	 * @param NodalVectorU Its size is the number of poles in U + the surface degree in U + 1 (PoleUNum + UDegre + 1)
 	 * @param NodalVectorV Its size is the number of poles in V + the surface degree in V + 1 (PoleVNum + VDegre + 1)
 	 */
-	FNURBSSurface(const double InToleranceGeometric, int32 InPoleUCount, int32 InPoleVCount, int32 InDegreU, int32 InDegreV, const TArray<double>& InNodalVectorU, const TArray<double>& InNodalVectorV, const TArray<FPoint>& InPoles, const TArray<double>& InWeights)
+	FNURBSSurface(const double InToleranceGeometric, int32 InPoleUCount, int32 InPoleVCount, int32 InDegreU, int32 InDegreV, const TArray<double>& InNodalVectorU, const TArray<double>& InNodalVectorV, const TArray<FVector>& InPoles, const TArray<double>& InWeights)
 		: FSurface(InToleranceGeometric)
 		, PoleUCount(InPoleUCount)
 		, PoleVCount(InPoleVCount)
@@ -109,7 +109,7 @@ protected:
 public:
 
 
-	virtual void ValidateUVPoints(TArray<FPoint2D>& UVPoints) const
+	virtual void ValidateUVPoints(TArray<FVector2d>& UVPoints) const
 	{
 		const double ToleranceU = GetIsoTolerance(EIso::IsoU);
 		const double ToleranceV = GetIsoTolerance(EIso::IsoV);
@@ -125,7 +125,7 @@ public:
 				return false;
 			};
 
-		for (FPoint2D& UVPoint : UVPoints)
+		for (FVector2d& UVPoint : UVPoints)
 		{
 			for (const double& NodalValue : UNodalVector)
 			{
@@ -197,7 +197,7 @@ public:
 		}
 	}
 
-	const TArray<FPoint>& GetPoles() const
+	const TArray<FVector>& GetPoles() const
 	{
 		return Poles;
 	}
@@ -232,7 +232,7 @@ public:
 
 	virtual TSharedPtr<FEntityGeom> ApplyMatrix(const FMatrixH& InMatrix) const override;
 
-	virtual void EvaluatePoint(const FPoint2D& InSurfacicCoordinate, FSurfacicPoint& OutPoint3D, int32 InDerivativeOrder = 0) const override
+	virtual void EvaluatePoint(const FVector2d& InSurfacicCoordinate, FSurfacicPoint& OutPoint3D, int32 InDerivativeOrder = 0) const override
 	{
 		BSpline::EvaluatePoint(*this, InSurfacicCoordinate, OutPoint3D, InDerivativeOrder);
 	}
