[CAD] Fixed importing "internal" nodes as overrides

- ignoring 'internal' product occurrences as they represent something else than overrides.

- In presented case those 'internal' occurrences broke the imported model by replacing(i.e. overriding) node because they have the same Entity name. And current code replaces overrides with matching Label(which is derived from Entity name). In general this is wrong as occurrences may have same name, even in this case - a door part and its mirrored part have exactly same name.

diff --git a/Engine/Plugins/Enterprise/DatasmithCADImporter/Source/CADInterfaces/Private/TechSoftFileParser.cpp b/Engine/Plugins/Enterprise/DatasmithCADImporter/Source/CADInterfaces/Private/TechSoftFileParser.cpp
--- a/Engine/Plugins/Enterprise/DatasmithCADImporter/Source/CADInterfaces/Private/TechSoftFileParser.cpp
+++ b/Engine/Plugins/Enterprise/DatasmithCADImporter/Source/CADInterfaces/Private/TechSoftFileParser.cpp
@@ -1006,6 +1006,10 @@ void FTechSoftFileParser::TraverseOccurrence(const A3DAsmProductOccurrence* Occu
 void FTechSoftFileParser::ExtractOverrideOccurrenceSubtree(const A3DAsmProductOccurrence* OccurrencePtr, FArchiveWithOverridenChildren& Parent)
 {
 	TUniqueTSObj<A3DAsmProductOccurrenceData> OccurrenceData(OccurrencePtr);
+	if (OccurrenceData->m_uiProductFlags & A3D_PRODUCT_FLAG_INTERNAL)
+	{
+		return;
+	}
 
 	FArchiveOverrideOccurrence& OverrideOccurrence = SceneGraph.AddOverrideOccurrence(Parent);
 	ExtractMetaData(OccurrencePtr, OverrideOccurrence);
