[In Progress] InterchangeCADImporter plugin
- Added skeleton translated nodes for meshes
- Added skeleton CAD pipelines for mesh, actor and level and a composite one
- Miscellaneous changes on the Wire parser

#rnx

diff --git a/Engine/Plugins/Enterprise/DatasmithCADImporter/Source/DatasmithCADTranslator/Private/DatasmithCADTranslator.cpp b/Engine/Plugins/Enterprise/DatasmithCADImporter/Source/DatasmithCADTranslator/Private/DatasmithCADTranslator.cpp
--- a/Engine/Plugins/Enterprise/DatasmithCADImporter/Source/DatasmithCADTranslator/Private/DatasmithCADTranslator.cpp
+++ b/Engine/Plugins/Enterprise/DatasmithCADImporter/Source/DatasmithCADTranslator/Private/DatasmithCADTranslator.cpp
@@ -37,6 +37,15 @@ Default is disable\n"),
 
 void FDatasmithCADTranslator::Initialize(FDatasmithTranslatorCapabilities& OutCapabilities)
 {
+	IConsoleVariable* CVar = IConsoleManager::Get().FindConsoleVariable(TEXT("Interchange.FeatureFlags.Import.CAD"), false);
+	bool bInterchangeCADEnabled = CVar ? CVar->GetBool() : false;
+
+	if (bInterchangeCADEnabled)
+	{
+		OutCapabilities.bIsEnabled = false;
+		return;
+	}
+
 	TFunction<bool()> GetCADInterfaceAvailability = []() -> bool
 	{
 		if (ICADInterfacesModule::GetAvailability() == ECADInterfaceAvailability::Unavailable)
@@ -143,6 +152,14 @@ void FDatasmithCADTranslator::Initialize(FDatasmithTranslatorCapabilities& OutCa
 
 bool FDatasmithCADTranslator::IsSourceSupported(const FDatasmithSceneSource& Source)
 {
+	IConsoleVariable* CVar = IConsoleManager::Get().FindConsoleVariable(TEXT("Interchange.FeatureFlags.Import.CAD"), false);
+	bool bInterchangeCADEnabled = CVar ? CVar->GetBool() : false;
+
+	if (bInterchangeCADEnabled)
+	{
+		return false;
+	}
+
 	if (Source.GetSourceFileExtension() != TEXT("xml"))
 	{
 		return true;
