Build fix: fix discarded math operations

#rnx

diff --git a/Engine/Plugins/Enterprise/DatasmithCADImporter/Source/CADInterfaces/Private/TechSoftFileParser.cpp b/Engine/Plugins/Enterprise/DatasmithCADImporter/Source/CADInterfaces/Private/TechSoftFileParser.cpp
--- a/Engine/Plugins/Enterprise/DatasmithCADImporter/Source/CADInterfaces/Private/TechSoftFileParser.cpp
+++ b/Engine/Plugins/Enterprise/DatasmithCADImporter/Source/CADInterfaces/Private/TechSoftFileParser.cpp
@@ -956,7 +956,8 @@ void FTechSoftFileParser::TraverseOccurrence(const A3DAsmProductOccurrence* Occu
 				Scale = Instance.Unit / ReferenceUnit;
 				if (!FMath::IsNearlyEqual(Scale, 1.))
 				{
-					Instance.TransformMatrix.ApplyScale(Scale);
+					// Commented as the result was discarded
+					// Instance.TransformMatrix = Instance.TransformMatrix.ApplyScale(Scale);
 				}
 			}
 
