Restored include files for Linux build
- Restored back the header files from previous TechSoft drop since Linux binaries have not been upgraded
- Modified techSoft.Build.cs to account for the change in header files
- Modified the CADInterface.Build.cs to temporarily work around a problem with build on Linux. Will be investigated for next release
- Updated license header to use license from previous TechSoft SDK version on Linux
Note: The TechSoft SDK for Linux will be updated in next release alongside the Win64 one. All those changes will be gone.

#rnx

diff --git a/Engine/Plugins/Enterprise/DatasmithCADImporter/Source/CADInterfaces/CADInterfaces.Build.cs b/Engine/Plugins/Enterprise/DatasmithCADImporter/Source/CADInterfaces/CADInterfaces.Build.cs
--- a/Engine/Plugins/Enterprise/DatasmithCADImporter/Source/CADInterfaces/CADInterfaces.Build.cs
+++ b/Engine/Plugins/Enterprise/DatasmithCADImporter/Source/CADInterfaces/CADInterfaces.Build.cs
@@ -26,10 +26,20 @@ namespace UnrealBuildTool.Rules
 			// CAD library is only available if TechSoft is available too
 			bool bHasTechSoft = System.Type.GetType("TechSoft") != null;
 
-			if ((Target.Platform == UnrealTargetPlatform.Win64 || Target.Platform == UnrealTargetPlatform.Linux) && bHasTechSoft)
+			if (Target.Platform == UnrealTargetPlatform.Win64 && bHasTechSoft)
 			{
 				PublicDependencyModuleNames.Add("TechSoft");
 			}
+			// System.Type.GetType("TechSoft") does not seem to work on Linux
+			// Temporary fix. I will investigate for 5.6
+			else if (Target.Platform == UnrealTargetPlatform.Linux)
+			{
+				string TechSoftPath = Path.Combine(EngineDirectory, "Restricted/NotForLicensees/Source/ThirdParty/Enterprise/TechSoft/TechSoft.Build.cs");
+				if(File.Exists(TechSoftPath))
+				{
+					PublicDependencyModuleNames.Add("TechSoft");
+				}
+			}
 		}
 	}
 }
