Fixed crash when loading specific CATProduct file
- Temporarily worked around the bug from TechSoft
- Updated worker and plugin to use the newly added TechSoft binaries.

#rnx

diff --git a/Engine/Plugins/Enterprise/DatasmithCADImporter/Source/CADInterfaces/CADInterfaces.Build.cs b/Engine/Plugins/Enterprise/DatasmithCADImporter/Source/CADInterfaces/CADInterfaces.Build.cs
--- a/Engine/Plugins/Enterprise/DatasmithCADImporter/Source/CADInterfaces/CADInterfaces.Build.cs
+++ b/Engine/Plugins/Enterprise/DatasmithCADImporter/Source/CADInterfaces/CADInterfaces.Build.cs
@@ -8,6 +8,7 @@ namespace UnrealBuildTool.Rules
 	{
 		public CADInterfaces(ReadOnlyTargetRules Target) : base(Target)
 		{
+			bEnableExceptions = true;
 			bLegalToDistributeObjectCode = true;
 			bUseUnity = false;
 
diff --git a/Engine/Plugins/Enterprise/DatasmithCADImporter/Source/CADInterfaces/Private/TechSoftInterface.cpp b/Engine/Plugins/Enterprise/DatasmithCADImporter/Source/CADInterfaces/Private/TechSoftInterface.cpp
--- a/Engine/Plugins/Enterprise/DatasmithCADImporter/Source/CADInterfaces/Private/TechSoftInterface.cpp
+++ b/Engine/Plugins/Enterprise/DatasmithCADImporter/Source/CADInterfaces/Private/TechSoftInterface.cpp
@@ -38,8 +38,7 @@ const TCHAR* FTechSoftInterface::GetVersion()
 #endif
 }
 
-#pragma optimize("", off)
-bool FTechSoftInterface::InitializeKernel(const TCHAR* InEnginePluginsPath)
+bool FTechSoftInterface::InitializeKernel(const TCHAR* InEngineDir)
 {
 #ifdef USE_TECHSOFT_SDK
 	if (bIsInitialize)
@@ -47,20 +46,15 @@ bool FTechSoftInterface::InitializeKernel(const TCHAR* InEnginePluginsPath)
 		return true;
 	}
 
-	FString EnginePluginsPath(InEnginePluginsPath);
-#ifndef CADKERNEL_DEV
-	if (EnginePluginsPath.IsEmpty())
+	FString EngineDir(InEngineDir);
+	if (EngineDir.IsEmpty())
 	{
-		EnginePluginsPath = FPaths::EnginePluginsDir();
+		EngineDir = FPaths::EngineDir();
 	}
-#endif
-
-#ifdef CADKERNEL_DEV
-	FString TechSoftDllPath = EnginePluginsPath;
-#else
-	FString TechSoftDllPath = FPaths::Combine(EnginePluginsPath, TEXT("Enterprise/DatasmithCADImporter"), TEXT("Binaries"), FPlatformProcess::GetBinariesSubdirectory(), TEXT("TechSoft"));
-#endif
 
+	static const TCHAR* ThirdPartyDir = TEXT("Binaries/ThirdParty/TechSoft");
+	
+	FString TechSoftDllPath = FPaths::Combine(EngineDir, ThirdPartyDir, FPlatformProcess::GetBinariesSubdirectory());
 	TechSoftDllPath = FPaths::ConvertRelativePathToFull(TechSoftDllPath);
 
 	if (A3DSDKLoadLibraryA(TCHAR_TO_UTF8(*TechSoftDllPath)))
@@ -86,7 +80,6 @@ bool FTechSoftInterface::InitializeKernel(const TCHAR* InEnginePluginsPath)
 #endif
 	return false;
 }
-#pragma optimize("", on)
 
 namespace TechSoftInterface
 {
@@ -467,9 +460,16 @@ A3DStatus ExportModelFileToPrcFile(const A3DAsmModelFile* ModelFile, const A3DRW
 double GetModelFileUnit(const A3DAsmModelFile* ModelFile)
 {
 	double FileUnit = 0.1;
-	if (A3DAsmModelFileGetUnit(ModelFile, &FileUnit) != A3DStatus::A3D_SUCCESS)
+	try
+	{
+		if (A3DAsmModelFileGetUnit(ModelFile, &FileUnit) != A3DStatus::A3D_SUCCESS)
+		{
+			return 0.1;
+		}
+	}
+	catch (...)
 	{
-		return 0.1;
+			return 0.1;
 	}
 	return FileUnit * 0.1;
 }
diff --git a/Engine/Plugins/Enterprise/DatasmithCADImporter/Source/DatasmithDispatcher/Private/DatasmithWorkerHandler.cpp b/Engine/Plugins/Enterprise/DatasmithCADImporter/Source/DatasmithDispatcher/Private/DatasmithWorkerHandler.cpp
--- a/Engine/Plugins/Enterprise/DatasmithCADImporter/Source/DatasmithDispatcher/Private/DatasmithWorkerHandler.cpp
+++ b/Engine/Plugins/Enterprise/DatasmithCADImporter/Source/DatasmithDispatcher/Private/DatasmithWorkerHandler.cpp
@@ -82,7 +82,7 @@ void FDatasmithWorkerHandler::StartWorkerProcess()
 		CommandToProcess += TEXT(" -ServerPort ") + FString::FromInt(ListenPort);
 		CommandToProcess += TEXT(" -CacheDir \"") + CachePath + TEXT('"');
 		CommandToProcess += TEXT(" -CacheVersion ") + FString::FromInt(FCADToolsModule::GetCacheVersion());
-		CommandToProcess += TEXT(" -EnginePluginsDir \"") + FPaths::EnginePluginsDir() + TEXT('"');
+		CommandToProcess += TEXT(" -EnginePluginsDir \"") + FPaths::EngineDir() + TEXT('"');
 		UE_LOG(LogDatasmithDispatcher, Verbose, TEXT("CommandToProcess: %s"), *CommandToProcess);
 
 		WorkerHandle = FPlatformProcess::CreateProc(*ProcessorPath, *CommandToProcess, true, false, false, &WorkerProcessId, 0, nullptr, nullptr);
