Disable wire interface modules in CAD importer plugin for Windows Arm64
- no OpenModel SDK available

#rnx

diff --git a/Engine/Plugins/Enterprise/DatasmithCADImporter/DatasmithCADImporter.uplugin b/Engine/Plugins/Enterprise/DatasmithCADImporter/DatasmithCADImporter.uplugin
--- a/Engine/Plugins/Enterprise/DatasmithCADImporter/DatasmithCADImporter.uplugin
+++ b/Engine/Plugins/Enterprise/DatasmithCADImporter/DatasmithCADImporter.uplugin
@@ -76,7 +76,8 @@
 			"LoadingPhase": "Default",
 			"PlatformAllowList": [
 				"Win64"
-			]
+			],
+			"PlatformArchitectureDenyList": [ "Win64:arm64" ]
 		},
 		{
 			"Name": "WireInterface2020",
@@ -84,7 +85,8 @@
 			"LoadingPhase": "Default",
 			"PlatformAllowList": [
 				"Win64"
-			]
+			],
+			"PlatformArchitectureDenyList": [ "Win64:arm64" ]
 		},
 		{
 			"Name": "WireInterface2021_3",
@@ -92,7 +94,8 @@
 			"LoadingPhase": "Default",
 			"PlatformAllowList": [
 				"Win64"
-			]
+			],
+			"PlatformArchitectureDenyList": [ "Win64:arm64" ]
 		},
 		{
 			"Name": "WireInterface2022",
@@ -100,7 +103,8 @@
 			"LoadingPhase": "Default",
 			"PlatformAllowList": [
 				"Win64"
-			]
+			],
+			"PlatformArchitectureDenyList": [ "Win64:arm64" ]
 		},
 		{
 			"Name": "WireInterface2022_1",
@@ -108,7 +112,8 @@
 			"LoadingPhase": "Default",
 			"PlatformAllowList": [
 				"Win64"
-			]
+			],
+			"PlatformArchitectureDenyList": [ "Win64:arm64" ]
 		},
 		{
 			"Name": "WireInterface2022_2",
@@ -116,7 +121,8 @@
 			"LoadingPhase": "Default",
 			"PlatformAllowList": [
 				"Win64"
-			]
+			],
+			"PlatformArchitectureDenyList": [ "Win64:arm64" ]
 		},
 		{
 			"Name": "WireInterface2023_0",
@@ -124,7 +130,8 @@
 			"LoadingPhase": "Default",
 			"PlatformAllowList": [
 				"Win64"
-			]
+			],
+			"PlatformArchitectureDenyList": [ "Win64:arm64" ]
 		},
 		{
 			"Name": "WireInterface2023_1",
@@ -132,7 +139,8 @@
 			"LoadingPhase": "Default",
 			"PlatformAllowList": [
 				"Win64"
-			]
+			],
+			"PlatformArchitectureDenyList": [ "Win64:arm64" ]
 		},
 		{
 			"Name": "WireInterface2024_1",
@@ -140,7 +148,8 @@
 			"LoadingPhase": "Default",
 			"PlatformAllowList": [
 				"Win64"
-			]
+			],
+			"PlatformArchitectureDenyList": [ "Win64:arm64" ]
 		},
 		{
 			"Name": "WireInterface2025_0",
@@ -148,7 +157,8 @@
 			"LoadingPhase": "Default",
 			"PlatformAllowList": [
 				"Win64"
-			]
+			],
+			"PlatformArchitectureDenyList": [ "Win64:arm64" ]
 		},
 		{
 			"Name": "ParametricSurfaceExtension",
@@ -157,7 +167,8 @@
 			"PlatformAllowList": [
 				"Win64",
 				"Linux"
-			]
+			],
+			"PlatformArchitectureDenyList": [ "Win64:arm64" ]
 		},
 		{
 			"Name": "ParametricSurface",
