Added support for Alias 2026.0

#rnx

diff --git a/Engine/Plugins/Enterprise/DatasmithCADImporter/DatasmithCADImporter.uplugin b/Engine/Plugins/Enterprise/DatasmithCADImporter/DatasmithCADImporter.uplugin
--- a/Engine/Plugins/Enterprise/DatasmithCADImporter/DatasmithCADImporter.uplugin
+++ b/Engine/Plugins/Enterprise/DatasmithCADImporter/DatasmithCADImporter.uplugin
@@ -160,6 +160,14 @@
 			],
 			"PlatformArchitectureDenyList": [ "Win64:arm64" ]
 		},
+		{
+			"Name": "WireInterface2026_0",
+			"Type": "Editor",
+			"LoadingPhase": "Default",
+			"PlatformAllowList": [
+				"Win64"
+			]
+		},
 		{
 			"Name": "ParametricSurfaceExtension",
 			"Type": "Editor",
diff --git a/Engine/Plugins/Enterprise/DatasmithCADImporter/Source/WireInterface/Private/WireInterfaceImpl.cpp b/Engine/Plugins/Enterprise/DatasmithCADImporter/Source/WireInterface/Private/WireInterfaceImpl.cpp
--- a/Engine/Plugins/Enterprise/DatasmithCADImporter/Source/WireInterface/Private/WireInterfaceImpl.cpp
+++ b/Engine/Plugins/Enterprise/DatasmithCADImporter/Source/WireInterface/Private/WireInterfaceImpl.cpp
@@ -91,6 +91,7 @@ namespace UE_DATASMITHWIRETRANSLATOR_NAMESPACE
 	static const FColor DefaultColor = FColor(200, 200, 200);
 
 	const uint64 LibAliasNext_Version = 0xffffffffffffffffull;
+	const uint64 LibAlias2026_0_0_Version = 9007199254741154;
 	const uint64 LibAlias2025_0_0_Version = 8725724278030572;
 	const uint64 LibAlias2024_1_0_Version = 8444253596292024;
 	const uint64 LibAlias2023_1_0_Version = 8162778619576619;
@@ -139,8 +140,12 @@ namespace UE_DATASMITHWIRETRANSLATOR_NAMESPACE
 	const FString AliasSdkVersion = TEXT("2024.1");
 #elif defined(OPEN_MODEL_2025_0)      
 	const uint64 LibAliasVersionMin = LibAlias2025_0_0_Version;
-	const uint64 LibAliasVersionMax = LibAliasNext_Version;
+	const uint64 LibAliasVersionMax = LibAlias2026_0_0_Version;
 	const FString AliasSdkVersion = TEXT("2025.0");
+#elif defined(OPEN_MODEL_2026_0)      
+	const uint64 LibAliasVersionMin = LibAlias2026_0_0_Version;
+	const uint64 LibAliasVersionMax = LibAliasNext_Version;
+	const FString AliasSdkVersion = TEXT("2026.0");
 #endif
 
 	// Alias material management (to allow sew of BReps of different materials):
@@ -349,9 +354,6 @@ namespace UE_DATASMITHWIRETRANSLATOR_NAMESPACE
 			return false;
 		}
 
-		AlRetrieveOptions options;
-		AlUniverse::retrieveOptions(options);
-
 		bSceneLoaded = true;
 
 		return TraverseModel();
diff --git a/Engine/Plugins/Enterprise/DatasmithCADImporter/Source/WireInterface/WireInterface2026_0.Build.cs b/Engine/Plugins/Enterprise/DatasmithCADImporter/Source/WireInterface/WireInterface2026_0.Build.cs
new file mode 100644
--- /dev/null
+++ b/Engine/Plugins/Enterprise/DatasmithCADImporter/Source/WireInterface/WireInterface2026_0.Build.cs
@@ -0,0 +1,32 @@
+// Copyright Epic Games, Inc. All Rights Reserved.
+
+using System.IO;
+using UnrealBuildTool;
+
+public class WireInterface2026_0 : WireInterfaceBase
+{
+	public WireInterface2026_0(ReadOnlyTargetRules Target) 
+		: base(Target)
+	{
+	}
+
+	public override string GetAliasVersion()
+	{
+		return "OpenModel2026_0";
+	}
+
+	public override string GetAliasDefinition()
+	{
+		return "OPEN_MODEL_2026_0";
+	}
+
+	public override int GetMajorVersion()
+	{
+		return 2026;
+	}
+
+	public override int GetMinorVersion()
+	{
+		return 0;
+	}
+}
\ No newline at end of file
