disable some additional plugins on Windows Arm64

#rnx

diff --git a/Engine/Plugins/Enterprise/DatasmithCADImporter/DatasmithCADImporter.uplugin b/Engine/Plugins/Enterprise/DatasmithCADImporter/DatasmithCADImporter.uplugin
--- a/Engine/Plugins/Enterprise/DatasmithCADImporter/DatasmithCADImporter.uplugin
+++ b/Engine/Plugins/Enterprise/DatasmithCADImporter/DatasmithCADImporter.uplugin
@@ -166,7 +166,8 @@
 			"LoadingPhase": "Default",
 			"PlatformAllowList": [
 				"Win64"
-			]
+			],
+			"PlatformArchitectureDenyList": [ "Win64:arm64" ]
 		},
 		{
 			"Name": "ParametricSurfaceExtension",
