Enabled FBX import into level by default at the global variable instead of ConsoleVariables.ini
Remove unnecessary CAD releated CVArs from ConsoleVariables.ini. The default values are updated in the cpp file.

#rnx
#DisabledValidations Preflight

diff --git a/Engine/Config/ConsoleVariables.ini b/Engine/Config/ConsoleVariables.ini
--- a/Engine/Config/ConsoleVariables.ini
+++ b/Engine/Config/ConsoleVariables.ini
@@ -215,10 +215,3 @@ fx.SkipVectorVMBackendOptimizations=1
 ; Force to import embedded tessellation instead of meshing BRep model.
 ; Default is false
 ;ds.CADTranslator.PreferJtFileEmbeddedTessellation=false
-
-; Temporarily disable all thin face/zone processing
-ds.CADTranslator.Meshing.ActivateThinZoneMeshing=false
-ds.CADTranslator.Stitching.RemoveThinFaces=false
-
-; Enable Interchange FBX Import into Level
-Interchange.FeatureFlags.Import.FBX.ToLevel=true
diff --git a/Engine/Plugins/Enterprise/DatasmithCADImporter/Source/CADTools/Private/CADOption.cpp b/Engine/Plugins/Enterprise/DatasmithCADImporter/Source/CADTools/Private/CADOption.cpp
--- a/Engine/Plugins/Enterprise/DatasmithCADImporter/Source/CADTools/Private/CADOption.cpp
+++ b/Engine/Plugins/Enterprise/DatasmithCADImporter/Source/CADTools/Private/CADOption.cpp
@@ -83,7 +83,7 @@ Only available with CADKernel.\n\
 Default value of ForceFactor is 5.\n"),
 ECVF_Default);
 
-bool FImportParameters::bGStitchingRemoveThinFaces = true;
+bool FImportParameters::bGStitchingRemoveThinFaces = false;
 FAutoConsoleVariableRef GCADTranslatorStitchingRemoveThinFaces(
 	TEXT("ds.CADTranslator.Stitching.RemoveThinFaces"),
 	FImportParameters::bGStitchingRemoveThinFaces,
@@ -92,7 +92,7 @@ Only available with CADKernel.\n\
 Default value of RemoveThinFaces is true\n"),
 ECVF_Default);
 
-bool FImportParameters::bGStitchingRemoveDuplicatedFaces = true;
+bool FImportParameters::bGStitchingRemoveDuplicatedFaces = false;
 FAutoConsoleVariableRef GCADTranslatorStitchingRemoveDuplicatedFaces(
 	TEXT("ds.CADTranslator.Stitching.RemoveDuplicatedFaces"),
 	FImportParameters::bGStitchingRemoveDuplicatedFaces,
@@ -101,7 +101,7 @@ Only available with CADKernel.\n\
 Default value of RemoveDuplicatedFaces is true\n"),
 ECVF_Default);
 
-bool FImportParameters::bGActivateThinZoneMeshing = true;
+bool FImportParameters::bGActivateThinZoneMeshing = false;
 FAutoConsoleVariableRef GActivateThinZoneMeshing(
 	TEXT("ds.CADTranslator.Meshing.ActivateThinZoneMeshing"),
 	FImportParameters::bGActivateThinZoneMeshing,
diff --git a/Engine/Plugins/Interchange/Runtime/Source/Import/Private/Fbx/InterchangeFbxTranslator.cpp b/Engine/Plugins/Interchange/Runtime/Source/Import/Private/Fbx/InterchangeFbxTranslator.cpp
--- a/Engine/Plugins/Interchange/Runtime/Source/Import/Private/Fbx/InterchangeFbxTranslator.cpp
+++ b/Engine/Plugins/Interchange/Runtime/Source/Import/Private/Fbx/InterchangeFbxTranslator.cpp
@@ -38,7 +38,7 @@ static FAutoConsoleVariableRef CCvarInterchangeEnableFBXImport(
 	TEXT("Whether FBX support is enabled."),
 	ECVF_Default);
 
-static bool GInterchangeEnableFBXLevelImport = false;
+static bool GInterchangeEnableFBXLevelImport = true;
 static FAutoConsoleVariableRef CCvarInterchangeEnableFBXLevelImport(
 	TEXT("Interchange.FeatureFlags.Import.FBX.ToLevel"),
 	GInterchangeEnableFBXLevelImport,
