Fixed slot name issue with CAD import
- Removed unnecessary assignments of material override on staic mesh actor

#rnx

#Synced-CL 50905747

diff --git a/Engine/Plugins/Enterprise/DatasmithCADImporter/Source/DatasmithCADTranslator/Private/DatasmithSceneGraphBuilder.cpp b/Engine/Plugins/Enterprise/DatasmithCADImporter/Source/DatasmithCADTranslator/Private/DatasmithSceneGraphBuilder.cpp
--- a/Engine/Plugins/Enterprise/DatasmithCADImporter/Source/DatasmithCADTranslator/Private/DatasmithSceneGraphBuilder.cpp
+++ b/Engine/Plugins/Enterprise/DatasmithCADImporter/Source/DatasmithCADTranslator/Private/DatasmithSceneGraphBuilder.cpp
@@ -519,18 +519,6 @@ TSharedPtr<IDatasmithActorElement> FDatasmithSceneBaseGraphBuilder::BuildBody(FC
 
 	DatasmithSceneGraphBuilderImpl::AddTransformToActor(Body, ActorElement, ImportParameters);
 
-	if (MaterialUId >= 0)
-	{
-		TSharedPtr<IDatasmithMaterialIDElement> PartMaterialIDElement = FindOrAddMaterial(MaterialUId);
-		const TCHAR* MaterialIDElementName = PartMaterialIDElement->GetName();
-
-		for (int32 Index = 0; Index < MeshElement->GetMaterialSlotCount(); ++Index)
-		{
-			TSharedPtr<IDatasmithMaterialIDElement> MaterialIDElement = FDatasmithSceneFactory::CreateMaterialId(MaterialIDElementName);
-			MaterialIDElement->SetId(MeshElement->GetMaterialSlotAt(Index)->GetId());
-			ActorElement->AddMaterialOverride(MaterialIDElement);
-		}
-	}
 	return ActorElement;
 }
 
